Fix crash when connecting to session and presence manager isn't initialized yet

diff --git a/Engine/Plugins/Developer/Concert/MultiUserClient/Source/MultiUserClient/Private/Widgets/SActiveSession.cpp b/Engine/Plugins/Developer/Concert/MultiUserClient/Source/MultiUserClient/Private/Widgets/SActiveSession.cpp
--- a/Engine/Plugins/Developer/Concert/MultiUserClient/Source/MultiUserClient/Private/Widgets/SActiveSession.cpp
+++ b/Engine/Plugins/Developer/Concert/MultiUserClient/Source/MultiUserClient/Private/Widgets/SActiveSession.cpp
@@ -180,8 +180,11 @@ public:
 		{
 			// The world path returned by the function below is always the one from the editor context, never the one from the PIE/SIE context, so
 			// it will not contain the prefix added to the world when in PIE/SIE.
-			EEditorPlayMode EditorPlayMode;
-			WorldPath = SyncClientPin->GetPresenceManager()->GetPresenceWorldPath(ClientInfoPin->ClientEndpointId, EditorPlayMode);
+			EEditorPlayMode EditorPlayMode = EEditorPlayMode::None;
+			if (IConcertClientPresenceManager* PresenceManager = SyncClientPin->GetPresenceManager())
+			{
+				WorldPath = PresenceManager->GetPresenceWorldPath(ClientInfoPin->ClientEndpointId, EditorPlayMode);
+			}
 			
 			// The world path is returned as something like /Game/MyMap.MyMap, but we are only interested to keep the
 			// string left to the '.' to display "/Game/MyMap"
@@ -207,7 +210,8 @@ public:
 		TSharedPtr<IConcertSyncClient> ConcertSyncClientPin = SyncClient.Pin();
 
 		// If the user clicks on another client row other than the one representing himself.
-		if (ClientInfoPin && ClientSessionPin && ConcertSyncClientPin && ClientInfoPin->ClientEndpointId != ClientSessionPin->GetSessionClientEndpointId())
+		if (ClientInfoPin && ClientSessionPin && ConcertSyncClientPin && ConcertSyncClientPin->GetPresenceManager() &&
+			ClientInfoPin->ClientEndpointId != ClientSessionPin->GetSessionClientEndpointId())
 		{
 			IConcertClientPresenceManager* PresenceManager = ConcertSyncClientPin->GetPresenceManager();
 			EEditorPlayMode EditorPlayMode; // The result stored in this variable is not used.
diff --git a/Engine/Plugins/Developer/Concert/MultiUserClient/Source/MultiUserClientLibrary/Private/MultiUserClientStatics.cpp b/Engine/Plugins/Developer/Concert/MultiUserClient/Source/MultiUserClientLibrary/Private/MultiUserClientStatics.cpp
--- a/Engine/Plugins/Developer/Concert/MultiUserClient/Source/MultiUserClientLibrary/Private/MultiUserClientStatics.cpp
+++ b/Engine/Plugins/Developer/Concert/MultiUserClient/Source/MultiUserClientLibrary/Private/MultiUserClientStatics.cpp
@@ -121,7 +121,7 @@ void UMultiUserClientStatics::JumpToMultiUserPresence(const FString& OtherUserNa
 				}
 			}
 
-			if (OtherClientId.IsValid())
+			if (OtherClientId.IsValid() && ConcertSyncClient->GetPresenceManager())
 			{
 				ConcertSyncClient->GetPresenceManager()->InitiateJumpToPresence(OtherClientId, TransformOffset);
 			}
