- Don't auto-commit the archive name when the text box lose focus because when 'Cancel' is clicked, the focus is lost first.

diff --git a/Engine/Plugins/Developer/Concert/MultiUserClient/Source/MultiUserClient/Private/Widgets/SConcertBrowser.cpp b/Engine/Plugins/Developer/Concert/MultiUserClient/Source/MultiUserClient/Private/Widgets/SConcertBrowser.cpp
--- a/Engine/Plugins/Developer/Concert/MultiUserClient/Source/MultiUserClient/Private/Widgets/SConcertBrowser.cpp
+++ b/Engine/Plugins/Developer/Concert/MultiUserClient/Source/MultiUserClient/Private/Widgets/SConcertBrowser.cpp
@@ -1528,7 +1528,7 @@ void SSaveRestoreSessionRow::OnSessionNameChanged(const FText& NewName)
 void SSaveRestoreSessionRow::OnSessionNameCommitted(const FText& NewText, ETextCommit::Type CommitType)
 {
 	TSharedPtr<FConcertSessionItem> ItemPin = Item.Pin();
-	if (CommitType == ETextCommit::Type::OnEnter || (ItemPin && CommitType == ETextCommit::OnUserMovedFocus && ItemPin->Type == FConcertSessionItem::EType::SaveSession))
+	if (CommitType == ETextCommit::Type::OnEnter)
 	{
 		OnAccept();
 	}
