- Replaced FontAwesome with corresponding icons.
  - Tweaked paddings between icons.
  - Tweaked icons color

diff --git a/Engine/Plugins/Developer/Concert/ConcertSync/ConcertSyncClient/Source/ConcertSyncClient/Private/ConcertClientPresenceManager.cpp b/Engine/Plugins/Developer/Concert/ConcertSync/ConcertSyncClient/Source/ConcertSyncClient/Private/ConcertClientPresenceManager.cpp
--- a/Engine/Plugins/Developer/Concert/ConcertSync/ConcertSyncClient/Source/ConcertSyncClient/Private/ConcertClientPresenceManager.cpp
+++ b/Engine/Plugins/Developer/Concert/ConcertSync/ConcertSyncClient/Source/ConcertSyncClient/Private/ConcertClientPresenceManager.cpp
@@ -749,12 +749,14 @@ void FConcertClientPresenceManager::GetPresenceClientActions(const FConcertSessi
 	JumpToPresenceDef.Text = FEditorFontGlyphs::Map_Marker;
 	JumpToPresenceDef.ToolTipText = LOCTEXT("JumpToPresenceToolTip", "Jump to the presence location of this client");
 	JumpToPresenceDef.OnExecute.BindSP(this, &FConcertClientPresenceManager::OnJumpToPresence, InClientInfo.ClientEndpointId, FTransform::Identity);
+	JumpToPresenceDef.IconStyle = TEXT("Concert.JumpToLocation");
 
 	FConcertActionDefinition& ShowHidePresenceDef = OutActionDefs.AddDefaulted_GetRef();
 	ShowHidePresenceDef.IsEnabled = MakeAttributeSP(this, &FConcertClientPresenceManager::IsShowHidePresenceEnabled, InClientInfo.ClientEndpointId);
 	ShowHidePresenceDef.Text = MakeAttributeSP(this, &FConcertClientPresenceManager::GetShowHidePresenceText, InClientInfo.ClientEndpointId);
 	ShowHidePresenceDef.ToolTipText = MakeAttributeSP(this, &FConcertClientPresenceManager::GetShowHidePresenceToolTip, InClientInfo.ClientEndpointId);
 	ShowHidePresenceDef.OnExecute.BindSP(this, &FConcertClientPresenceManager::OnShowHidePresence, InClientInfo.ClientEndpointId);
+	ShowHidePresenceDef.IconStyle = MakeAttributeSP(this, &FConcertClientPresenceManager::GetShowHidePresenceIconStyle, InClientInfo.ClientEndpointId);
 }
 
 bool FConcertClientPresenceManager::IsJumpToPresenceEnabled(FGuid InEndpointId) const
@@ -853,6 +855,13 @@ FText FConcertClientPresenceManager::GetShowHidePresenceText(FGuid InEndpointId)
 		: FEditorFontGlyphs::Eye_Slash;
 }
 
+FName FConcertClientPresenceManager::GetShowHidePresenceIconStyle(FGuid InEndpointId) const
+{
+	return IsPresenceVisible(InEndpointId)
+		? TEXT("Concert.ShowPresence")  // Eye open icon.
+		: TEXT("Concert.HidePresence"); // Eye closed icon.
+}
+
 FText FConcertClientPresenceManager::GetShowHidePresenceToolTip(FGuid InEndpointId) const
 {
 	return IsPresenceVisible(InEndpointId)
diff --git a/Engine/Plugins/Developer/Concert/ConcertSync/ConcertSyncClient/Source/ConcertSyncClient/Private/ConcertClientPresenceManager.h b/Engine/Plugins/Developer/Concert/ConcertSync/ConcertSyncClient/Source/ConcertSyncClient/Private/ConcertClientPresenceManager.h
--- a/Engine/Plugins/Developer/Concert/ConcertSync/ConcertSyncClient/Source/ConcertSyncClient/Private/ConcertClientPresenceManager.h
+++ b/Engine/Plugins/Developer/Concert/ConcertSync/ConcertSyncClient/Source/ConcertSyncClient/Private/ConcertClientPresenceManager.h
@@ -233,6 +233,9 @@ private:
 
 	/** Get the correct show/hide text for the given endpoint */
 	FText GetShowHidePresenceText(FGuid InEndpointId) const;
+	
+	/** Get the correct show/hide icon style for the given endpoint */
+	FName GetShowHidePresenceIconStyle(FGuid InEndpointId) const;
 
 	/** Get the correct show/hide tooltip for the given endpoint */
 	FText GetShowHidePresenceToolTip(FGuid InEndpointId) const;
diff --git a/Engine/Plugins/Developer/Concert/ConcertSync/ConcertSyncClient/Source/ConcertSyncClient/Public/ConcertActionDefinition.h b/Engine/Plugins/Developer/Concert/ConcertSync/ConcertSyncClient/Source/ConcertSyncClient/Public/ConcertActionDefinition.h
--- a/Engine/Plugins/Developer/Concert/ConcertSync/ConcertSyncClient/Source/ConcertSyncClient/Public/ConcertActionDefinition.h
+++ b/Engine/Plugins/Developer/Concert/ConcertSync/ConcertSyncClient/Source/ConcertSyncClient/Public/ConcertActionDefinition.h
@@ -32,4 +32,5 @@ struct FConcertActionDefinition
 	TAttribute<FText> Text;        // The action name. Usually correspond to a button caption in UI.
 	TAttribute<FText> ToolTipText; // The action tooltip. Usually correspond to a button tooltip.
 	FOnExecuteAction  OnExecute;   // The function to call to execute the action.
+	TAttribute<FName> IconStyle;   // The icon brush style name (if set, it replaces the text).
 };
diff --git a/Engine/Plugins/Developer/Concert/ConcertSync/ConcertSyncClient/Source/ConcertSyncUI/Private/ConcertFrontendStyle.cpp b/Engine/Plugins/Developer/Concert/ConcertSync/ConcertSyncClient/Source/ConcertSyncUI/Private/ConcertFrontendStyle.cpp
--- a/Engine/Plugins/Developer/Concert/ConcertSync/ConcertSyncClient/Source/ConcertSyncUI/Private/ConcertFrontendStyle.cpp
+++ b/Engine/Plugins/Developer/Concert/ConcertSync/ConcertSyncClient/Source/ConcertSyncUI/Private/ConcertFrontendStyle.cpp
@@ -48,11 +48,38 @@ void FConcertFrontendStyle::Initialize()
 	StyleSet->Set("Concert.MyLock", new IMAGE_PLUGIN_BRUSH("Icons/icon_ConcertMyLock_16x", Icon16x16));
 	StyleSet->Set("Concert.OtherLock", new IMAGE_PLUGIN_BRUSH("Icons/icon_ConcertOtherLock_16x", Icon16x16));
 	StyleSet->Set("Concert.ModifiedByOther", new IMAGE_PLUGIN_BRUSH("Icons/icon_ConcertModifiedByOther_16x", Icon16x16));
-	StyleSet->Set("Concert.Stateless", new IMAGE_PLUGIN_BRUSH("Icons/icon_ConcertStateless_16x", Icon16x16)); // Used as icon aside on multi-user tabs title.
-	StyleSet->Set("Concert.OnlineServer", new IMAGE_PLUGIN_BRUSH("Icons/icon_ConcertServerOnline_16x", Icon16x16)); // Used to make a list of online server more attractive.
-	StyleSet->Set("Concert.DefaultServer", new IMAGE_PLUGIN_BRUSH("Icons/icon_ConcertServerDefault_16x", Icon16x16)); // Used to quickly spot the configured 'default server' in a list.
 
-	// 20x20 -> For toolbar small icons.
+	// Multi-user Tab/Menu icons
+	StyleSet->Set("Concert.MultiUser", new IMAGE_PLUGIN_BRUSH("Icons/icon_MultiUser_32x", Icon16x16));
+
+	// Maps the UI Command name in Multi-User module. (UI_COMMAND does magic icon mapping when style name and command name matches)
+	StyleSet->Set("Concert.OpenBrowser",  new IMAGE_PLUGIN_BRUSH("Icons/icon_MultiUser_32x", Icon16x16));
+	StyleSet->Set("Concert.OpenSettings", new IMAGE_PLUGIN_BRUSH("Icons/icon_Settings_32x",  Icon16x16));
+	StyleSet->Set("Concert.LaunchServer", new IMAGE_PLUGIN_BRUSH("Icons/icon_NewServer_32x", Icon16x16));
+
+	// Multi-User Browser
+	StyleSet->Set("Concert.ArchiveSession",     new IMAGE_PLUGIN_BRUSH("Icons/icon_ArchiveSession_48x",     Icon24x24));
+	StyleSet->Set("Concert.CancelAutoJoin",     new IMAGE_PLUGIN_BRUSH("Icons/icon_CancelAutoJoin_48x",     Icon24x24));
+	StyleSet->Set("Concert.CloseServer",        new IMAGE_PLUGIN_BRUSH("Icons/icon_CloseServer_48x",        Icon24x24));
+	StyleSet->Set("Concert.DeleteSession",      new IMAGE_PLUGIN_BRUSH("Icons/icon_DeleteSession_48x",      Icon24x24));
+	StyleSet->Set("Concert.JoinDefaultSession", new IMAGE_PLUGIN_BRUSH("Icons/icon_JoinDefaultSession_48x", Icon24x24));
+	StyleSet->Set("Concert.JoinSession",        new IMAGE_PLUGIN_BRUSH("Icons/icon_JoinSelectedSession_48x",Icon24x24));
+	StyleSet->Set("Concert.LeaveSession",       new IMAGE_PLUGIN_BRUSH("Icons/icon_LeaveSession_48x",       Icon24x24));
+	StyleSet->Set("Concert.NewServer",          new IMAGE_PLUGIN_BRUSH("Icons/icon_NewServer_48x",          Icon24x24));
+	StyleSet->Set("Concert.NewSession",         new IMAGE_PLUGIN_BRUSH("Icons/icon_NewSession_48x",         Icon24x24));
+	StyleSet->Set("Concert.PauseSession",       new IMAGE_PLUGIN_BRUSH("Icons/icon_PauseSession_48x",       Icon24x24));
+	StyleSet->Set("Concert.RestoreSession",     new IMAGE_PLUGIN_BRUSH("Icons/icon_RestoreSession_48x",     Icon24x24));
+	StyleSet->Set("Concert.ResumeSession",      new IMAGE_PLUGIN_BRUSH("Icons/icon_ResumeSession_48x",      Icon24x24));
+	StyleSet->Set("Concert.Settings",           new IMAGE_PLUGIN_BRUSH("Icons/icon_Settings_48x",           Icon24x24));
+	StyleSet->Set("Concert.NewServer.Small",    new IMAGE_PLUGIN_BRUSH("Icons/icon_NewServer_32x",          Icon16x16));
+	StyleSet->Set("Concert.NewSession.Small",   new IMAGE_PLUGIN_BRUSH("Icons/icon_NewSession_32x",         Icon16x16));
+
+	// Multi-user Active session
+	StyleSet->Set("Concert.JumpToLocation",     new IMAGE_PLUGIN_BRUSH("Icons/icon_PresenceLocation_32x",   Icon16x16));
+	StyleSet->Set("Concert.HidePresence",       new IMAGE_PLUGIN_BRUSH("Icons/icon_PresenceEyeOff_32x",     Icon16x16));
+	StyleSet->Set("Concert.ShowPresence",       new IMAGE_PLUGIN_BRUSH("Icons/icon_PresenceEyeOn_32x",      Icon16x16));
+
+	// 20x20 -> Editor toolbar small icons.
 	StyleSet->Set("Concert.Online.Small", new IMAGE_PLUGIN_BRUSH("Icons/icon_ConcertOnline_40x", Icon20x20));
 	StyleSet->Set("Concert.Offline.Small", new IMAGE_PLUGIN_BRUSH("Icons/icon_ConcertOffline_40x", Icon20x20));
 
@@ -62,7 +89,7 @@ void FConcertFrontendStyle::Initialize()
 	StyleSet->Set("Concert.Sequencer.SyncSequence",       new IMAGE_PLUGIN_BRUSH("Icons/icon_SequencerSyncSequence_48x", Icon48x48)); // Allows or not a remote client to open/close sequencer.
 	StyleSet->Set("Concert.Sequencer.SyncSequence.Small", new IMAGE_PLUGIN_BRUSH("Icons/icon_SequencerSyncSequence_48x", Icon24x24));
 
-	// 40x40
+	// 40x40 -> Editor toolbar large icons.
 	StyleSet->Set("Concert.Online", new IMAGE_PLUGIN_BRUSH("Icons/icon_ConcertOnline_40x", Icon40x40));
 	StyleSet->Set("Concert.Offline", new IMAGE_PLUGIN_BRUSH("Icons/icon_ConcertOffline_40x", Icon40x40));
 
@@ -76,6 +103,7 @@ void FConcertFrontendStyle::Initialize()
 	{
 		StyleSet->Set("Concert.Color.LocalUser", FLinearColor(0.31f, 0.749f, 0.333f));
 		StyleSet->Set("Concert.Color.OtherUser", FLinearColor(0.93f, 0.608f, 0.169f));
+		StyleSet->Set("Concert.Color.ToolbarButtonIcon", FLinearColor::White * FLinearColor(1.0f, 1.0f, 1.0f, 0.8f)); // Icons are made to have 80% opacity.
 	}
 
 	// Colors
diff --git a/Engine/Plugins/Developer/Concert/ConcertSync/ConcertSyncClient/Source/ConcertSyncUI/Public/ConcertFrontendUtils.h b/Engine/Plugins/Developer/Concert/ConcertSync/ConcertSyncClient/Source/ConcertSyncUI/Public/ConcertFrontendUtils.h
--- a/Engine/Plugins/Developer/Concert/ConcertSync/ConcertSyncClient/Source/ConcertSyncUI/Public/ConcertFrontendUtils.h
+++ b/Engine/Plugins/Developer/Concert/ConcertSync/ConcertSyncClient/Source/ConcertSyncUI/Public/ConcertFrontendUtils.h
@@ -12,6 +12,7 @@
 #include "Widgets/Layout/SExpandableArea.h"
 #include "Widgets/SBoxPanel.h"
 #include "EditorFontGlyphs.h"
+#include "ConcertFrontendStyle.h"
 
 #define LOCTEXT_NAMESPACE "ConcertFrontendUtils"
 
@@ -46,7 +47,7 @@ namespace ConcertFrontendUtils
 			];
 	}
 
-	inline TSharedRef<SButton> CreateButton(const FConcertActionDefinition& InDef)
+	inline TSharedRef<SButton> CreateTextButton(const FConcertActionDefinition& InDef)
 	{
 		const FButtonStyle* ButtonStyle = &FEditorStyle::Get().GetWidgetStyle<FButtonStyle>(ButtonStyleNames[(int32)InDef.Type]);
 		check(ButtonStyle);
@@ -74,6 +75,28 @@ namespace ConcertFrontendUtils
 			];
 	}
 
+	inline TSharedRef<SButton> CreateIconButton(const FConcertActionDefinition& InDef)
+	{
+		const FButtonStyle* ButtonStyle = &FEditorStyle::Get().GetWidgetStyle<FButtonStyle>(ButtonStyleNames[(int32)InDef.Type]);
+		check(ButtonStyle);
+
+		return SNew(SButton)
+			.ButtonStyle(ButtonStyle)
+			.ForegroundColor(FSlateColor::UseForeground())
+			.ToolTipText(InDef.ToolTipText)
+			.ContentPadding(FMargin(0, 0))
+			.IsEnabled(InDef.IsEnabled)
+			.Visibility_Lambda([IsVisible = InDef.IsVisible]() { return IsVisible.Get() ? EVisibility::Visible : EVisibility::Collapsed; })
+			.OnClicked_Lambda([OnExecute = InDef.OnExecute]() { OnExecute.ExecuteIfBound(); return FReply::Handled(); })
+			.VAlign(VAlign_Center)
+			.HAlign(HAlign_Center)
+			[
+				SNew(SImage)
+				.Image(TAttribute<const FSlateBrush*>::Create([IconStyleAttr = InDef.IconStyle]() { return FConcertFrontendStyle::Get()->GetBrush(IconStyleAttr.Get()); }))
+				.ColorAndOpacity(FConcertFrontendStyle::Get()->GetColor("Concert.Color.ToolbarButtonIcon"))
+			];
+	}
+
 	inline void AppendButtons(TSharedRef<SHorizontalBox> InHorizBox, TArrayView<const FConcertActionDefinition> InDefs)
 	{
 		for (const FConcertActionDefinition& Def : InDefs)
@@ -83,7 +106,7 @@ namespace ConcertFrontendUtils
 				.VAlign(VAlign_Center)
 				.Padding(FMargin(1.0f))
 				[
-					CreateButton(Def)
+					Def.IconStyle.IsSet() ? CreateIconButton(Def) : CreateTextButton(Def)
 				];
 		}
 	}
diff --git a/Engine/Plugins/Developer/Concert/MultiUserClient/Source/MultiUserClient/Private/MultiUserClientModule.cpp b/Engine/Plugins/Developer/Concert/MultiUserClient/Source/MultiUserClient/Private/MultiUserClientModule.cpp
--- a/Engine/Plugins/Developer/Concert/MultiUserClient/Source/MultiUserClient/Private/MultiUserClientModule.cpp
+++ b/Engine/Plugins/Developer/Concert/MultiUserClient/Source/MultiUserClient/Private/MultiUserClientModule.cpp
@@ -567,7 +567,7 @@ public:
 
 		FTabSpawnerEntry& BrowserSpawnerEntry = FGlobalTabmanager::Get()->RegisterNomadTabSpawner(ConcertBrowserTabName,
 			FOnSpawnTab::CreateRaw(this, &FMultiUserClientModule::SpawnConcertBrowserTab))
-			.SetIcon(FSlateIcon(FConcertFrontendStyle::GetStyleSetName(), TEXT("Concert.Stateless")))
+			.SetIcon(FSlateIcon(FConcertFrontendStyle::GetStyleSetName(), TEXT("Concert.MultiUser")))
 			.SetDisplayName(LOCTEXT("BrowserTabTitle", "Multi-User Browser"))
 			.SetTooltipText(LOCTEXT("BrowserTooltipText", "Open the Multi-User session browser"))
 			.SetMenuType(ETabSpawnerMenuType::Enabled);
diff --git a/Engine/Plugins/Developer/Concert/MultiUserClient/Source/MultiUserClient/Private/Widgets/SActiveSession.cpp b/Engine/Plugins/Developer/Concert/MultiUserClient/Source/MultiUserClient/Private/Widgets/SActiveSession.cpp
--- a/Engine/Plugins/Developer/Concert/MultiUserClient/Source/MultiUserClient/Private/Widgets/SActiveSession.cpp
+++ b/Engine/Plugins/Developer/Concert/MultiUserClient/Source/MultiUserClient/Private/Widgets/SActiveSession.cpp
@@ -337,27 +337,30 @@ void SActiveSession::Construct(const FArguments& InArgs, TSharedPtr<IConcertSync
 
 		// Resume Session
 		FConcertActionDefinition& ResumeSessionDef = ButtonDefs.AddDefaulted_GetRef();
-		ResumeSessionDef.Type = EConcertActionType::Success;
+		ResumeSessionDef.Type = EConcertActionType::Normal;
 		ResumeSessionDef.IsVisible = MakeAttributeSP(this, &SActiveSession::IsStatusBarResumeSessionVisible);
 		ResumeSessionDef.Text = FEditorFontGlyphs::Play_Circle;
 		ResumeSessionDef.ToolTipText = LOCTEXT("ResumeCurrentSessionToolTip", "Resume receiving updates from the current session");
 		ResumeSessionDef.OnExecute.BindLambda([this]() { OnClickResumeSession(); });
+		ResumeSessionDef.IconStyle = TEXT("Concert.ResumeSession");
 
 		// Suspend Session
 		FConcertActionDefinition& SuspendSessionDef = ButtonDefs.AddDefaulted_GetRef();
-		SuspendSessionDef.Type = EConcertActionType::Warning;
+		SuspendSessionDef.Type = EConcertActionType::Normal;
 		SuspendSessionDef.IsVisible = MakeAttributeSP(this, &SActiveSession::IsStatusBarSuspendSessionVisible);
 		SuspendSessionDef.Text = FEditorFontGlyphs::Pause_Circle;
 		SuspendSessionDef.ToolTipText = LOCTEXT("SuspendCurrentSessionToolTip", "Suspend receiving updates from the current session");
 		SuspendSessionDef.OnExecute.BindLambda([this]() { OnClickSuspendSession(); });
+		SuspendSessionDef.IconStyle = TEXT("Concert.PauseSession");
 
 		// Leave Session
 		FConcertActionDefinition& LeaveSessionDef = ButtonDefs.AddDefaulted_GetRef();
-		LeaveSessionDef.Type = EConcertActionType::Danger;
+		LeaveSessionDef.Type = EConcertActionType::Normal;
 		LeaveSessionDef.IsVisible = MakeAttributeSP(this, &SActiveSession::IsStatusBarLeaveSessionVisible);
 		LeaveSessionDef.Text = FEditorFontGlyphs::Sign_Out;
 		LeaveSessionDef.ToolTipText = LOCTEXT("LeaveCurrentSessionToolTip", "Leave the current session");
 		LeaveSessionDef.OnExecute.BindLambda([this]() { OnClickLeaveSession(); });
+		LeaveSessionDef.IconStyle = TEXT("Concert.LeaveSession");
 
 		ConcertFrontendUtils::AppendButtons(StatusBar, ButtonDefs);
 	}
diff --git a/Engine/Plugins/Developer/Concert/MultiUserClient/Source/MultiUserClient/Private/Widgets/SConcertBrowser.cpp b/Engine/Plugins/Developer/Concert/MultiUserClient/Source/MultiUserClient/Private/Widgets/SConcertBrowser.cpp
--- a/Engine/Plugins/Developer/Concert/MultiUserClient/Source/MultiUserClient/Private/Widgets/SConcertBrowser.cpp
+++ b/Engine/Plugins/Developer/Concert/MultiUserClient/Source/MultiUserClient/Private/Widgets/SConcertBrowser.cpp
@@ -52,24 +52,24 @@ namespace ConcertBrowserUtils
 {
 
 // Defines the sessions list view column tag names.
-extern const FName IconColName(TEXT("Icon"));
-extern const FName SessionColName(TEXT("Session"));
-extern const FName ServerColName(TEXT("Server"));
+const FName IconColName(TEXT("Icon"));
+const FName SessionColName(TEXT("Session"));
+const FName ServerColName(TEXT("Server"));
 
 // Name of the filter box in the View option.
-extern const FName ActiveSessionsCheckBoxMenuName(TEXT("ActiveSessions"));
-extern const FName ArchivedSessionsCheckBoxMenuName(TEXT("ArchivedSessions"));
-extern const FName DefaultServerCheckBoxMenuName(TEXT("DefaultServer"));
+const FName ActiveSessionsCheckBoxMenuName(TEXT("ActiveSessions"));
+const FName ArchivedSessionsCheckBoxMenuName(TEXT("ArchivedSessions"));
+const FName DefaultServerCheckBoxMenuName(TEXT("DefaultServer"));
 
 // The awesome font used to pick the icon displayed in the session list view 'Icon' column.
-extern const FName IconColumnFontName(TEXT("FontAwesome.9"));
+const FName IconColumnFontName(TEXT("FontAwesome.9"));
 
-/** Utility function used to create buttons displaying only an icon. */
-TSharedRef<SButton> MakeIconButton(const FName& ButtonStyle, const TAttribute<FText>& GlyphIcon, const TAttribute<FText>& Tooltip, const TAttribute<bool>& EnabledAttribute, const FOnClicked& OnClicked, 
-	const TAttribute<EVisibility>& Visibility = EVisibility::Visible, const TAttribute<FMargin>& ContentPadding = FMargin(3.0, 2.0))
+/** Utility function used to create buttons displaying only an icon (using FontAwesome). */
+TSharedRef<SButton> MakeIconButton(const FName& ButtonStyle, const TAttribute<FText>& GlyphIcon, const TAttribute<FText>& Tooltip, const TAttribute<bool>& EnabledAttribute, const FOnClicked& OnClicked,
+	const FSlateColor& ForegroundColor, const TAttribute<EVisibility>& Visibility = EVisibility::Visible, const TAttribute<FMargin>& ContentPadding = FMargin(3.0, 2.0), const FName FontStyle = IconColumnFontName)
 {
 	return SNew(SButton)
-		.ForegroundColor(FLinearColor::White)
+		.ForegroundColor(ForegroundColor)
 		.ButtonStyle(FEditorStyle::Get(), ButtonStyle)
 		.OnClicked(OnClicked)
 		.ToolTipText(Tooltip)
@@ -80,11 +80,29 @@ TSharedRef<SButton> MakeIconButton(const FName& ButtonStyle, const TAttribute<FT
 		.HAlign(HAlign_Center)
 		[
 			SNew(STextBlock)
-			.Font(FEditorStyle::Get().GetFontStyle("FontAwesome.9"))
+			.Font(FEditorStyle::Get().GetFontStyle(FontStyle))
 			.Text(GlyphIcon)
 		];
 }
 
+/** Utility function used to create buttons displaying only an icon (using a brush) */
+TSharedRef<SButton> MakeIconButton(const FName& ButtonStyle, const TAttribute<const FSlateBrush*>& Icon, const TAttribute<FText>& Tooltip, const TAttribute<bool>& EnabledAttribute, const FOnClicked& OnClicked, const TAttribute<EVisibility>& Visibility = EVisibility::Visible)
+{
+	return SNew(SButton)
+		.ForegroundColor(FSlateColor::UseForeground())
+		.ButtonStyle(FEditorStyle::Get(), ButtonStyle)
+		.OnClicked(OnClicked)
+		.ToolTipText(Tooltip)
+		.ContentPadding(FMargin(0, 0))
+		.Visibility(Visibility)
+		.IsEnabled(EnabledAttribute)
+		.VAlign(VAlign_Center)
+		.HAlign(HAlign_Center)
+		[
+			SNew(SImage).Image(Icon).ColorAndOpacity(FConcertFrontendStyle::Get()->GetColor("Concert.Color.ToolbarButtonIcon"))
+		];
+}
+
 } // namespace ConcertSessionBrowserUtils
 
 /** Signal emitted when a session name text field should enter in edit mode. */
@@ -285,6 +303,7 @@ private:
 	// True if a request to create a session is ongoing.
 	TArray<FAsyncRequest> CreateSessionRequests;
 	TArray<FPendingSessionDiscovery> ExpectedSessionsToDiscover;
+	TSet<FString> IgnoredServers; // List of ignored servers (Useful for testing/debugging)
 };
 
 
@@ -294,6 +313,9 @@ FConcertBrowserController::FConcertBrowserController(IConcertClientPtr InConcert
 	check(InConcertClient->IsConfigured()); // Expected to be done by higher level code.
 	ConcertClient = InConcertClient;
 
+	// When others servers are running, add them to the list if you want to test the UI displayed if no servers/no sessions exists.
+	//IgnoredServers.Add(TEXT("wksyul10355")); // TODO: COMMENT BEFORE SUBMIT
+
 	// Start server discovery to find the available Concert servers.
 	ConcertClient->StartDiscovery();
 
@@ -510,6 +532,10 @@ void FConcertBrowserController::UpdateSessionsAsync()
 {
 	// Get the list of known servers.
 	TArray<FConcertServerInfo> OnlineServers = ConcertClient->GetKnownServers();
+	if (IgnoredServers.Num())
+	{
+		OnlineServers.RemoveAll([this](const FConcertServerInfo& ServerInfo) { return IgnoredServers.Contains(ServerInfo.ServerName); });
+	}
 
 	bool bServerListVersionUpdated = false;
 
@@ -824,18 +850,18 @@ void FConcertBrowserController::UpdateClientsAsync(const FGuid& ServerAdminEndpo
 class SSessionRow : public SMultiColumnTableRow<TSharedPtr<FConcertSessionItem>>
 {
 public:
-	typedef TFunction<void(TSharedPtr<FConcertSessionItem>)> FClickIconFunc;
+	typedef TFunction<void(TSharedPtr<FConcertSessionItem>)> FDoubleClickFunc;
 	typedef TFunction<void(TSharedPtr<FConcertSessionItem>, const FString&)> FRenameFunc;
 
 	SLATE_BEGIN_ARGS(SSessionRow)
-		: _OnClickIconFunc()
+		: _OnDoubleClickFunc()
 		, _OnRenameFunc()
 		, _HighlightText()
 		, _IsSelected(false)
 	{
 	}
 
-	SLATE_ARGUMENT(FClickIconFunc, OnClickIconFunc)
+	SLATE_ARGUMENT(FDoubleClickFunc, OnDoubleClickFunc)
 	SLATE_ARGUMENT(FRenameFunc, OnRenameFunc)
 	SLATE_ATTRIBUTE(FText, HighlightText)
 	SLATE_ATTRIBUTE(bool, IsSelected)
@@ -843,11 +869,10 @@ public:
 	SLATE_END_ARGS()
 
 public:
-	void Construct(const FArguments& InArgs, TSharedPtr<FConcertSessionItem> InItem, const TSharedRef<STableViewBase>& InOwnerTableView);
+	void Construct(const FArguments& InArgs, TSharedPtr<FConcertSessionItem> InItem, const TSharedRef<STableViewBase>& InOwnerTableView, IConcertClientPtr InConcertClient);
 	virtual TSharedRef<SWidget> GenerateWidgetForColumn(const FName& ColumnName) override;
 	virtual FReply OnMouseButtonDoubleClick( const FGeometry& InMyGeometry, const FPointerEvent& InMouseEvent) override;
 
-	FReply OnIconClicked();
 	void OnSessionNameCommitted(const FText& NewSessionName, ETextCommit::Type CommitType);
 
 private:
@@ -856,21 +881,23 @@ private:
 
 private:
 	TWeakPtr<FConcertSessionItem> Item;
-	FClickIconFunc ClickIconFunc; // Invoked when the user clicks the row icon.
+	FDoubleClickFunc DoubleClickFunc; // Invoked when the user double click on the row.
 	FRenameFunc RenameFunc; // Invoked when the user commit the session rename. (This will send the request to server)
 	TAttribute<FText> HighlightText;
 	TAttribute<bool> IsSelected;
 	TSharedPtr<SInlineEditableTextBlock> SessionNameText;
+	IConcertClientPtr ConcertClient;
 };
 
 
-void SSessionRow::Construct(const FArguments& InArgs, TSharedPtr<FConcertSessionItem> InItem, const TSharedRef<STableViewBase>& InOwnerTableView)
+void SSessionRow::Construct(const FArguments& InArgs, TSharedPtr<FConcertSessionItem> InItem, const TSharedRef<STableViewBase>& InOwnerTableView, IConcertClientPtr InConcertClient)
 {
 	Item = MoveTemp(InItem);
-	ClickIconFunc = InArgs._OnClickIconFunc; // This function should join a session or add a row to restore an archive.
+	DoubleClickFunc = InArgs._OnDoubleClickFunc; // This function should join a session or add a row to restore an archive.
 	RenameFunc = InArgs._OnRenameFunc; // Function invoked to send a rename request to the server.
 	HighlightText = InArgs._HighlightText;
 	IsSelected = InArgs._IsSelected;
+	ConcertClient = InConcertClient;
 
 	// Construct base class
 	SMultiColumnTableRow<TSharedPtr<FConcertSessionItem>>::Construct(FSuperRowType::FArguments(), InOwnerTableView);
@@ -885,60 +912,97 @@ TSharedRef<SWidget> SSessionRow::GenerateWidgetForColumn(const FName& ColumnName
 
 	if (ColumnName == ConcertBrowserUtils::IconColName)
 	{
-		return SNew(SButton)
-			.ForegroundColor(FSlateColor::UseForeground())
-			.ButtonStyle(FEditorStyle::Get(), TEXT("FlatButton"))
-			.OnClicked(this, &SSessionRow::OnIconClicked)
-			.ContentPadding(FMargin(0, 2))
-			[
-				SNew(SBox)
-				.VAlign(VAlign_Center)
-				.HAlign(HAlign_Center)
-				.ToolTipText(ItemPin->Type == FConcertSessionItem::EType::ActiveSession ? LOCTEXT("JoinIconTooltip", "Join the session") : LOCTEXT("RestoreIconTooltip", "Restore the session"))
-				[
-					SNew(STextBlock)
-						.Font(FEditorStyle::Get().GetFontStyle(ConcertBrowserUtils::IconColumnFontName))
-						.Text(ItemPin->Type == FConcertSessionItem::EType::ActiveSession ? FText::FromString(TEXT("\xf044")) : FEditorFontGlyphs::File_Archive_O) // \xf044 -> edit
-				]
-			];
+		return SNew(SBox)
+		.VAlign(VAlign_Center)
+		.HAlign(HAlign_Center)
+		.Padding(2)
+		.ToolTipText(ItemPin->Type == FConcertSessionItem::EType::ActiveSession ? LOCTEXT("ActiveIconTooltip", "Active session") : LOCTEXT("ArchivedIconTooltip", "Archived Session"))
+		[
+			SNew(STextBlock)
+			.Font(FEditorStyle::Get().GetFontStyle(ConcertBrowserUtils::IconColumnFontName))
+			.Text(ItemPin->Type == FConcertSessionItem::EType::ActiveSession ? FEditorFontGlyphs::Circle : FEditorFontGlyphs::Archive)
+			.ColorAndOpacity(ItemPin->Type == FConcertSessionItem::EType::ActiveSession ? FEditorStyle::Get().GetWidgetStyle<FButtonStyle>("FlatButton.Success").Normal.TintColor : FSlateColor::UseSubduedForeground())
+		];
 	}
-	else if (ColumnName == ConcertBrowserUtils::SessionColName)
+
+	bool bIsDefaultConfig = (ItemPin->Type == FConcertSessionItem::EType::ActiveSession && ItemPin->SessionName == ConcertClient->GetConfiguration()->DefaultSessionName && ItemPin->ServerName == ConcertClient->GetConfiguration()->DefaultServerURL);
+	FSlateFontInfo FontInfo;
+	FSlateColor FontColor;
+	if (ItemPin->Type == FConcertSessionItem::EType::ActiveSession)
+	{
+		FontColor = bIsDefaultConfig ? FSlateColor(FLinearColor::White) : FSlateColor(FLinearColor::White * 0.8f);
+		FontInfo = FEditorStyle::Get().GetFontStyle("NormalFont");
+	}
+	else
+	{
+		FontColor = FSlateColor::UseSubduedForeground();
+		FontInfo = FCoreStyle::GetDefaultFontStyle("Italic", 9);
+	}
+
+	if (ColumnName == ConcertBrowserUtils::SessionColName)
 	{
 		return SNew(SBox)
 			.VAlign(VAlign_Center)
 			[
 				SAssignNew(SessionNameText, SInlineEditableTextBlock)
-				.Text_Lambda([this]() { return FText::FromString(Item.Pin()->SessionName); })
+				.Text_Lambda([this]() { return FText::AsCultureInvariant(Item.Pin()->SessionName); })
 				.HighlightText(HighlightText)
 				.OnTextCommitted(this, &SSessionRow::OnSessionNameCommitted)
 				.IsReadOnly(false)
 				.IsSelected(FIsSelected::CreateLambda([this]() { return IsSelected.Get(); }))
 				.OnVerifyTextChanged(this, &SSessionRow::OnValidatingSessionName)
+				.Font(FontInfo)
+				.ColorAndOpacity(FontColor)
 			];
 	}
 
 	check(ColumnName == ConcertBrowserUtils::ServerColName);
-	return SNew(SBox)
-		.VAlign(VAlign_Center)
-		[
-			SNew(STextBlock)
-			.Text(FText::FromString(ItemPin->ServerName))
-			.HighlightText(HighlightText)
-		];
-}
 
-FReply SSessionRow::OnIconClicked()
-{
-	if (TSharedPtr<FConcertSessionItem> ItemPin = Item.Pin())
+	if (bIsDefaultConfig)
+	{
+		return SNew(SHorizontalBox)
+			+SHorizontalBox::Slot()
+			.AutoWidth()
+			.VAlign(VAlign_Center)
+			[
+				SNew(STextBlock)
+				.Text(FText::Format(INVTEXT("{0} * "), FText::AsCultureInvariant(ItemPin->ServerName)))
+				.HighlightText(HighlightText)
+				.Font(FontInfo)
+				.ColorAndOpacity(FontColor)
+			]
+			+SHorizontalBox::Slot()
+			.AutoWidth()
+			.VAlign(VAlign_Center)
+			[
+				SNew(STextBlock)
+				.Text(LOCTEXT("DefaultServerSession", "(Default Session/Server)"))
+				.HighlightText(HighlightText)
+				.Font(FCoreStyle::GetDefaultFontStyle("Regular", 9))
+				.ColorAndOpacity(FontColor)
+			];
+	}
+	else
 	{
-		ClickIconFunc(ItemPin);
+		return SNew(SBox)
+			.VAlign(VAlign_Center)
+			[
+				SNew(STextBlock)
+				.Text(FText::AsCultureInvariant(ItemPin->ServerName))
+				.HighlightText(HighlightText)
+				.Font(FontInfo)
+				.ColorAndOpacity(FontColor)
+			];
 	}
-	return FReply::Handled();
 }
 
 FReply SSessionRow::OnMouseButtonDoubleClick(const FGeometry& InMyGeometry, const FPointerEvent& InMouseEvent)
 {
-	return OnIconClicked(); // Just like clicking on the icon.
+	if (TSharedPtr<FConcertSessionItem> ItemPin = Item.Pin())
+	{
+		DoubleClickFunc(ItemPin);
+	}
+	return FReply::Handled();
 }
 
 bool SSessionRow::OnValidatingSessionName(const FText& NewSessionName, FText& OutError)
@@ -1123,7 +1187,8 @@ TSharedRef<SWidget> SNewSessionRow::GenerateWidgetForColumn(const FName& ColumnN
 						FEditorFontGlyphs::Check,
 						LOCTEXT("CreateCheckIconTooltip", "Create the session"),
 						TAttribute<bool>::Create([this]() { return !EditableSessionName->GetText().IsEmpty(); }),
-						FOnClicked::CreateRaw(this, &SNewSessionRow::OnAccept))
+						FOnClicked::CreateRaw(this, &SNewSessionRow::OnAccept),
+						FSlateColor(FLinearColor::White))
 				]
 
 				// 'Decline' button
@@ -1134,7 +1199,8 @@ TSharedRef<SWidget> SNewSessionRow::GenerateWidgetForColumn(const FName& ColumnN
 						FEditorFontGlyphs::Times,
 						LOCTEXT("CancelIconTooltip", "Cancel"),
 						true, // Always enabled.
-						FOnClicked::CreateRaw(this, &SNewSessionRow::OnDecline))
+						FOnClicked::CreateRaw(this, &SNewSessionRow::OnDecline),
+						FSlateColor(FLinearColor::White))
 				]
 			];
 	}
@@ -1143,7 +1209,6 @@ TSharedRef<SWidget> SNewSessionRow::GenerateWidgetForColumn(const FName& ColumnN
 TSharedRef<SWidget> SNewSessionRow::OnGenerateServersComboOptionWidget(TSharedPtr<FConcertServerInfo> ServerItem)
 {
 	bool bIsDefaultServer = ServerItem->ServerName == ConcertClient->GetConfiguration()->DefaultServerURL;
-	const FSlateBrush* Brush = bIsDefaultServer ? FConcertFrontendStyle::Get()->GetBrush("Concert.DefaultServer") : FConcertFrontendStyle::Get()->GetBrush("Concert.OnlineServer");
 
 	FText Tooltip;
 	if (bIsDefaultServer)
@@ -1165,13 +1230,9 @@ TSharedRef<SWidget> SNewSessionRow::OnGenerateServersComboOptionWidget(TSharedPt
 		+SHorizontalBox::Slot()
 		.AutoWidth()
 		[
-			SNew(SImage).Image(Brush)
-		]
-
-		+SHorizontalBox::Slot()
-		.AutoWidth()
-		[
-			SNew(STextBlock).Text(GetServerDisplayName(ServerItem->ServerName))
+			SNew(STextBlock)
+			.Font(bIsDefaultServer ? FEditorStyle::GetFontStyle("BoldFont") : FEditorStyle::GetFontStyle("NormalFont"))
+			.Text(GetServerDisplayName(ServerItem->ServerName))
 		];
 }
 
@@ -1263,7 +1324,12 @@ FText SNewSessionRow::GetSelectedServerText() const
 
 FText SNewSessionRow::GetServerDisplayName(const FString& ServerName) const
 {
-	if (ServerName == FPlatformProcess::ComputerName())
+	bool bIsDefaultServer = ServerName == ConcertClient->GetConfiguration()->DefaultServerURL;
+	if (bIsDefaultServer)
+	{
+		return FText::Format(LOCTEXT("DefaultServer", "{0} (Default)"), FText::FromString(FPlatformProcess::ComputerName()));
+	}
+	else if (ServerName == FPlatformProcess::ComputerName())
 	{
 		return FText::Format(LOCTEXT("MyComputer", "{0} (My Computer)"), FText::FromString(FPlatformProcess::ComputerName()));
 	}
@@ -1426,21 +1492,11 @@ TSharedRef<SWidget> SSaveRestoreSessionRow::GenerateWidgetForColumn(const FName&
 	{
 		return SNew(SHorizontalBox)
 
-			// 'Restore' icon
-			+SHorizontalBox::Slot()
-			.VAlign(VAlign_Center)
-			.AutoWidth()
-			[
-				SNew(STextBlock)
-					.Font(FEditorStyle::Get().GetFontStyle(ConcertBrowserUtils::IconColumnFontName))
-					.Text(ItemPin->Type == FConcertSessionItem::EType::RestoreSession ? FEditorFontGlyphs::Arrow_Circle_O_Up : FEditorFontGlyphs::File_Archive_O)
-			]
-
 			// 'Restore as/Save as' text
 			+SHorizontalBox::Slot()
 			.VAlign(VAlign_Center)
 			.AutoWidth()
-			.Padding(4.0, 0.0)
+			.Padding(2.0, 0.0)
 			[
 				SNew(STextBlock).Text(ItemPin->Type == FConcertSessionItem::EType::RestoreSession ? LOCTEXT("RestoreAs", "Restore as:") : LOCTEXT("ArchiveAs", "Archive as:"))
 			]
@@ -1494,7 +1550,8 @@ TSharedRef<SWidget> SSaveRestoreSessionRow::GenerateWidgetForColumn(const FName&
 					FEditorFontGlyphs::Check,
 					ItemPin->Type == FConcertSessionItem::EType::RestoreSession ? LOCTEXT("RestoreCheckIconTooltip", "Restore the session") : LOCTEXT("ArchiveCheckIconTooltip", "Archive the session"),
 					TAttribute<bool>::Create([this]() { return !EditableSessionName->GetText().IsEmpty(); }), // Enabled?
-					FOnClicked::CreateRaw(this, &SSaveRestoreSessionRow::OnAccept))
+					FOnClicked::CreateRaw(this, &SSaveRestoreSessionRow::OnAccept),
+					FSlateColor(FLinearColor::White))
 			]
 
 			// 'Cancel' button
@@ -1505,7 +1562,8 @@ TSharedRef<SWidget> SSaveRestoreSessionRow::GenerateWidgetForColumn(const FName&
 					FEditorFontGlyphs::Times,
 					LOCTEXT("CancelTooltip", "Cancel"),
 					true, // Enabled?
-					FOnClicked::CreateRaw(this, &SSaveRestoreSessionRow::OnDecline))
+					FOnClicked::CreateRaw(this, &SSaveRestoreSessionRow::OnDecline),
+					FSlateColor(FLinearColor::White))
 			]
 		];
 	}
@@ -1592,7 +1650,7 @@ public:
 	SLATE_ATTRIBUTE(EVisibility, ButtonVisibility)
 	SLATE_ATTRIBUTE(bool, IsButtonEnabled)
 	SLATE_STYLE_ARGUMENT(FButtonStyle, ButtonStyle)
-	SLATE_ATTRIBUTE(FText, ButtonIcon)
+	SLATE_ATTRIBUTE(const FSlateBrush*, ButtonIcon)
 	SLATE_ATTRIBUTE(FText, ButtonText)
 	SLATE_ATTRIBUTE(FText, ButtonToolTip)
 	SLATE_EVENT( FOnClicked, OnButtonClicked)
@@ -1631,7 +1689,7 @@ public:
 				.HAlign(HAlign_Center)
 				[
 					SNew(SButton)
-					.ForegroundColor(FLinearColor::White)
+					.ForegroundColor(FSlateColor::UseForeground())
 					.ButtonStyle(InArgs._ButtonStyle)
 					.Visibility(InArgs._ButtonVisibility)
 					.HAlign(HAlign_Center)
@@ -1645,11 +1703,9 @@ public:
 						+SHorizontalBox::Slot()
 						.AutoWidth()
 						.VAlign(VAlign_Center)
-						.Padding(0, 1, 2, 0)
+						.Padding(0, 0, 3, 0)
 						[
-							SNew(STextBlock)
-							.Font(FEditorStyle::Get().GetFontStyle("FontAwesome.8"))
-							.Text(InArgs._ButtonIcon)
+							SNew(SImage).Image(InArgs._ButtonIcon).ColorAndOpacity(FConcertFrontendStyle::Get()->GetColor("Concert.Color.ToolbarButtonIcon"))
 						]
 						+SHorizontalBox::Slot()
 						.AutoWidth()
@@ -1867,13 +1923,13 @@ void SConcertSessionBrowser::Construct(const FArguments& InArgs, IConcertClientP
 		.Text(LOCTEXT("LookingForServer", "Looking for Multi-User Servers..."))
 		.Visibility_Lambda([this]() { return Controller->GetServers().Num() == 0 ? EVisibility::Visible : EVisibility::Hidden; })
 		.IsButtonEnabled(this, &SConcertSessionBrowser::IsLaunchServerButtonEnabled)
-		.ButtonStyle(FEditorStyle::Get(), "FlatButton.Info")
-		.ButtonIcon(FEditorFontGlyphs::Server)
+		.ButtonStyle(FEditorStyle::Get(), "FlatButton.Default")
+		.ButtonIcon(FConcertFrontendStyle::Get()->GetBrush("Concert.NewServer.Small"))
 		.ButtonText(LOCTEXT("LaunchLocalServer", "Launch a Server"))
 		.ButtonToolTip(LOCTEXT("LaunchServerTooltip", "Launch a Multi-User server on your computer unless one is already running"))
 		.OnButtonClicked(this, &SConcertSessionBrowser::OnLaunchServerButtonClicked);
 
-	// Controls the text displayed in the 'No session' panel.
+	// Controls the text displayed in the 'No sessions' panel.
 	auto GetNoSessionText = [this]()
 	{
 		if (!Controller->HasReceivedInitialSessionList())
@@ -1882,8 +1938,8 @@ void SConcertSessionBrowser::Construct(const FArguments& InArgs, IConcertClientP
 		}
 
 		return Controller->GetActiveSessions().Num() == 0 && Controller->GetArchivedSessions().Num() == 0 ?
-			LOCTEXT("NoSessionAvailable", "No Session Available") :
-			LOCTEXT("AllSessionsFilteredOut", "No Session Matches the Filters\nChange Your Filter to View Sessions");
+			LOCTEXT("NoSessionAvailable", "No Sessions Available") :
+			LOCTEXT("AllSessionsFilteredOut", "No Sessions Match the Filters\nChange Your Filter to View Sessions");
 	};
 
 	// Displayed when discovering session or if no session is available.
@@ -1892,15 +1948,15 @@ void SConcertSessionBrowser::Construct(const FArguments& InArgs, IConcertClientP
 		.Visibility_Lambda([this]() { return Controller->GetServers().Num() > 0 && Sessions.Num() == 0 && !Controller->IsCreatingSession() ? EVisibility::Visible : EVisibility::Hidden; })
 		.ThrobberVisibility_Lambda([this]() { return !Controller->HasReceivedInitialSessionList() ? EVisibility::Visible : EVisibility::Collapsed; })
 		.ButtonVisibility_Lambda([this]() { return Controller->HasReceivedInitialSessionList() && Controller->GetActiveSessions().Num() == 0 && Controller->GetArchivedSessions().Num() == 0 ? EVisibility::Visible : EVisibility::Collapsed; })
-		.ButtonStyle(FEditorStyle::Get(), "FlatButton.Success")
-		.ButtonIcon(FEditorFontGlyphs::Plus)
+		.ButtonStyle(FEditorStyle::Get(), "FlatButton.Default")
+		.ButtonIcon(FConcertFrontendStyle::Get()->GetBrush("Concert.NewSession.Small"))
 		.ButtonText(LOCTEXT("CreateSession", "Create Session"))
 		.ButtonToolTip(LOCTEXT("CreateSessionTooltip", "Create a new session"))
 		.OnButtonClicked(FOnClicked::CreateSP(this, &SConcertSessionBrowser::OnNewButtonClicked));
 
 	// Displayed when the selected session client view is empty (no client to display).
 	NoClientPanel = SNew(SConcertNoAvailability)
-		.Text(LOCTEXT("NoClientAvailable", "No Connected Client"))
+		.Text(LOCTEXT("NoClientAvailable", "No Connected Clients"))
 		.Visibility_Lambda([this]() { return Clients.Num() == 0 ? EVisibility::Visible : EVisibility::Hidden; });
 
 	// Displayed as details when no session is selected. (No session selected or the selected session doesn't have any)
@@ -2201,7 +2257,7 @@ TSharedRef<SWidget> SConcertSessionBrowser::MakeControlBar()
 		// The search text.
 		+SHorizontalBox::Slot()
 		.FillWidth(1.0)
-		.Padding(4.0f, 1.0f, 8.0f, 1.0f)
+		.Padding(4.0f, 3.0f, 8.0f, 3.0f)
 		[
 			SAssignNew(SearchBox, SSearchBox)
 			.HintText(LOCTEXT("SearchHint", "Search Session"))
@@ -2244,7 +2300,7 @@ TSharedRef<SWidget> SConcertSessionBrowser::MakeControlBar()
 			.HAlign(HAlign_Fill)
 			[
 				SNew(STextBlock)
-				.Font(FEditorStyle::Get().GetFontStyle("FontAwesome.12"))
+				.Font(FEditorStyle::Get().GetFontStyle("FontAwesome.16"))
 				.Text(FEditorFontGlyphs::Cogs)
 			]
 		];
@@ -2273,108 +2329,99 @@ TSharedRef<SWidget> SConcertSessionBrowser::MakeButtonBar()
 			FText::FromString(Controller->GetConcertClient()->GetConfiguration()->DefaultServerURL));
 	});
 
+	constexpr float PaddingBetweenButtons = 3.0f;
+
 	return SNew(SHorizontalBox)
+
+		// Launch server.
 		+SHorizontalBox::Slot()
 		.AutoWidth()
+		.Padding(0, 0, PaddingBetweenButtons, 0)
 		[
-			SNew(SUniformGridPanel)
-				.SlotPadding(FMargin(1.0f, 0.0f))
-				.MinDesiredSlotWidth(20)
+			ConcertBrowserUtils::MakeIconButton(TEXT("FlatButton"), FConcertFrontendStyle::Get()->GetBrush("Concert.NewServer"), LOCTEXT("LaunchServerTooltip", "Launch a Multi-User server on your computer unless one is already running"),
+				TAttribute<bool>(this, &SConcertSessionBrowser::IsLaunchServerButtonEnabled),
+				FOnClicked::CreateSP(this, &SConcertSessionBrowser::OnLaunchServerButtonClicked))
+		]
 
-			// Launch server.
-			+SUniformGridPanel::Slot(0, 0)
-			[
-				ConcertBrowserUtils::MakeIconButton(TEXT("FlatButton.Info"), FEditorFontGlyphs::Server, LOCTEXT("LaunchServerTooltip", "Launch a Multi-User server on your computer unless one is already running"),
-					TAttribute<bool>(this, &SConcertSessionBrowser::IsLaunchServerButtonEnabled),
-					FOnClicked::CreateSP(this, &SConcertSessionBrowser::OnLaunchServerButtonClicked),
-					EVisibility::Visible,
-					FMargin(6, 1, 6, 0))
-			]
+		// New Session
+		+SHorizontalBox::Slot()
+		.AutoWidth()
+		.Padding(0, 0, PaddingBetweenButtons, 0)
+		[
+			ConcertBrowserUtils::MakeIconButton(TEXT("FlatButton"), FConcertFrontendStyle::Get()->GetBrush("Concert.NewSession"), LOCTEXT("NewButtonTooltip", "Create a new session"),
+				TAttribute<bool>(this, &SConcertSessionBrowser::IsNewButtonEnabled),
+				FOnClicked::CreateSP(this, &SConcertSessionBrowser::OnNewButtonClicked))
+		]
 
-			// Auto-join
-			+SUniformGridPanel::Slot(1, 0)
-			[
-				ConcertBrowserUtils::MakeIconButton(TEXT("FlatButton.Info"), FEditorFontGlyphs::Arrow_Circle_Right, AutoJoinTooltip,
+		// Separator
+		+SHorizontalBox::Slot()
+		.AutoWidth()
+		.Padding(0, 1, PaddingBetweenButtons, 1)
+		[
+			SNew(SSeparator).Orientation(Orient_Vertical)
+		]
+
+		// Auto-join
+		+SHorizontalBox::Slot()
+		.AutoWidth()
+		.Padding(0, 0, PaddingBetweenButtons, 0)
+		[
+			ConcertBrowserUtils::MakeIconButton(TEXT("FlatButton"), FConcertFrontendStyle::Get()->GetBrush("Concert.JoinDefaultSession"), AutoJoinTooltip,
 				TAttribute<bool>(this, &SConcertSessionBrowser::IsAutoJoinButtonEnabled),
 				FOnClicked::CreateSP(this, &SConcertSessionBrowser::OnAutoJoinButtonClicked),
-				TAttribute<EVisibility>::Create([this]() { return !IsCancelAutoJoinButtonEnabled() ? EVisibility::Visible : EVisibility::Hidden; }), // Default button shown if both auto-join/cancel are disabled.
-				FMargin(5, 0, 6, 0))
-			]
+				TAttribute<EVisibility>::Create([this]() { return !IsCancelAutoJoinButtonEnabled() ? EVisibility::Visible : EVisibility::Collapsed; })) // Default button shown if both auto-join/cancel are disabled.
+		]
 
-			// Cancel auto join. (Shares the same slot as Auto-join)
-			+SUniformGridPanel::Slot(1, 0)
-			[
-				ConcertBrowserUtils::MakeIconButton(TEXT("FlatButton.Danger"), FEditorFontGlyphs::Times, CancelAutoJoinTooltip,
+		// Cancel auto join. (Shares the same slot as Auto-join)
+		+SHorizontalBox::Slot()
+		.AutoWidth()
+		.Padding(0, 0, PaddingBetweenButtons, 0)
+		[
+			ConcertBrowserUtils::MakeIconButton(TEXT("FlatButton"), FConcertFrontendStyle::Get()->GetBrush("Concert.CancelAutoJoin"), CancelAutoJoinTooltip,
 				TAttribute<bool>(this, &SConcertSessionBrowser::IsCancelAutoJoinButtonEnabled),
 				FOnClicked::CreateSP(this, &SConcertSessionBrowser::OnCancelAutoJoinButtonClicked),
-				TAttribute<EVisibility>::Create([this]() { return IsCancelAutoJoinButtonEnabled() ? EVisibility::Visible : EVisibility::Hidden; }),
-				FMargin(6, 0, 6, 0))
-			]
+				TAttribute<EVisibility>::Create([this]() { return IsCancelAutoJoinButtonEnabled() ? EVisibility::Visible : EVisibility::Collapsed; }))
 		]
 
+		// Join
 		+SHorizontalBox::Slot()
 		.AutoWidth()
-		.Padding(1, 1)
+		.Padding(0, 0, PaddingBetweenButtons, 0)
 		[
-			SNew(SSeparator).Orientation(Orient_Vertical)
+			ConcertBrowserUtils::MakeIconButton(TEXT("FlatButton"), FConcertFrontendStyle::Get()->GetBrush("Concert.JoinSession"), LOCTEXT("JoinButtonTooltip", "Join the selected session"),
+				TAttribute<bool>(this, &SConcertSessionBrowser::IsJoinButtonEnabled),
+				FOnClicked::CreateSP(this, &SConcertSessionBrowser::OnJoinButtonClicked),
+				TAttribute<EVisibility>::Create([this]() { return !IsRestoreButtonEnabled() ? EVisibility::Visible : EVisibility::Collapsed; })) // Default button shown if both join/restore are disabled.
 		]
 
+		// Restore (Share the same slot as Join)
 		+SHorizontalBox::Slot()
 		.AutoWidth()
+		.Padding(0, 0, PaddingBetweenButtons, 0)
 		[
-			SNew(SUniformGridPanel)
-			.SlotPadding(FMargin(1.0f, 0.0f))
-			.MinDesiredSlotWidth(20)
-
-			// New
-			+SUniformGridPanel::Slot(0, 0)
-			[
-				ConcertBrowserUtils::MakeIconButton(TEXT("FlatButton.Success"), FEditorFontGlyphs::Plus, LOCTEXT("NewButtonTooltip", "Create a new session"),
-					TAttribute<bool>(this, &SConcertSessionBrowser::IsNewButtonEnabled), 
-					FOnClicked::CreateSP(this, &SConcertSessionBrowser::OnNewButtonClicked),
-					EVisibility::Visible,
-					FMargin(6, 1, 6, 0))
-			]
-
-			// Join
-			+SUniformGridPanel::Slot(1, 0)
-			[
-				ConcertBrowserUtils::MakeIconButton(TEXT("FlatButton.Success"), FEditorFontGlyphs::Sign_In, LOCTEXT("JoinButtonTooltip", "Join the selected session"),
-					TAttribute<bool>(this, &SConcertSessionBrowser::IsJoinButtonEnabled),
-					FOnClicked::CreateSP(this, &SConcertSessionBrowser::OnJoinButtonClicked),
-					TAttribute<EVisibility>::Create([this]() { return !IsRestoreButtonEnabled() ? EVisibility::Visible : EVisibility::Hidden; }), // Default button shown if both join/restore are disabled.
-					FMargin(5, 0, 6, 0))
-			]
-
-			// Restore (Share the same slot as Join)
-			+SUniformGridPanel::Slot(1, 0)
-			[
-				ConcertBrowserUtils::MakeIconButton(TEXT("FlatButton.Success"), FEditorFontGlyphs::Arrow_Circle_O_Up, LOCTEXT("RestoreButtonTooltip", "Restore the selected session"),
-					TAttribute<bool>(this, &SConcertSessionBrowser::IsRestoreButtonEnabled),
-					FOnClicked::CreateSP(this, &SConcertSessionBrowser::OnRestoreButtonClicked),
-					TAttribute<EVisibility>::Create([this]() { return IsRestoreButtonEnabled() ? EVisibility::Visible : EVisibility::Hidden; }),
-					FMargin(6, 0, 6, 0))
-			]
+			ConcertBrowserUtils::MakeIconButton(TEXT("FlatButton"), FConcertFrontendStyle::Get()->GetBrush("Concert.RestoreSession"), LOCTEXT("RestoreButtonTooltip", "Restore the selected session"),
+				TAttribute<bool>(this, &SConcertSessionBrowser::IsRestoreButtonEnabled),
+				FOnClicked::CreateSP(this, &SConcertSessionBrowser::OnRestoreButtonClicked),
+				TAttribute<EVisibility>::Create([this]() { return IsRestoreButtonEnabled() ? EVisibility::Visible : EVisibility::Collapsed; }))
+		]
 
-			// Archive.
-			+SUniformGridPanel::Slot(2, 0)
-			[
-				ConcertBrowserUtils::MakeIconButton(TEXT("FlatButton.Warning"), FEditorFontGlyphs::File_Archive_O, LOCTEXT("ArchiveButtonTooltip", "Archive the selected session"),
-					TAttribute<bool>(this, &SConcertSessionBrowser::IsArchiveButtonEnabled),
-					FOnClicked::CreateSP(this, &SConcertSessionBrowser::OnArchiveButtonClicked),
-					EVisibility::Visible,
-					FMargin(6, 0, 6, 0))
-			]
+		// Archive.
+		+SHorizontalBox::Slot()
+		.AutoWidth()
+		.Padding(0, 0, PaddingBetweenButtons, 0)
+		[
+			ConcertBrowserUtils::MakeIconButton(TEXT("FlatButton"), FConcertFrontendStyle::Get()->GetBrush("Concert.ArchiveSession"), LOCTEXT("ArchiveButtonTooltip", "Archive the selected session"),
+				TAttribute<bool>(this, &SConcertSessionBrowser::IsArchiveButtonEnabled),
+				FOnClicked::CreateSP(this, &SConcertSessionBrowser::OnArchiveButtonClicked))
+		]
 
-			// Delete.
-			+SUniformGridPanel::Slot(3, 0)
-			[
-				ConcertBrowserUtils::MakeIconButton(TEXT("FlatButton.Danger"), FEditorFontGlyphs::Trash, LOCTEXT("DeleteButtonTooltip", "Delete the selected session if permitted"),
-					TAttribute<bool>(this, &SConcertSessionBrowser::IsDeleteButtonEnabled),
-					FOnClicked::CreateSP(this, &SConcertSessionBrowser::OnDeleteButtonClicked),
-					EVisibility::Visible,
-					FMargin(5, 0, 6, 0))
-			]
+		// Delete.
+		+SHorizontalBox::Slot()
+		.AutoWidth()
+		[
+			ConcertBrowserUtils::MakeIconButton(TEXT("FlatButton"), FConcertFrontendStyle::Get()->GetBrush("Concert.DeleteSession"), LOCTEXT("DeleteButtonTooltip", "Delete the selected session if permitted"),
+				TAttribute<bool>(this, &SConcertSessionBrowser::IsDeleteButtonEnabled),
+				FOnClicked::CreateSP(this, &SConcertSessionBrowser::OnDeleteButtonClicked))
 		];
 }
 
@@ -2557,8 +2604,8 @@ TSharedRef<ITableRow> SConcertSessionBrowser::MakeActiveSessionRowWidget(const T
 	const FConcertSessionInfo* SessionInfo = Controller->GetActiveSessionInfo(ActiveItem->ServerAdminEndpointId, ActiveItem->SessionId);
 
 	// Add an 'Active Session' row. Clicking the row icon joins the session.
-	return SNew(SSessionRow, ActiveItem, OwnerTable)
-		.OnClickIconFunc([this](const TSharedPtr<FConcertSessionItem>& Item) { RequestJoinSession(Item); })
+	return SNew(SSessionRow, ActiveItem, OwnerTable, Controller->GetConcertClient())
+		.OnDoubleClickFunc([this](const TSharedPtr<FConcertSessionItem>& Item) { RequestJoinSession(Item); })
 		.OnRenameFunc([this](const TSharedPtr<FConcertSessionItem>& Item, const FString& NewName) { RequestRenameSession(Item, NewName); })
 		.HighlightText(this, &SConcertSessionBrowser::HighlightSearchText)
 		.ToolTipText(SessionInfo ? SessionInfo->ToDisplayString() : FText::GetEmpty())
@@ -2570,8 +2617,8 @@ TSharedRef<ITableRow> SConcertSessionBrowser::MakeArchivedSessionRowWidget(const
 	const FConcertSessionInfo* SessionInfo = Controller->GetArchivedSessionInfo(ArchivedItem->ServerAdminEndpointId, ArchivedItem->SessionId);
 
 	// Add an 'Archived Session' row. Clicking the row icon adds a 'Restore as' row to the table.
-	return SNew(SSessionRow, ArchivedItem, OwnerTable)
-		.OnClickIconFunc([this](const TSharedPtr<FConcertSessionItem>& Item) { InsertRestoreSessionAsEditableRow(Item); })
+	return SNew(SSessionRow, ArchivedItem, OwnerTable, Controller->GetConcertClient())
+		.OnDoubleClickFunc([this](const TSharedPtr<FConcertSessionItem>& Item) { InsertRestoreSessionAsEditableRow(Item); })
 		.OnRenameFunc([this](const TSharedPtr<FConcertSessionItem>& Item, const FString& NewName) { RequestRenameSession(Item, NewName); })
 		.HighlightText(this, &SConcertSessionBrowser::HighlightSearchText)
 		.ToolTipText(SessionInfo ? SessionInfo->ToDisplayString() : FText::GetEmpty())
