Fixed ensure when having the source control provider panel open while joining a Multi-User session

diff --git a/Engine/Plugins/Developer/Concert/MultiUserClient/Source/MultiUserClient/Private/MultiUserClientModule.cpp b/Engine/Plugins/Developer/Concert/MultiUserClient/Source/MultiUserClient/Private/MultiUserClientModule.cpp
--- a/Engine/Plugins/Developer/Concert/MultiUserClient/Source/MultiUserClient/Private/MultiUserClientModule.cpp
+++ b/Engine/Plugins/Developer/Concert/MultiUserClient/Source/MultiUserClient/Private/MultiUserClientModule.cpp
@@ -156,7 +156,7 @@ private:
 		{
 			ISourceControlModule& SourceControlModule = ISourceControlModule::Get();
 			ISourceControlProvider& SourceControlProvider = SourceControlModule.GetProvider();
-			if (ensure(SourceControlModule.IsEnabled() && SourceControlProvider.IsAvailable()))
+			if (ensure(SourceControlProvider.IsEnabled() && SourceControlProvider.IsAvailable()))
 			{
 				bool bHasLocalChanges = false;
 				for (const FString& ContentPath : InSharedState->ContentPaths)
