Edigrated 9450999

diff --git a/Engine/Plugins/Developer/Concert/ConcertSync/ConcertSyncClient/Source/ConcertSyncUI/Private/ConcertFrontendStyle.cpp b/Engine/Plugins/Developer/Concert/ConcertSync/ConcertSyncClient/Source/ConcertSyncUI/Private/ConcertFrontendStyle.cpp
--- a/Engine/Plugins/Developer/Concert/ConcertSync/ConcertSyncClient/Source/ConcertSyncUI/Private/ConcertFrontendStyle.cpp
+++ b/Engine/Plugins/Developer/Concert/ConcertSync/ConcertSyncClient/Source/ConcertSyncUI/Private/ConcertFrontendStyle.cpp
@@ -46,9 +46,10 @@ void FConcertFrontendStyle::Initialize()
 	const FLinearColor IconColorAndOpacity(FLinearColor(1.f, 1.f, 1.f, 0.8f));
 
 	// 16x16
-	StyleSet->Set("Concert.Persist", new IMAGE_PLUGIN_BRUSH("Icons/icon_ConcertPersist_16x", Icon16x16));
-	StyleSet->Set("Concert.MyLock", new IMAGE_PLUGIN_BRUSH("Icons/icon_ConcertMyLock_16x", Icon16x16));
-	StyleSet->Set("Concert.OtherLock", new IMAGE_PLUGIN_BRUSH("Icons/icon_ConcertOtherLock_16x", Icon16x16));
+	StyleSet->Set("Concert.Persist",         new IMAGE_PLUGIN_BRUSH("Icons/icon_ConcertPersist_16x", Icon16x16));
+	StyleSet->Set("Concert.LockBackground",  new IMAGE_PLUGIN_BRUSH("Icons/icon_ConcertLockBackground_16x", Icon16x16));
+	StyleSet->Set("Concert.MyLock",          new IMAGE_PLUGIN_BRUSH("Icons/icon_ConcertMyLock_16x", Icon16x16));
+	StyleSet->Set("Concert.OtherLock",       new IMAGE_PLUGIN_BRUSH("Icons/icon_ConcertOtherLock_16x", Icon16x16));
 	StyleSet->Set("Concert.ModifiedByOther", new IMAGE_PLUGIN_BRUSH("Icons/icon_ConcertModifiedByOther_16x", Icon16x16));
 
 	// Multi-user Tab/Menu icons
diff --git a/Engine/Plugins/Developer/Concert/MultiUserClient/Source/MultiUserClient/Private/ConcertWorkspaceUI.cpp b/Engine/Plugins/Developer/Concert/MultiUserClient/Source/MultiUserClient/Private/ConcertWorkspaceUI.cpp
--- a/Engine/Plugins/Developer/Concert/MultiUserClient/Source/MultiUserClient/Private/ConcertWorkspaceUI.cpp
+++ b/Engine/Plugins/Developer/Concert/MultiUserClient/Source/MultiUserClient/Private/ConcertWorkspaceUI.cpp
@@ -77,8 +77,19 @@ public:
 
 		ChildSlot
 		[
-			SNew(SImage)
-			.Image(this, &SConcertWorkspaceLockStateIndicator::GetImageBrush)
+			SNew(SOverlay)
+
+			+SOverlay::Slot() // A colored plane (avatar color) that leaks through the lock image transparent parts (if any).
+			[
+				SNew(SImage)
+				.ColorAndOpacity(this, &SConcertWorkspaceLockStateIndicator::GetBackgroundColor)
+				.Image(LockBackgroundBrush)
+			]
+			+SOverlay::Slot() // The lock image on top.
+			[
+				SNew(SImage)
+				.Image(this, &SConcertWorkspaceLockStateIndicator::GetImageBrush)
+			]
 		];
 	}
 
@@ -87,6 +98,7 @@ public:
 	{
 		if (MyLockBrush == nullptr)
 		{
+			LockBackgroundBrush = FConcertFrontendStyle::Get()->GetBrush(TEXT("Concert.LockBackground"));
 			MyLockBrush = FConcertFrontendStyle::Get()->GetBrush(TEXT("Concert.MyLock"));
 			OtherLockBrush = FConcertFrontendStyle::Get()->GetBrush(TEXT("Concert.OtherLock"));
 		}
@@ -99,6 +111,11 @@ private:
 		return WorkspaceFrontend->GetResourceLockId(AssetPath).IsValid() ? EVisibility::Visible : EVisibility::Collapsed;
 	}
 
+	FSlateColor GetBackgroundColor() const
+	{
+		return WorkspaceFrontend->GetUserAvatarColor(WorkspaceFrontend->GetResourceLockId(AssetPath));
+	}
+
 	const FSlateBrush* GetImageBrush() const
 	{
 		FGuid LockId = WorkspaceFrontend->GetResourceLockId(AssetPath);
@@ -119,6 +136,7 @@ private:
 	// Brushes used to render the lock icon.
 	static const FSlateBrush* MyLockBrush;
 	static const FSlateBrush* OtherLockBrush;
+	static const FSlateBrush* LockBackgroundBrush;
 
 	/** Asset path for this indicator widget.*/
 	FName AssetPath;
@@ -129,6 +147,7 @@ private:
 
 const FSlateBrush* SConcertWorkspaceLockStateIndicator::MyLockBrush = nullptr;
 const FSlateBrush* SConcertWorkspaceLockStateIndicator::OtherLockBrush = nullptr;
+const FSlateBrush* SConcertWorkspaceLockStateIndicator::LockBackgroundBrush = nullptr;
 
 /**
  * Displays a tooltip when moving the mouse over the 'lock' icon displayed on asset locked
@@ -693,8 +712,7 @@ bool FConcertWorkspaceUI::SubmitChangelist(const FConcertPersistCommand& Persist
 
 FText FConcertWorkspaceUI::GetUserDescriptionText(const FGuid& ClientId) const
 {
-	TSharedPtr<IConcertClientWorkspace> ClientWorkspacePin = ClientWorkspace.Pin();
-	if (ClientWorkspacePin)
+	if (TSharedPtr<IConcertClientWorkspace> ClientWorkspacePin = ClientWorkspace.Pin())
 	{
 		FConcertSessionClientInfo ClientSessionInfo;
 		if (ClientWorkspacePin->GetSession().FindSessionClient(ClientId, ClientSessionInfo))
@@ -712,6 +730,25 @@ FText FConcertWorkspaceUI::GetUserDescriptionText(const FConcertClientInfo& Clie
 		FText::Format(LOCTEXT("ConcertUserNameOnDevice", "'{0}' on {1}"), FText::FromString(ClientInfo.DisplayName), FText::FromString(ClientInfo.DeviceName));
 }
 
+FLinearColor FConcertWorkspaceUI::GetUserAvatarColor(const FGuid& ClientId) const
+{
+	if (TSharedPtr<IConcertClientWorkspace> ClientWorkspacePin = ClientWorkspace.Pin())
+	{
+		if (ClientWorkspacePin->GetSession().GetLocalClientInfo().InstanceInfo.InstanceId == ClientId) // This client?
+		{
+			return ClientWorkspacePin->GetSession().GetLocalClientInfo().AvatarColor;
+		}
+
+		FConcertSessionClientInfo ClientSessionInfo;
+		if (ClientWorkspacePin->GetSession().FindSessionClient(ClientId, ClientSessionInfo))
+		{
+			return ClientSessionInfo.ClientInfo.AvatarColor;
+		}
+	}
+
+	return FLinearColor::Transparent;
+}
+
 FGuid FConcertWorkspaceUI::GetWorkspaceLockId() const
 {
 	TSharedPtr<IConcertClientWorkspace> ClientWorkspacePin = ClientWorkspace.Pin();
diff --git a/Engine/Plugins/Developer/Concert/MultiUserClient/Source/MultiUserClient/Private/ConcertWorkspaceUI.h b/Engine/Plugins/Developer/Concert/MultiUserClient/Source/MultiUserClient/Private/ConcertWorkspaceUI.h
--- a/Engine/Plugins/Developer/Concert/MultiUserClient/Source/MultiUserClient/Private/ConcertWorkspaceUI.h
+++ b/Engine/Plugins/Developer/Concert/MultiUserClient/Source/MultiUserClient/Private/ConcertWorkspaceUI.h
@@ -58,6 +58,9 @@ private:
 	FText GetUserDescriptionText(const FGuid& ClientId) const;
 	FText GetUserDescriptionText(const FConcertClientInfo& ClientInfo) const;
 
+	/** Returns the avatar color corresponding to the specified client Id (including this client). */
+	FLinearColor GetUserAvatarColor(const FGuid& ClientId) const;
+
 	/** Get the local workspace's lock id. */
 	FGuid GetWorkspaceLockId() const;
 
