Edigrated 9493133

  - Updated the Multi-User active to display a "N/A (-game)" when a connected client is in -game with a tooltip saying that clients in -game do not emits their presence/level.

diff --git a/Engine/Plugins/Developer/Concert/ConcertSync/ConcertSyncClient/Source/ConcertSyncClient/Private/ConcertClientPresenceManager.cpp b/Engine/Plugins/Developer/Concert/ConcertSync/ConcertSyncClient/Source/ConcertSyncClient/Private/ConcertClientPresenceManager.cpp
--- a/Engine/Plugins/Developer/Concert/ConcertSync/ConcertSyncClient/Source/ConcertSyncClient/Private/ConcertClientPresenceManager.cpp
+++ b/Engine/Plugins/Developer/Concert/ConcertSync/ConcertSyncClient/Source/ConcertSyncClient/Private/ConcertClientPresenceManager.cpp
@@ -556,8 +556,11 @@ void FConcertClientPresenceManager::OnSessionClientChanged(IConcertClientSession
 			SetPresenceVisibility(InClientInfo.ClientEndpointId, PresencePersistentState->bVisible, PresencePersistentState->bPropagateToAll);
 		}
 
-		// Send avatar-related info for this client when a remote client connects or is updated
-		SendPresenceInVREvent(&InClientInfo.ClientEndpointId);
+		if (!IsInGame())
+		{
+			// Send avatar-related info for this client when a remote client connects or is updated
+			SendPresenceInVREvent(&InClientInfo.ClientEndpointId);
+		}
 	}
 	else if (InClientStatus == EConcertClientStatus::Disconnected)
 	{
@@ -700,6 +703,7 @@ FConcertClientPresenceState& FConcertClientPresenceManager::EnsurePresenceState(
 			if (Session->FindSessionClient(InEndpointId, ClientSessionInfo))
 			{
 				PresenceState->DisplayName = ClientSessionInfo.ClientInfo.DisplayName;
+				PresenceState->AvatarColor = ClientSessionInfo.ClientInfo.AvatarColor;
 			}
 		}
 		PresencePersistentStateMap.FindOrAdd(PresenceState->DisplayName);
diff --git a/Engine/Plugins/Developer/Concert/MultiUserClient/Source/MultiUserClient/Private/Widgets/SActiveSession.cpp b/Engine/Plugins/Developer/Concert/MultiUserClient/Source/MultiUserClient/Private/Widgets/SActiveSession.cpp
--- a/Engine/Plugins/Developer/Concert/MultiUserClient/Source/MultiUserClient/Private/Widgets/SActiveSession.cpp
+++ b/Engine/Plugins/Developer/Concert/MultiUserClient/Source/MultiUserClient/Private/Widgets/SActiveSession.cpp
@@ -110,10 +110,22 @@ public:
 			// If the row is not for "this client", then add a hyper link to open the other client level.
 			if (ClientInfoPin.IsValid() && ClientSessionPin.IsValid() && ClientInfoPin->ClientEndpointId != ClientSessionPin->GetSessionClientEndpointId())
 			{
-				LevelWidget = SNew(SHyperlink)
-					.Text(this, &SActiveSessionDetailsRow::GetLevel)
-					.ToolTipText(FText::Format(LOCTEXT("OtherUserLevel", "The level opened in {0} editor. Click to open this level."), FText::FromString(ClientInfoPin->ClientInfo.DisplayName)))
-					.OnNavigate(this, &SActiveSessionDetailsRow::OnLevelClicked);
+				LevelWidget = SNew(SOverlay)
+					+SOverlay::Slot()
+					[
+						SNew(SHyperlink)
+						.Visibility_Lambda([this]() { return GetOtherLevelVisibility(); })
+						.Text(this, &SActiveSessionDetailsRow::GetLevel)
+						.ToolTipText(this, &SActiveSessionDetailsRow::GetOtherLevelTooltip)
+						.OnNavigate(this, &SActiveSessionDetailsRow::OnOtherLevelClicked)
+					]
+					+SOverlay::Slot()
+					[
+						SNew(STextBlock)
+						.Visibility_Lambda([this]() { return GetOtherLevelVisibility() == EVisibility::Visible ? EVisibility::Collapsed : EVisibility::Visible; })
+						.Text(LOCTEXT("LevelNotAvailableInGame", "N/A (-game)")) // The only known reason for the level to be unknown and invisible is b/c the user is in -game.
+						.ToolTipText(LOCTEXT("LevelNotAvailableInGame_Tooltip", "The presence/level information is not emitted by clients running Editor in -game mode."))
+					];
 			}
 			else
 			{
@@ -169,6 +181,36 @@ public:
 		return FText();
 	}
 
+	EVisibility GetOtherLevelVisibility() const
+	{
+		TSharedPtr<FConcertSessionClientInfo> ClientInfoPin = SessionClientInfo.Pin();
+		TSharedPtr<IConcertSyncClient> SyncClientPin = SyncClient.Pin();
+
+		if (ClientInfoPin && SyncClientPin)
+		{
+			EEditorPlayMode EditorPlayMode = EEditorPlayMode::None;
+			if (IConcertClientPresenceManager* PresenceManager = SyncClientPin->GetPresenceManager())
+			{
+				if (!PresenceManager->GetPresenceWorldPath(ClientInfoPin->ClientEndpointId, EditorPlayMode).IsEmpty())
+				{
+					return EVisibility::Visible;
+				}
+			}
+		}
+
+		return EVisibility::Hidden;
+	}
+
+	FText GetOtherLevelTooltip() const
+	{
+		if (TSharedPtr<FConcertSessionClientInfo> ClientInfoPin = SessionClientInfo.Pin())
+		{
+			return FText::Format(LOCTEXT("OtherUserLevel_Tooltip", "The level opened in {0} editor. Click to open this level."), FText::FromString(ClientInfoPin->ClientInfo.DisplayName));
+		}
+
+		return FText::GetEmpty();
+	}
+
 	FText GetLevel() const
 	{
 		FString WorldPath;
@@ -192,18 +234,18 @@ public:
 
 			if (EditorPlayMode == EEditorPlayMode::PIE)
 			{
-				WorldPath += " (PIE)";
+				WorldPath += TEXT(" (PIE)");
 			}
 			else if (EditorPlayMode == EEditorPlayMode::SIE)
 			{
-				WorldPath += " (SIE)";
+				WorldPath += TEXT(" (SIE)");
 			}
 		}
 
 		return FText::FromString(WorldPath);
 	}
 
-	void OnLevelClicked()
+	void OnOtherLevelClicked()
 	{
 		TSharedPtr<FConcertSessionClientInfo> ClientInfoPin = SessionClientInfo.Pin();
 		TSharedPtr<IConcertClientSession> ClientSessionPin = ClientSession.Pin();
