- Fix typo in multi user category

diff --git a/Engine/Plugins/Developer/Concert/MultiUserClient/Source/MultiUserClientLibrary/Public/MultiUserClientStatics.h b/Engine/Plugins/Developer/Concert/MultiUserClient/Source/MultiUserClientLibrary/Public/MultiUserClientStatics.h
--- a/Engine/Plugins/Developer/Concert/MultiUserClient/Source/MultiUserClientLibrary/Public/MultiUserClientStatics.h
+++ b/Engine/Plugins/Developer/Concert/MultiUserClient/Source/MultiUserClientLibrary/Public/MultiUserClientStatics.h
@@ -40,47 +40,47 @@ class MULTIUSERCLIENTLIBRARY_API UMultiUserClientStatics : public UBlueprintFunc
 public:
 
 	/** Set whether presence is currently enabled and should be shown (unless hidden by other settings) */
-	UFUNCTION(BlueprintCallable, Category = "Mutli-User Presence", meta=(DevelopmentOnly, DisplayName = "Set Multi-User Presence Enabled"))
+	UFUNCTION(BlueprintCallable, Category = "Multi-User Presence", meta=(DevelopmentOnly, DisplayName = "Set Multi-User Presence Enabled"))
 	static void SetMultiUserPresenceEnabled(const bool IsEnabled = true);
 
 	/** Set Presence Actor Visibility by display name */
-	UFUNCTION(BlueprintCallable, Category = "Mutli-User Presence", meta=(DevelopmentOnly, DisplayName = "Set Multi-User Presence Visibility"))
+	UFUNCTION(BlueprintCallable, Category = "Multi-User Presence", meta=(DevelopmentOnly, DisplayName = "Set Multi-User Presence Visibility"))
 	static void SetMultiUserPresenceVisibility(const FString& Name, bool Visibility, bool PropagateToAll = false);
 
 	/** Set Presence Actor Visibility by client id */
-	UFUNCTION(BlueprintCallable, Category = "Mutli-User Presence", meta = (DevelopmentOnly, DisplayName = "Set Multi-User Presence Visibility By Id"))
+	UFUNCTION(BlueprintCallable, Category = "Multi-User Presence", meta = (DevelopmentOnly, DisplayName = "Set Multi-User Presence Visibility By Id"))
 	static void SetMultiUserPresenceVisibilityById(const FGuid& ClientEndpointId, bool Visibility, bool PropagateToAll = false);
 
 	/** Get the Presence Actor transform for the specified client endpoint id or identity if the client isn't found */
-	UFUNCTION(BlueprintCallable, Category = "Mutli-User Presence", meta = (DevelopmentOnly, DisplayName = "Get Multi-User Presence Transform"))
+	UFUNCTION(BlueprintCallable, Category = "Multi-User Presence", meta = (DevelopmentOnly, DisplayName = "Get Multi-User Presence Transform"))
 	static FTransform GetMultiUserPresenceTransform(const FGuid& ClientEndpointId);
 
-	/** Teleport to another Mutli-User user's presence. */
-	UFUNCTION(BlueprintCallable, Category = "Mutli-User Client", meta = (DevelopmentOnly, DisplayName = "Jump to Multi-User Presence"))
+	/** Teleport to another Multi-User user's presence. */
+	UFUNCTION(BlueprintCallable, Category = "Multi-User Client", meta = (DevelopmentOnly, DisplayName = "Jump to Multi-User Presence"))
 	static void JumpToMultiUserPresence(const FString& OtherUserName, FTransform TransformOffset);
 
-	/** Update mutli-User Workspace Modified Packages to be in sync for source control submission. */
-	UFUNCTION(BlueprintCallable, Category = "Mutli-User Source Control", meta=(DevelopmentOnly, DeprecatedFunction, DeprecationMessage = "UpdateWorkspaceModifiedPackages is deprecated. Please use PersistMultiUserSessionChanges instead."))
+	/** Update Multi-User Workspace Modified Packages to be in sync for source control submission. */
+	UFUNCTION(BlueprintCallable, Category = "Multi-User Source Control", meta=(DevelopmentOnly, DeprecatedFunction, DeprecationMessage = "UpdateWorkspaceModifiedPackages is deprecated. Please use PersistMultiUserSessionChanges instead."))
 	static void UpdateWorkspaceModifiedPackages();
 
 	/** Persist the session changes and prepare the files for source control submission. */
-	UFUNCTION(BlueprintCallable, Category = "Mutli-User Source Control", meta=(DevelopmentOnly, DisplayName = "Persist Multi-User Session Changes"))
+	UFUNCTION(BlueprintCallable, Category = "Multi-User Source Control", meta=(DevelopmentOnly, DisplayName = "Persist Multi-User Session Changes"))
 	static void PersistMultiUserSessionChanges();
 
 	/** Get the local ClientInfo. Works when not connected to a session. */
-	UFUNCTION(BlueprintCallable, Category = "Mutli-User Client", meta=(DevelopmentOnly, DisplayName = "Get Local Multi-User Client Info"))
+	UFUNCTION(BlueprintCallable, Category = "Multi-User Client", meta=(DevelopmentOnly, DisplayName = "Get Local Multi-User Client Info"))
 	static FMultiUserClientInfo GetLocalMultiUserClientInfo();
 
-	/** Get the ClientInfo for any mutli-User participant by name. The local user is found even when not connected to a session. Returns false is no client was found. */
-	UFUNCTION(BlueprintCallable, Category = "Mutli-User Client", meta=(DevelopmentOnly, DisplayName = "Get Multi-User Client Info by Name"))
+	/** Get the ClientInfo for any Multi-User participant by name. The local user is found even when not connected to a session. Returns false is no client was found. */
+	UFUNCTION(BlueprintCallable, Category = "Multi-User Client", meta=(DevelopmentOnly, DisplayName = "Get Multi-User Client Info by Name"))
 	static bool GetMultiUserClientInfoByName(const FString& ClientName, FMultiUserClientInfo& ClientInfo);
 
-	/** Get ClientInfos of current mutli-User participants except for the local user. Returns false is no remote clients were found. */
-	UFUNCTION(BlueprintCallable, Category = "Mutli-User Client", meta=(DevelopmentOnly, DisplayName = "Get Remote Multi-User Client Infos"))
+	/** Get ClientInfos of current Multi-User participants except for the local user. Returns false is no remote clients were found. */
+	UFUNCTION(BlueprintCallable, Category = "Multi-User Client", meta=(DevelopmentOnly, DisplayName = "Get Remote Multi-User Client Infos"))
 	static bool GetRemoteMultiUserClientInfos(TArray<FMultiUserClientInfo>& ClientInfos);
 
-	/** Get mutli-User connection status. */
-	UFUNCTION(BlueprintCallable, Category = "Mutli-User Client", meta=(DevelopmentOnly, DisplayName = "Get Multi-User Connection Status"))
+	/** Get Multi-User connection status. */
+	UFUNCTION(BlueprintCallable, Category = "Multi-User Client", meta=(DevelopmentOnly, DisplayName = "Get Multi-User Connection Status"))
 	static bool GetMultiUserConnectionStatus();
 
 };
