Temporarily disable dirty package check when joining a session in Multi-User using a cvar until we can re-work the feature

diff --git a/Engine/Plugins/Developer/Concert/MultiUserClient/Source/MultiUserClient/Private/MultiUserClientModule.cpp b/Engine/Plugins/Developer/Concert/MultiUserClient/Source/MultiUserClient/Private/MultiUserClientModule.cpp
--- a/Engine/Plugins/Developer/Concert/MultiUserClient/Source/MultiUserClient/Private/MultiUserClientModule.cpp
+++ b/Engine/Plugins/Developer/Concert/MultiUserClient/Source/MultiUserClient/Private/MultiUserClientModule.cpp
@@ -89,6 +89,9 @@ void KillProcess(uint32 ProcessID)
 
 }
 
+// Disable dirty package check 
+static TAutoConsoleVariable<int32> CVarDisableDirtyPackageCheck(TEXT("Concert.DisableDirtyPackageCheck"), 1, TEXT("Disable dirty package check when joining a session."));
+
 /**
  * Connection task used to validate that the workspace has no local changes (according to source control) or in-memory changes (dirty packages).
  */
@@ -176,6 +179,13 @@ private:
 	/** Common function to query for in-memory changes to packages */
 	static void HandleDirtyPackages(TSharedRef<FSharedAsyncState, ESPMode::ThreadSafe> InSharedState)
 	{
+		//@todo FH: Temporary cvar to disable the dirty package check until it can be reworked as a skippable warning.
+		if (CVarDisableDirtyPackageCheck.GetValueOnAnyThread() > 0)
+		{
+			InSharedState->Result = EConcertResponseCode::Success;
+			return;
+		}
+		
 		TArray<UPackage*> DirtyPackages;
 #if WITH_EDITOR
 		{
@@ -285,7 +295,6 @@ private:
 	TSharedPtr<FSharedAsyncState, ESPMode::ThreadSafe> SharedState;
 };
 
-
 /**
  * Customize the multi-user settings to add validation and visual feedback when a user enter something invalid.
  */
