Fix crash when trying to persist a session with transaction data on assets where asset creation isn't part of the session and the asset doesn't exist locally

diff --git a/Engine/Plugins/Developer/Concert/MultiUserClient/Source/MultiUserClient/Private/ConcertWorkspaceUI.cpp b/Engine/Plugins/Developer/Concert/MultiUserClient/Source/MultiUserClient/Private/ConcertWorkspaceUI.cpp
--- a/Engine/Plugins/Developer/Concert/MultiUserClient/Source/MultiUserClient/Private/ConcertWorkspaceUI.cpp
+++ b/Engine/Plugins/Developer/Concert/MultiUserClient/Source/MultiUserClient/Private/ConcertWorkspaceUI.cpp
@@ -599,14 +599,21 @@ bool FConcertWorkspaceUI::PromptPersistSessionChanges()
 		TArray<FSourceControlStateRef> States;
 		PackageNames = ClientWorkspacePin->GatherSessionChanges();
 		FString Filename;
-		for (const FName& PackageName : PackageNames)
+		for (auto It = PackageNames.CreateIterator(); It; ++It)
 		{
-			if (FPackageName::DoesPackageExist(PackageName.ToString(), nullptr, &Filename))
+			if (FPackageName::DoesPackageExist(It->ToString(), nullptr, &Filename))
 			{
 				PackageFilenames.Add(FPaths::ConvertRelativePathToFull(MoveTemp(Filename)));
 			}
+			// if the package file does not exist locally, remove it from the persist list, the db contains transaction data on file not propagated through Multi-User.
+			else
+			{
+				It.RemoveCurrent();
+			}
 		}
-		ISourceControlModule::Get().GetProvider().GetState(PackageFilenames, States, EStateCacheUsage::ForceUpdate);
+		ECommandResult::Type Result = ISourceControlModule::Get().GetProvider().GetState(PackageFilenames, States, EStateCacheUsage::ForceUpdate);
+		// The dummy Multi-User source control provider always succeed and always return proxy states.
+		ensure(Result == ECommandResult::Succeeded);
 
 		// Create the list of persist items from the package names and source control state
 		for (int32 Index = 0; Index < PackageNames.Num(); ++Index)
