added a settings for different connection validation mode allowing for prompt on connection instead of hard failure
Source Control validation now check if checked out files are actually modified instead of just being checked out, if validation mode is soft, Multi-User connection are allowed to proceed
Prompt on connection when dirty packages are found on soft validation.
Hot reload dirty packages on connection

diff --git a/Engine/Plugins/Developer/Concert/ConcertMain/Source/Concert/Private/ConcertClient.cpp b/Engine/Plugins/Developer/Concert/ConcertMain/Source/Concert/Private/ConcertClient.cpp
--- a/Engine/Plugins/Developer/Concert/ConcertMain/Source/Concert/Private/ConcertClient.cpp
+++ b/Engine/Plugins/Developer/Concert/ConcertMain/Source/Concert/Private/ConcertClient.cpp
@@ -110,7 +110,7 @@ private:
 		}
 
 		// Should cancel request?
-		if (AutoConnectionNotification.IsValid() && AutoConnectionNotification->ShouldCancel())
+		if (AutoConnectionNotification.IsValid() && AutoConnectionNotification->GetPromptAction() == EAsyncTaskNotificationPromptAction::Cancel)
 		{
 			SetAsyncNotificationComplete(GetAutoConnectionCanceledMessage(), false);
 			Client->StopAutoConnect(); // Indirect self-destruct.
@@ -312,7 +312,7 @@ public:
 			ConnectionTasks[0]->Abort();
 
 			// Clear the tasks, set the notification text and fulfill the 'Execute()' promise.
-			SetResult(EConcertResponseCode::Failed, /*bWasCanceled*/true, GetCanceledErrorMessage());
+			SetResult(EConcertResponseCode::Failed, GetCanceledErrorMessage());
 		}
 	}
 
@@ -360,15 +360,34 @@ private:
 
 	void Tick()
 	{
+		auto GetTaskAction = [](EAsyncTaskNotificationPromptAction InPromptAction) ->EConcertConnectionTaskAction
+		{
+			switch (InPromptAction)
+			{
+			case EAsyncTaskNotificationPromptAction::None:
+				return EConcertConnectionTaskAction::None;
+			case EAsyncTaskNotificationPromptAction::Cancel:
+				return EConcertConnectionTaskAction::Cancel;
+			case EAsyncTaskNotificationPromptAction::Continue:
+				return EConcertConnectionTaskAction::Continue;
+			// unattended case resolve as a continue
+			default:
+				return EConcertConnectionTaskAction::Continue;
+			}
+		};
+
 		// We should only Tick while we have tasks to process
 		check(ConnectionTasks.Num() > 0);
 
-		const bool bCanceled = Notification->ShouldCancel();
+		EAsyncTaskNotificationPromptAction PromptAction = Notification->GetPromptAction();
+		EConcertConnectionTaskAction TaskAction = GetTaskAction(PromptAction);
+		const bool bCanceled = TaskAction == EConcertConnectionTaskAction::Cancel;
+		EConcertResponseCode TaskStatus = ConnectionTasks[0]->GetStatus();
 		if (bCanceled)
 		{
-			if (ConnectionTasks[0]->GetStatus() == EConcertResponseCode::Pending)
+			if (TaskStatus == EConcertResponseCode::Pending)
 			{
-				ConnectionTasks[0]->Tick(bCanceled); // Give it a last tick to give it a chance to cancel cleanly.
+				ConnectionTasks[0]->Tick(TaskAction); // Give it a last tick to give it a chance to cancel cleanly.
 			}
 
 			SetResultAndDelete(EConcertResponseCode::Failed, bCanceled, GetCanceledErrorMessage()); // Cancellation has priority over other possible errors (it could hide some failure)
@@ -376,19 +395,29 @@ private:
 		}
 
 		// Update the current task
-		switch (ConnectionTasks[0]->GetStatus())
+		switch (TaskStatus)
 		{
 			// Pending state - update the task
 		case EConcertResponseCode::Pending:
-			ConnectionTasks[0]->Tick(bCanceled);
+			ConnectionTasks[0]->Tick(TaskAction);
 			return;
-
+			// Prompt state - wait for user action to either proceed (success) or stop (fail)
+		case EConcertResponseCode::InvalidRequest:
+			{
+				FAsyncNotificationStateData StateData(Config.PendingTitleText.Get(FText::GetEmpty()), ConnectionTasks[0]->GetError(), EAsyncTaskNotificationState::Prompt);
+				StateData.PromptText = ConnectionTasks[0]->GetPrompt();
+				StateData.HyperlinkText = LOCTEXT("PendingConnectionFailureDetails", "See Details...");
+				StateData.Hyperlink = ConnectionTasks[0]->GetErrorDelegate();
+				Notification->SetNotificationState(StateData);
+				ConnectionTasks[0]->Tick(TaskAction);
+				return;
+			}
 			// Success state - move on to the next task
 		case EConcertResponseCode::Success:
 			ConnectionTasks.RemoveAt(0, 1, /*bAllowShrinking*/false);
 			if (ConnectionTasks.Num() > 0)
 			{
-				Notification->SetProgressText(ConnectionTasks[0]->GetDescription());
+				Notification->SetNotificationState(FAsyncNotificationStateData(Config.PendingTitleText.Get(FText::GetEmpty()), ConnectionTasks[0]->GetDescription(), EAsyncTaskNotificationState::Pending));
 				ConnectionTasks[0]->Execute();
 			}
 			else
@@ -397,16 +426,15 @@ private:
 				SetResultAndDelete(EConcertResponseCode::Success, bCanceled); // do not use 'this' after this call!
 			}
 			return;
-
 			// Error state - fail the connection
 		default:
-			SetResultAndDelete(ConnectionTasks[0]->GetStatus(), bCanceled, ConnectionTasks[0]->GetError()); // do not use 'this' after this call!
+			SetResultAndDelete(TaskStatus, bCanceled, ConnectionTasks[0]->GetError(), ConnectionTasks[0]->GetErrorDelegate()); // do not use 'this' after this call!
 			return;
 		}
 	}
 
 	/** Set the result */
-	void SetResult(const EConcertResponseCode InResult, bool bWasCanceled, const FText InFailureReason = FText())
+	void SetResult(const EConcertResponseCode InResult, const FText InFailureReason = FText(), const FSimpleDelegate& InErrorDelegate = FSimpleDelegate())
 	{
 		if (InResult == EConcertResponseCode::Success)
 		{
@@ -414,23 +442,12 @@ private:
 		}
 		else
 		{
-			// Disconnect the session -> Prevent the user from successfully connecting if he was just about to connect (in case the 'join session' event was in-flight but the connection was canceled/aborted).
-			if (bWasCanceled)
+			Client->InternalDisconnectSession();
+			if (InResult == EConcertResponseCode::Failed)
 			{
-				// Don't keep the notification open if canceled
-				Notification->SetKeepOpenOnFailure(false);
-				Client->DisconnectSession(); // This also stops auto-connect (it the client was auto-connecting, this will also prevent auto-connect to retry)
+				Notification->SetKeepOpenOnFailure(Config.KeepNotificationOpenOnError);
 			}
-			else
-			{
-				Client->InternalDisconnectSession(); // This doesn't stop auto-connect, so if it was auto-connecting, it will retry on error.
-
-				if (InResult == EConcertResponseCode::Failed)
-				{
-					Notification->SetKeepOpenOnFailure(Config.KeepNotificationOpenOnError);
-				}
-			}
-
+			Notification->SetHyperlink(InErrorDelegate, LOCTEXT("PendingConnectionFailureDetails", "See Details..."));
 			Notification->SetComplete(Config.FailureTitleText.Get(FText::GetEmpty()), InFailureReason, /*bSuccess*/false);
 		}
 
@@ -439,11 +456,16 @@ private:
 	}
 
 	/** Set the result and delete ourself - 'this' will be garbage after calling this function! */
-	void SetResultAndDelete(const EConcertResponseCode InResult, bool bWasCanceled, const FText InFailureReason = FText())
+	void SetResultAndDelete(const EConcertResponseCode InResult, bool bWasCanceled, const FText InFailureReason = FText(), const FSimpleDelegate& InErrorDelegate = FSimpleDelegate())
 	{
 		// Set the result and delete ourself
-		SetResult(InResult, bWasCanceled, InFailureReason);
+		SetResult(InResult, InFailureReason, InErrorDelegate);
 		check(Client->PendingConnection.Get() == this);
+		// if the connection was canceled, also cancel the auto connection, so it won't retry on failure		
+		if (bWasCanceled)
+		{
+			Client->AutoConnection.Reset();
+		}
 		Client->PendingConnection.Reset();
 	}
 
@@ -472,7 +494,7 @@ public:
 		Result.Reset();
 	}
 
-	virtual void Tick(const bool bShouldCancel) override
+	virtual void Tick(EConcertConnectionTaskAction TaskAction) override
 	{
 	}
 
@@ -500,11 +522,22 @@ public:
 		return EConcertResponseCode::Failed;
 	}
 
+	virtual FText GetPrompt() const override
+	{
+		// client connection task have no prompt
+		return FText::GetEmpty();
+	}
+
 	virtual FText GetError() const override
 	{
 		return Result.IsValid() ? ErrorText : LOCTEXT("RemoteConnectionAttemptAborted", "Remote Connection Attempt Aborted.");
 	}
 
+	virtual FSimpleDelegate GetErrorDelegate() const override
+	{
+		return FSimpleDelegate();
+	}
+
 	virtual FText GetDescription() const override
 	{
 		return LOCTEXT("AttemptingRemoteConnection", "Attempting Remote Connection...");
diff --git a/Engine/Plugins/Developer/Concert/ConcertMain/Source/Concert/Public/ConcertSettings.h b/Engine/Plugins/Developer/Concert/ConcertMain/Source/Concert/Public/ConcertSettings.h
--- a/Engine/Plugins/Developer/Concert/ConcertMain/Source/Concert/Public/ConcertSettings.h
+++ b/Engine/Plugins/Developer/Concert/ConcertMain/Source/Concert/Public/ConcertSettings.h
@@ -233,6 +233,34 @@ struct FConcertClientSettings
 	FString ClientAuthenticationKey;
 };
 
+UENUM()
+enum class EConcertSourceValidationMode
+{
+	/** Source control validation will fail on any changes when connecting to a Multi-User Session. */
+	Hard,
+	/** 
+	 * Source control validation will warn and prompt on any changes when connecting to a Multi-User session. 
+	 * In Memory changes will be hot-reloaded.
+	 * Source control changes aren't affected but will be stashed/shelved in the future.
+	 */
+	Soft,
+	/** Soft validation mode with auto proceed on prompts. */
+	SoftAutoProceed
+};
+
+USTRUCT()
+struct FConcertSourceControlSettings
+{
+	GENERATED_BODY()
+
+	FConcertSourceControlSettings()
+		: ValidationMode(EConcertSourceValidationMode::Hard)
+	{}
+
+	UPROPERTY(config, EditAnywhere, Category="Source Control Settings")
+	EConcertSourceValidationMode ValidationMode;
+};
+
 UCLASS(config=Engine)
 class CONCERT_API UConcertClientConfig : public UObject
 {
@@ -244,7 +272,7 @@ public:
 	 * Mark this setting object as editor only.
 	 * This so soft object path reference made by this setting object won't be automatically grabbed by the cooker.
 	 * @see UPackage::Save, FSoftObjectPathThreadContext::GetSerializationOptions, FSoftObjectPath::ImportTextItem
-	 * @todo: cooker should have a better way to filter editor only objects for 'unsollicited' references.
+	 * @todo: cooker should have a better way to filter editor only objects for 'unsolicited' references.
 	 */
 	virtual bool IsEditorOnly() const
 	{
@@ -313,6 +341,9 @@ public:
 	UPROPERTY(config, EditAnywhere, Category="Client Settings", meta=(ShowOnlyInnerProperties))
 	FConcertClientSettings ClientSettings;
 
+	UPROPERTY(config, EditAnywhere, Category = "Source Control Settings", meta=(ShowOnlyInnerProperties))
+	FConcertSourceControlSettings SourceControlSettings;
+
 	/** Endpoint settings passed down to endpoints on creation */
 	UPROPERTY(config, EditAnywhere, AdvancedDisplay, Category="Endpoint Settings", meta=(ShowOnlyInnerProperties))
 	FConcertEndpointSettings EndpointSettings;
diff --git a/Engine/Plugins/Developer/Concert/ConcertMain/Source/Concert/Public/IConcertClient.h b/Engine/Plugins/Developer/Concert/ConcertMain/Source/Concert/Public/IConcertClient.h
--- a/Engine/Plugins/Developer/Concert/ConcertMain/Source/Concert/Public/IConcertClient.h
+++ b/Engine/Plugins/Developer/Concert/ConcertMain/Source/Concert/Public/IConcertClient.h
@@ -16,6 +16,16 @@ class IConcertClientConnectionTask;
 DECLARE_MULTICAST_DELEGATE_OneParam(FOnConcertClientSessionStartupOrShutdown, TSharedRef<IConcertClientSession>);
 DECLARE_MULTICAST_DELEGATE_TwoParams(FOnConcertClientSessionGetPreConnectionTasks, const IConcertClient&, TArray<TUniquePtr<IConcertClientConnectionTask>>&);
 
+/**
+ * Enum to indicate if any action were taken for the connection task
+ */
+enum class EConcertConnectionTaskAction : uint8
+{
+	None = 0,
+	Cancel,
+	Continue
+};
+
 /** Interface for tasks executed during the Concert client connection flow (eg, validation, creation, connection) */
 class IConcertClientConnectionTask
 {
@@ -37,7 +47,7 @@ public:
 	/**
 	 * Tick this task, optionally requesting that it should gracefully cancel.
 	 */
-	virtual void Tick(const bool bShouldCancel) = 0;
+	virtual void Tick(EConcertConnectionTaskAction TaskAction) = 0;
 
 	/**
 	 * Get whether this task can be gracefully cancelled.
@@ -50,11 +60,21 @@ public:
 	 */
 	virtual EConcertResponseCode GetStatus() const = 0;
 
+	/**
+	 * Get the prompt message of this task to be displayed on prompt button if the task require action.
+	 */
+	virtual FText GetPrompt() const = 0;
+
 	/**
 	 * Get the extended error status of this task that can be used in the error notification (if any).
 	 */
 	virtual FText GetError() const = 0;
 
+	/**
+	 * Get the delegate to gather more error details for this task error notification (if any).
+	 */
+	virtual FSimpleDelegate GetErrorDelegate() const = 0;
+
 	/**
 	 * Get a description of this task that can be used in the progress notification (if any).
 	 */
diff --git a/Engine/Plugins/Developer/Concert/ConcertSync/ConcertSyncClient/Source/ConcertSyncClient/Private/ConcertClientPackageManager.cpp b/Engine/Plugins/Developer/Concert/ConcertSync/ConcertSyncClient/Source/ConcertSyncClient/Private/ConcertClientPackageManager.cpp
--- a/Engine/Plugins/Developer/Concert/ConcertSync/ConcertSyncClient/Source/ConcertSyncClient/Private/ConcertClientPackageManager.cpp
+++ b/Engine/Plugins/Developer/Concert/ConcertSync/ConcertSyncClient/Source/ConcertSyncClient/Private/ConcertClientPackageManager.cpp
@@ -22,6 +22,7 @@
 #if WITH_EDITOR
 	#include "Editor.h"
 	#include "Editor/EditorEngine.h"
+	#include "FileHelpers.h"
 #endif
 
 #define LOCTEXT_NAMESPACE "ConcertClientPackageManager"
@@ -162,6 +163,23 @@ void FConcertClientPackageManager::SynchronizePersistedFiles(const TMap<FString,
 #endif
 }
 
+void FConcertClientPackageManager::QueueDirtyPackagesForReload()
+{
+	TArray<UPackage*> DirtyPkgs;
+#if WITH_EDITOR
+	{
+		UEditorLoadingAndSavingUtils::GetDirtyMapPackages(DirtyPkgs);
+		UEditorLoadingAndSavingUtils::GetDirtyContentPackages(DirtyPkgs);
+	}
+#endif
+	for (UPackage* DirtyPkg : DirtyPkgs)
+	{
+		FName PackageName = DirtyPkg->GetFName();
+		PackagesPendingHotReload.Add(PackageName);
+		PackagesPendingPurge.Remove(PackageName);
+	}
+}
+
 void FConcertClientPackageManager::SynchronizeInMemoryPackages()
 {
 	// Purge pending packages first, since hot reloading can prompt on them before we clear their dirty flags
diff --git a/Engine/Plugins/Developer/Concert/ConcertSync/ConcertSyncClient/Source/ConcertSyncClient/Private/ConcertClientPackageManager.h b/Engine/Plugins/Developer/Concert/ConcertSync/ConcertSyncClient/Source/ConcertSyncClient/Private/ConcertClientPackageManager.h
--- a/Engine/Plugins/Developer/Concert/ConcertSync/ConcertSyncClient/Source/ConcertSyncClient/Private/ConcertClientPackageManager.h
+++ b/Engine/Plugins/Developer/Concert/ConcertSync/ConcertSyncClient/Source/ConcertSyncClient/Private/ConcertClientPackageManager.h
@@ -40,6 +40,11 @@ public:
 	 */
 	void SynchronizePersistedFiles(const TMap<FString, int64>& PersistedFiles);
 
+	/**
+	 * Discard dirty packages changes by queuing them for hot reload
+	 */
+	void QueueDirtyPackagesForReload();
+
 	/**
 	 * Synchronize any pending updates to in-memory packages (hot-reloads or purges) to keep them up-to-date with the on-disk state.
 	 */
diff --git a/Engine/Plugins/Developer/Concert/ConcertSync/ConcertSyncClient/Source/ConcertSyncClient/Private/ConcertClientWorkspace.cpp b/Engine/Plugins/Developer/Concert/ConcertSync/ConcertSyncClient/Source/ConcertSyncClient/Private/ConcertClientWorkspace.cpp
--- a/Engine/Plugins/Developer/Concert/ConcertSync/ConcertSyncClient/Source/ConcertSyncClient/Private/ConcertClientWorkspace.cpp
+++ b/Engine/Plugins/Developer/Concert/ConcertSync/ConcertSyncClient/Source/ConcertSyncClient/Private/ConcertClientWorkspace.cpp
@@ -750,6 +750,7 @@ void FConcertClientWorkspace::OnEndFrame()
 		if (PackageManager)
 		{
 			PackageManager->SynchronizePersistedFiles(SessionData.PersistedFiles);
+			PackageManager->QueueDirtyPackagesForReload();
 			PackageManager->ApplyAllHeadPackageData();
 			PackageManager->SynchronizeInMemoryPackages();
 		}
diff --git a/Engine/Plugins/Developer/Concert/MultiUserClient/Source/MultiUserClient/Private/MultiUserClientModule.cpp b/Engine/Plugins/Developer/Concert/MultiUserClient/Source/MultiUserClient/Private/MultiUserClientModule.cpp
--- a/Engine/Plugins/Developer/Concert/MultiUserClient/Source/MultiUserClient/Private/MultiUserClientModule.cpp
+++ b/Engine/Plugins/Developer/Concert/MultiUserClient/Source/MultiUserClient/Private/MultiUserClientModule.cpp
@@ -91,48 +91,16 @@
 
 }
 
-// Disable dirty package check 
-static TAutoConsoleVariable<int32> CVarDisableDirtyPackageCheck(TEXT("Concert.DisableDirtyPackageCheck"), 1, TEXT("Disable dirty package check when joining a session."));
-
 /**
- * Connection task used to validate that the workspace has no local changes (according to source control) or in-memory changes (dirty packages).
+ * Common connection task used to validate that the workspace
  */
 class FConcertClientConnectionValidationTask : public IConcertClientConnectionTask
 {
 public:
-	FConcertClientConnectionValidationTask()
-		: SharedState(MakeShared<FSharedAsyncState, ESPMode::ThreadSafe>())
-	{
-	}
+	FConcertClientConnectionValidationTask(const UConcertClientConfig* InClientConfig)
+		: ClientConfig(InClientConfig)
+		, SharedState(MakeShared<FSharedAsyncState, ESPMode::ThreadSafe>())
+	{}
 
-	virtual void Execute() override
-	{
-		check(SharedState && SharedState->Result == EConcertResponseCode::Pending);
-
-		ISourceControlModule& SourceControlModule = ISourceControlModule::Get();
-		ISourceControlProvider& SourceControlProvider = SourceControlModule.GetProvider();
-
-		// Query source control to make sure we don't have any local changes before allowing us to join a remote session
-		if (SourceControlModule.IsEnabled() && SourceControlProvider.IsAvailable())
-		{
-			// Query all content paths
-			TArray<FString> RootPaths;
-			FPackageName::QueryRootContentPaths(RootPaths);
-			for (const FString& RootPath : RootPaths)
-			{
-				const FString RootPathOnDisk = FPackageName::LongPackageNameToFilename(RootPath);
-				SharedState->ContentPaths.Add(FPaths::ConvertRelativePathToFull(RootPathOnDisk));
-			}
-
-			UpdateStatusOperation = ISourceControlOperation::Create<FUpdateStatus>();
-			UpdateStatusOperation->SetUpdateModifiedState(true);
-			SourceControlProvider.Execute(UpdateStatusOperation.ToSharedRef(), SharedState->ContentPaths, EConcurrency::Asynchronous, FSourceControlOperationComplete::CreateStatic(&FConcertClientConnectionValidationTask::HandleAsyncResult, SharedState.ToSharedRef()));
-		}
-		else
-		{
-			HandleDirtyPackages(SharedState.ToSharedRef());
-		}
-	}
-
 	virtual void Abort() override
 	{
@@ -140,9 +108,21 @@
 		SharedState.Reset(); // Always abandon the result
 	}
 
-	virtual void Tick(const bool bShouldCancel) override
+	virtual void Tick(EConcertConnectionTaskAction TaskAction) override
 	{
-		if (bShouldCancel)
+		// InvalidRequest is used here to signify a prompt
+		if (GetStatus() == EConcertResponseCode::InvalidRequest)
 		{
+			if (TaskAction == EConcertConnectionTaskAction::Continue)
+			{
+				SharedState->Result = EConcertResponseCode::Success;
+			}
+			else if (TaskAction == EConcertConnectionTaskAction::Cancel)
+			{
+				SharedState->Result = EConcertResponseCode::Failed;
+			}
+		}
+		else if (TaskAction == EConcertConnectionTaskAction::Cancel)
+		{
 			Cancel();
 		}
@@ -160,65 +140,124 @@
 		return SharedState ? SharedState->Result : EConcertResponseCode::Failed;
 	}
 
+	virtual FText GetPrompt() const override
+	{
+		return SharedState ? SharedState->PromptText : FText::GetEmpty();
+	}
+
 	virtual FText GetError() const override
 	{
 		return SharedState ? SharedState->ErrorText : LOCTEXT("ValidatingWorkspace_Aborted", "The workspace validation request was aborted.");
 	}
 
+	virtual FSimpleDelegate GetErrorDelegate() const override
+	{
+		return FSimpleDelegate::CreateLambda([]()
+		{
+			FMessageLog("Concert").Open();
+		});
+	}
+
 	virtual FText GetDescription() const override
 	{
 		return LOCTEXT("ValidatingWorkspace", "Validating Workspace...");
 	}
 
-private:
+protected:
 	struct FSharedAsyncState
 	{
 		TArray<FString> ContentPaths;
 		EConcertResponseCode Result = EConcertResponseCode::Pending;
+		FText PromptText;
 		FText ErrorText;
 	};
 
-	/** Common function to query for in-memory changes to packages */
-	static void HandleDirtyPackages(TSharedRef<FSharedAsyncState, ESPMode::ThreadSafe> InSharedState)
+	static EConcertResponseCode GetValidationModeResultOnFailure(const UConcertClientConfig* InClientConfig)
 	{
-		//@todo FH: Temporary cvar to disable the dirty package check until it can be reworked as a skippable warning.
-		if (CVarDisableDirtyPackageCheck.GetValueOnAnyThread() > 0)
+		switch (InClientConfig->SourceControlSettings.ValidationMode)
 		{
-			InSharedState->Result = EConcertResponseCode::Success;
-			return;
+		case EConcertSourceValidationMode::Hard:
+			return EConcertResponseCode::Failed;
+		case EConcertSourceValidationMode::Soft:
+			// InvalidRequest is used here to signify a prompt
+			return EConcertResponseCode::InvalidRequest;
+		case EConcertSourceValidationMode::SoftAutoProceed:
+			return EConcertResponseCode::Success;
 		}
-		
-		TArray<UPackage*> DirtyPackages;
-#if WITH_EDITOR
+		return EConcertResponseCode::Failed;
+	}
+
+	virtual void Cancel() {}
+
+	const UConcertClientConfig* ClientConfig;
+	TSharedPtr<FSharedAsyncState, ESPMode::ThreadSafe> SharedState;
+};
+
+/**
+ * Client pre-connection source control validation task
+ */
+class FConcertClientConnectionSourceControlValidationTask : public FConcertClientConnectionValidationTask
+{
+public:
+	FConcertClientConnectionSourceControlValidationTask(const UConcertClientConfig* InClientConfig)
+		: FConcertClientConnectionValidationTask(InClientConfig)
+	{}
+
+	virtual void Execute() override
+	{
+		// Source Control Validation Task
+		check(SharedState && SharedState->Result == EConcertResponseCode::Pending);
+		ISourceControlModule& SourceControlModule = ISourceControlModule::Get();
+		ISourceControlProvider& SourceControlProvider = SourceControlModule.GetProvider();
+
+		// Query source control to make sure we don't have any local changes before allowing us to join a remote session
+		if (SourceControlModule.IsEnabled() && SourceControlProvider.IsAvailable())
 		{
-			FScopedSlowTask SlowTask(1.0f, LOCTEXT("ValidatingWorkspace_DirtyPackagesTask", "Checking for Dirty Packages..."));
-			SlowTask.MakeDialog();
-
-			UEditorLoadingAndSavingUtils::GetDirtyMapPackages(DirtyPackages);
-			UEditorLoadingAndSavingUtils::GetDirtyContentPackages(DirtyPackages);
-		}
-#endif
-
-		if (DirtyPackages.Num() > 0)
-		{
-			// Print which package are dirty to the message log
-			FMessageLog MessageLog("Concert");
-			for (UPackage* Package : DirtyPackages)
+			// Query all content paths
+			TArray<FString> RootPaths;
+			FPackageName::QueryRootContentPaths(RootPaths);
+			for (const FString& RootPath : RootPaths)
 			{
-				MessageLog.Warning(FText::Format(LOCTEXT("ValidatingWorkspace_InMemoryChanges_Log", "Workspace validation failed: '{0}' contains in-memory changes before joining a multi-user session."), FText::FromName(Package->GetFName())));
+				const FString RootPathOnDisk = FPackageName::LongPackageNameToFilename(RootPath);
+				SharedState->ContentPaths.Add(FPaths::ConvertRelativePathToFull(RootPathOnDisk));
 			}
 
-			InSharedState->Result = EConcertResponseCode::Failed;
-			InSharedState->ErrorText = LOCTEXT("ValidatingWorkspace_InMemoryChanges", "This workspace has in-memory changes. Please save or discard these changes before attempting to connect. See Muti-User message log for details.");
+			UpdateStatusOperation = ISourceControlOperation::Create<FUpdateStatus>();
+			UpdateStatusOperation->SetUpdateModifiedState(true);
+			SourceControlProvider.Execute(UpdateStatusOperation.ToSharedRef(), SharedState->ContentPaths, EConcurrency::Asynchronous, FSourceControlOperationComplete::CreateStatic(&FConcertClientConnectionSourceControlValidationTask::HandleAsyncResult, SharedState.ToSharedRef(), ClientConfig));
 		}
 		else
 		{
-			InSharedState->Result = EConcertResponseCode::Success;
+			SharedState->Result = EConcertResponseCode::Success;
 		}
 	}
 
+private:
+	virtual void Cancel() override
+	{
+		if (SharedState && UpdateStatusOperation)
+		{
+			ISourceControlModule& SourceControlModule = ISourceControlModule::Get();
+			ISourceControlProvider& SourceControlProvider = SourceControlModule.GetProvider();
+			if (ensure(SourceControlModule.IsEnabled() && SourceControlProvider.IsAvailable()))
+			{
+				// Gracefully cancel the operation if we're able to
+				// Otherwise just abort it by disowning the result
+				if (SourceControlProvider.CanCancelOperation(UpdateStatusOperation.ToSharedRef()))
+				{
+					SourceControlProvider.CancelOperation(UpdateStatusOperation.ToSharedRef());
+				}
+				else
+				{
+					SharedState.Reset();
+				}
+			}
+			UpdateStatusOperation.Reset();
+		}
+	}
+
 	/** Callback for the source control result - deliberately not a member function as 'this' may be deleted while the request is in flight, so the shared state is used as a safe bridge */
-	static void HandleAsyncResult(const FSourceControlOperationRef& InOperation, ECommandResult::Type InResult, TSharedRef<FSharedAsyncState, ESPMode::ThreadSafe> InSharedState)
+	static void HandleAsyncResult(const FSourceControlOperationRef& InOperation, ECommandResult::Type InResult, TSharedRef<FSharedAsyncState, ESPMode::ThreadSafe> InSharedState, const UConcertClientConfig* InClientConfig)
 	{
 		switch (InResult)
 		{
@@ -228,49 +267,38 @@
 			ISourceControlProvider& SourceControlProvider = SourceControlModule.GetProvider();
 			if (ensure(SourceControlProvider.IsEnabled() && SourceControlProvider.IsAvailable()))
 			{
-				bool bHasLocalChanges = false;
 				FMessageLog MessageLog("Concert");
-				for (const FString& ContentPath : InSharedState->ContentPaths)
-				{
-					IFileManager::Get().IterateDirectoryRecursively(*ContentPath, [&bHasLocalChanges, &SourceControlProvider, &MessageLog](const TCHAR* InFilename, bool InIsDirectory) -> bool
+				TArray<FSourceControlStateRef> ModifiedStates = SourceControlProvider.GetCachedStateByPredicate([](const FSourceControlStateRef& StateRef)
 					{
-						const FString FilenameStr = InFilename;
-						if (!InIsDirectory && FPackageName::IsPackageFilename(FilenameStr))
-						{
-							FSourceControlStatePtr FileState = SourceControlProvider.GetState(FilenameStr, EStateCacheUsage::Use);
-							if (FileState.IsValid() && (FileState->IsAdded() || FileState->IsDeleted() || FileState->IsModified() || (SourceControlProvider.UsesCheckout() && FileState->IsCheckedOut()))) // TODO: Include unversioned files?
-							{
-								bHasLocalChanges = true;
-								MessageLog.Warning(FText::Format(LOCTEXT("ValidatingWorkspace_LocalChanges_Log", "Workspace validation failed: '{0}' contains local changes before joining a multi-user session."), FText::FromString(FileState->GetFilename())));
-								return false; // end iteration
-							}
-						}
-						return true; // continue iteration
+						return ((StateRef->IsAdded() || StateRef->IsDeleted() || StateRef->IsModified()) &&
+							FPackageName::IsPackageFilename(StateRef->GetFilename()));
 					});
 
-					if (bHasLocalChanges)
-					{
-						break;
-					}
+				// Add warning to the multi-user message log
+				for (const FSourceControlStateRef& StateRef : ModifiedStates)
+				{
+					MessageLog.Warning(FText::Format(LOCTEXT("ValidatingWorkspace_LocalChanges_Log", "Workspace validation failed: '{0}' contains local changes before joining a multi-user session."), FText::FromString(StateRef->GetFilename())));
 				}
-				if (bHasLocalChanges)
+
+				if (ModifiedStates.Num() > 0)
 				{
-					InSharedState->Result = EConcertResponseCode::Failed;
-					InSharedState->ErrorText = LOCTEXT("ValidatingWorkspace_LocalChanges", "This workspace has local changes. Please submit or revert these changes before attempting to connect. See Muti-User message log for details.");
+					InSharedState->Result = GetValidationModeResultOnFailure(InClientConfig);
+					InSharedState->PromptText = LOCTEXT("ValidatingWorkspace_SCContinue", "Continue");
+					InSharedState->ErrorText = InClientConfig->SourceControlSettings.ValidationMode == EConcertSourceValidationMode::Hard ?
+						LOCTEXT("ValidatingWorkspace_LocalChangesHard", "This workspace has local changes. Please submit or revert these changes before attempting to connect.") :
+						LOCTEXT("ValidatingWorkspace_LocalChangesSoft", "This workspace has local changes. Local changes won't be immediately available to other users.");
 				}
 				else
 				{
-					HandleDirtyPackages(InSharedState);
+					InSharedState->Result = EConcertResponseCode::Success;
 				}
 			}
 		}
 		break;
-
 		case ECommandResult::Cancelled:
 			InSharedState->Result = EConcertResponseCode::Failed;
 			InSharedState->ErrorText = LOCTEXT("ValidatingWorkspace_Canceled", "The workspace validation request was canceled.");
 			break;
-
 		default:
 			InSharedState->Result = EConcertResponseCode::Failed;
 			InSharedState->ErrorText = LOCTEXT("ValidatingWorkspace_Failed", "The workspace validation request failed. Please check your source control settings.");
@@ -278,31 +306,49 @@
 		}
 	}
 
-	void Cancel()
+	TSharedPtr<FUpdateStatus, ESPMode::ThreadSafe> UpdateStatusOperation;
+};
+
+/**
+ * Client pre-connection in memory change (dirty packages) validation task
+ */
+class FConcertClientConnectionDirtyPackagesValidationTask : public FConcertClientConnectionValidationTask
+{
+public:
+	FConcertClientConnectionDirtyPackagesValidationTask(const UConcertClientConfig* InClientConfig)
+		: FConcertClientConnectionValidationTask(InClientConfig)
+	{}
+
+	virtual void Execute() override
 	{
-		if (SharedState && UpdateStatusOperation)
+		TArray<UPackage*> DirtyPackages;
+#if WITH_EDITOR
 		{
-			ISourceControlModule& SourceControlModule = ISourceControlModule::Get();
-			ISourceControlProvider& SourceControlProvider = SourceControlModule.GetProvider();
-			if (ensure(SourceControlModule.IsEnabled() && SourceControlProvider.IsAvailable()))
+			FScopedSlowTask SlowTask(1.0f, LOCTEXT("ValidatingWorkspace_DirtyPackagesTask", "Checking for Dirty Packages..."));
+			SlowTask.MakeDialog();
+
+			UEditorLoadingAndSavingUtils::GetDirtyMapPackages(DirtyPackages);
+			UEditorLoadingAndSavingUtils::GetDirtyContentPackages(DirtyPackages);
+		}
+#endif
+		if (DirtyPackages.Num() > 0)
+		{
+			// Print which package are dirty to the message log
+			FMessageLog MessageLog("Concert");
+			for (UPackage* Package : DirtyPackages)
 			{
-				// Gracefully cancel the operation if we're able to
-				// Otherwise just abort it by disowning the result
-				if (SourceControlProvider.CanCancelOperation(UpdateStatusOperation.ToSharedRef()))
-				{
-					SourceControlProvider.CancelOperation(UpdateStatusOperation.ToSharedRef());
-				}
-				else
-				{
-					SharedState.Reset();
-				}
+				MessageLog.Warning(FText::Format(LOCTEXT("ValidatingWorkspace_InMemoryChanges_Log", "Workspace validation failed: '{0}' contains in-memory changes before joining a multi-user session."), FText::FromName(Package->GetFName())));
 			}
-			UpdateStatusOperation.Reset();
+
+			SharedState->Result = GetValidationModeResultOnFailure(ClientConfig);
+			SharedState->PromptText = LOCTEXT("ValidatingWorkspace_DiscardChanges", "Discard");
+			SharedState->ErrorText = LOCTEXT("ValidatingWorkspace_InMemoryChanges", "This workspace has in-memory changes. Continue will discard changes.");
+		}
+		else
+		{
+			SharedState->Result = EConcertResponseCode::Success;
 		}
 	}
-
-	TSharedPtr<FUpdateStatus, ESPMode::ThreadSafe> UpdateStatusOperation;
-	TSharedPtr<FSharedAsyncState, ESPMode::ThreadSafe> SharedState;
 };
 
 /**
@@ -785,7 +831,8 @@
 
 	void GetPreConnectionTasks(const IConcertClient& InClient, TArray<TUniquePtr<IConcertClientConnectionTask>>& OutTasks)
 	{
-		OutTasks.Emplace(MakeUnique<FConcertClientConnectionValidationTask>());
+		OutTasks.Emplace(MakeUnique<FConcertClientConnectionSourceControlValidationTask>(InClient.GetConfiguration()));
+		OutTasks.Emplace(MakeUnique<FConcertClientConnectionDirtyPackagesValidationTask>(InClient.GetConfiguration()));
 	}
 
 	/** Returns the toolbar icon title. */
