fixed small typo

diff --git a/Engine/Plugins/Developer/Concert/MultiUserClient/Source/MultiUserClient/Private/MultiUserClientModule.cpp b/Engine/Plugins/Developer/Concert/MultiUserClient/Source/MultiUserClient/Private/MultiUserClientModule.cpp
--- a/Engine/Plugins/Developer/Concert/MultiUserClient/Source/MultiUserClient/Private/MultiUserClientModule.cpp
+++ b/Engine/Plugins/Developer/Concert/MultiUserClient/Source/MultiUserClient/Private/MultiUserClientModule.cpp
@@ -92,7 +92,7 @@ void KillProcess(uint32 ProcessID)
 }
 
 /**
- * Common connection task used to validate that the workspace
+ * Common connection task used to validate the workspace
  */
 class FConcertClientConnectionValidationTask : public IConcertClientConnectionTask
 {
