Validate that there are no locally dirty packages when joining a multi-user session

diff --git a/Engine/Plugins/Developer/Concert/MultiUserClient/Source/MultiUserClient/Private/MultiUserClientModule.cpp b/Engine/Plugins/Developer/Concert/MultiUserClient/Source/MultiUserClient/Private/MultiUserClientModule.cpp
--- a/Engine/Plugins/Developer/Concert/MultiUserClient/Source/MultiUserClient/Private/MultiUserClientModule.cpp
+++ b/Engine/Plugins/Developer/Concert/MultiUserClient/Source/MultiUserClient/Private/MultiUserClientModule.cpp
@@ -19,6 +19,7 @@
 #include "Misc/MessageDialog.h"
 #include "Misc/PackageName.h"
 #include "Misc/Paths.h"
+#include "Misc/ScopedSlowTask.h"
 #include "HAL/IConsoleManager.h"
 #include "HAL/FileManager.h"
 #include "HAL/PlatformProcess.h"
@@ -53,6 +54,7 @@
 	#include "ISettingsModule.h"
 	#include "ISettingsSection.h"
 	#include "LevelEditor.h"
+	#include "FileHelpers.h"
 	#include "WorkspaceMenuStructure.h"
 	#include "WorkspaceMenuStructureModule.h"
 #endif
@@ -64,7 +66,7 @@ static const TCHAR MultiUserServerAppName[] = TEXT("UnrealMultiUserServer");
 #define LOCTEXT_NAMESPACE "MultiUserClient"
 
 /**
- * Connection task used to validate that the workspace has no local changes (according to source control).
+ * Connection task used to validate that the workspace has no local changes (according to source control) or in-memory changes (dirty packages).
  */
 class FConcertClientConnectionValidationTask : public IConcertClientConnectionTask
 {
@@ -99,7 +101,7 @@ public:
 		}
 		else
 		{
-			SharedState->Result = EConcertResponseCode::Success;
+			HandleDirtyPackages(SharedState.ToSharedRef());
 		}
 	}
 
@@ -147,6 +149,31 @@ private:
 		FText ErrorText;
 	};
 
+	/** Common function to query for in-memory changes to packages */
+	static void HandleDirtyPackages(TSharedRef<FSharedAsyncState, ESPMode::ThreadSafe> InSharedState)
+	{
+		TArray<UPackage*> DirtyPackages;
+#if WITH_EDITOR
+		{
+			FScopedSlowTask SlowTask(1.0f, LOCTEXT("ValidatingWorkspace_DirtyPackagesTask", "Checking for Dirty Packages..."));
+			SlowTask.MakeDialog();
+
+			UEditorLoadingAndSavingUtils::GetDirtyMapPackages(DirtyPackages);
+			UEditorLoadingAndSavingUtils::GetDirtyContentPackages(DirtyPackages);
+		}
+#endif
+
+		if (DirtyPackages.Num() > 0)
+		{
+			InSharedState->Result = EConcertResponseCode::Failed;
+			InSharedState->ErrorText = LOCTEXT("ValidatingWorkspace_InMemoryChanges", "This workspace has in-memory changes. Please save or discard these changes before attempting to connect.");
+		}
+		else
+		{
+			InSharedState->Result = EConcertResponseCode::Success;
+		}
+	}
+
 	/** Callback for the source control result - deliberately not a member function as 'this' may be deleted while the request is in flight, so the shared state is used as a safe bridge */
 	static void HandleAsyncResult(const FSourceControlOperationRef& InOperation, ECommandResult::Type InResult, TSharedRef<FSharedAsyncState, ESPMode::ThreadSafe> InSharedState)
 	{
@@ -188,7 +215,7 @@ private:
 				}
 				else
 				{
-					InSharedState->Result = EConcertResponseCode::Success;
+					HandleDirtyPackages(InSharedState);
 				}
 			}
 		}
