Fixed Udp Messaging static endpoints not being properly sent to when the unicast endpoint was forced to a specific interface.
Modified the Multi-User Server launching to bind different port ont he unicast endpoint port of the editor is set when it transfer its settings to it.

Updated Udp messaging settings tooltip as well as multi-user launch server button

diff --git a/Engine/Plugins/Developer/Concert/MultiUserClient/Source/MultiUserClient/Private/MultiUserClientModule.cpp b/Engine/Plugins/Developer/Concert/MultiUserClient/Source/MultiUserClient/Private/MultiUserClientModule.cpp
--- a/Engine/Plugins/Developer/Concert/MultiUserClient/Source/MultiUserClient/Private/MultiUserClientModule.cpp
+++ b/Engine/Plugins/Developer/Concert/MultiUserClient/Source/MultiUserClient/Private/MultiUserClientModule.cpp
@@ -1295,10 +1295,22 @@ private:
 		if (EngineConfig)
 		{
 			TArray<FString> Settings;
-			FString& Setting = Settings.Emplace_GetRef();
+			FString Setting;
 
 			// Unicast endpoint setting
 			EngineConfig->GetString(TEXT("/Script/UdpMessaging.UdpMessagingSettings"), TEXT("UnicastEndpoint"), Setting);
+
+			// if the unicast endpoint port is bound, add 1 to the port for server 
+			if (Setting.ParseIntoArray(Settings, TEXT(":"), false) == 2)
+			{
+				if (Settings[1] != TEXT("0"))
+				{
+					int32 Port = FCString::Atoi(*Settings[1]);
+					Port += 1;
+					Setting = Settings[0] + TEXT(":") + FString::FromInt(Port);
+				}
+			}
+
 			CmdLine = TEXT("-UDPMESSAGING_TRANSPORT_UNICAST=") + Setting;
 
 			// Multicast endpoint setting
diff --git a/Engine/Plugins/Developer/Concert/MultiUserClient/Source/MultiUserClient/Private/Widgets/SConcertBrowser.cpp b/Engine/Plugins/Developer/Concert/MultiUserClient/Source/MultiUserClient/Private/Widgets/SConcertBrowser.cpp
--- a/Engine/Plugins/Developer/Concert/MultiUserClient/Source/MultiUserClient/Private/Widgets/SConcertBrowser.cpp
+++ b/Engine/Plugins/Developer/Concert/MultiUserClient/Source/MultiUserClient/Private/Widgets/SConcertBrowser.cpp
@@ -2113,7 +2113,7 @@ void SConcertSessionBrowser::Construct(const FArguments& InArgs, IConcertClientP
 		.ButtonStyle(FEditorStyle::Get(), "FlatButton.Default")
 		.ButtonIcon(FConcertFrontendStyle::Get()->GetBrush("Concert.NewServer.Small"))
 		.ButtonText(LOCTEXT("LaunchLocalServer", "Launch a Server"))
-		.ButtonToolTip(LOCTEXT("LaunchServerTooltip", "Launch a Multi-User server on your computer unless one is already running"))
+		.ButtonToolTip(LOCTEXT("LaunchServerTooltip", "Launch a Multi-User server on your computer unless one is already running.\nThe editor UDP messaging settings will be passed to the launching server. if the unicast port is bound, the server will use that port number + 1."))
 		.OnButtonClicked(this, &SConcertSessionBrowser::OnLaunchServerButtonClicked);
 
 	// Controls the text displayed in the 'No sessions' panel.
@@ -2534,7 +2534,7 @@ TSharedRef<SWidget> SConcertSessionBrowser::MakeButtonBar()
 		.AutoWidth()
 		.Padding(0, 0, PaddingBetweenButtons, 0)
 		[
-			ConcertBrowserUtils::MakeIconButton(TEXT("FlatButton"), FConcertFrontendStyle::Get()->GetBrush("Concert.NewServer"), LOCTEXT("LaunchServerTooltip", "Launch a Multi-User server on your computer unless one is already running"),
+			ConcertBrowserUtils::MakeIconButton(TEXT("FlatButton"), FConcertFrontendStyle::Get()->GetBrush("Concert.NewServer"), LOCTEXT("LaunchServerTooltip", "Launch a Multi-User server on your computer unless one is already running.\nThe editor UDP messaging settings will be passed to the launching server. if the unicast port is bound, the server will use that port number + 1."),
 				TAttribute<bool>(this, &SConcertSessionBrowser::IsLaunchServerButtonEnabled),
 				FOnClicked::CreateSP(this, &SConcertSessionBrowser::OnLaunchServerButtonClicked),
 				TAttribute<EVisibility>::Create([this]() { return IsLaunchServerButtonEnabled() ? EVisibility::Visible : EVisibility::Collapsed; }))
diff --git a/Engine/Plugins/Messaging/UdpMessaging/Source/UdpMessaging/Private/Shared/UdpMessagingSettings.h b/Engine/Plugins/Messaging/UdpMessaging/Source/UdpMessaging/Private/Shared/UdpMessagingSettings.h
--- a/Engine/Plugins/Messaging/UdpMessaging/Source/UdpMessaging/Private/Shared/UdpMessagingSettings.h
+++ b/Engine/Plugins/Messaging/UdpMessaging/Source/UdpMessaging/Private/Shared/UdpMessagingSettings.h
@@ -65,6 +65,8 @@ public:
 	 * The format is IP_ADDRESS:PORT_NUMBER.
 	 * 0.0.0.0:0 will bind to the default network adapter on Windows,
 	 * and all available network adapters on other operating systems.
+	 * Specifying an interface IP here, will use that interface for multicasting and static endpoint *might* also reach this client through <unicast ip:multicast port>
+	 * Specifying both the IP and Port will allow usage of static endpoint to reach this client
 	 * Can be specified on the command line with `-UDPMESSAGING_TRANSPORT_UNICAST=`
 	 */
 	UPROPERTY(config, EditAnywhere, Category=Transport)
@@ -91,7 +93,7 @@ public:
 	/**
 	 * The IP endpoints of static devices.
 	 *
-	 * Use this setting to list devices on other subnets, such as mobile phones on a WiFi network.
+	 * Use this setting to reach devices on other subnets, such as mobile phones on a WiFi network.
 	 * The format is IP_ADDRESS:PORT_NUMBER.
 	 */
 	UPROPERTY(config, EditAnywhere, Category=Transport, AdvancedDisplay)
diff --git a/Engine/Plugins/Messaging/UdpMessaging/Source/UdpMessaging/Private/Transport/UdpMessageBeacon.cpp b/Engine/Plugins/Messaging/UdpMessaging/Source/UdpMessaging/Private/Transport/UdpMessageBeacon.cpp
--- a/Engine/Plugins/Messaging/UdpMessaging/Source/UdpMessaging/Private/Transport/UdpMessageBeacon.cpp
+++ b/Engine/Plugins/Messaging/UdpMessaging/Source/UdpMessaging/Private/Transport/UdpMessageBeacon.cpp
@@ -188,15 +188,16 @@ bool FUdpMessageBeacon::SendPing(const FTimespan& SocketWaitTime)
 	}
 
 	int32 Sent;
+	bool Result = true;
 	for (const auto& StaticAddress : StaticAddresses)
 	{
 		if (!Socket->SendTo(Writer.GetData(), Writer.Num(), Sent, *StaticAddress))
 		{	
-			return false; // send failed
+			Result = false; // send failed
 		}
 
 	}
-	return true;
+	return Result;
 }
 
 
diff --git a/Engine/Plugins/Messaging/UdpMessaging/Source/UdpMessaging/Private/Transport/UdpMessageTransport.cpp b/Engine/Plugins/Messaging/UdpMessaging/Source/UdpMessaging/Private/Transport/UdpMessageTransport.cpp
--- a/Engine/Plugins/Messaging/UdpMessaging/Source/UdpMessaging/Private/Transport/UdpMessageTransport.cpp
+++ b/Engine/Plugins/Messaging/UdpMessaging/Source/UdpMessaging/Private/Transport/UdpMessageTransport.cpp
@@ -100,7 +100,6 @@ bool FUdpMessageTransport::StartTransport(IMessageTransportHandler& Handler)
 	// create & initialize unicast socket (only on multi-process platforms)
 	UnicastSocket = FUdpSocketBuilder(TEXT("UdpMessageUnicastSocket"))
 		.AsNonBlocking()
-		// Do not bind the unicast socket so it may send to addresses not in its range (i.e other adapters)
 		.BoundToEndpoint(UnicastEndpoint)
 		.WithMulticastLoopback()
 		.WithReceiveBufferSize(UDP_MESSAGING_RECEIVE_BUFFER_SIZE);
@@ -118,11 +117,12 @@ bool FUdpMessageTransport::StartTransport(IMessageTransportHandler& Handler)
 		.AsNonBlocking()
 		.AsReusable()
 #if PLATFORM_WINDOWS
-		// For 0.0.0.0, Windows will pick the default interface instead of all
-		// interfaces. Here we allow to specify which interface to bind to. 
-		// On all other platforms we bind to the wildcard IP address in order
-		// to be able to also receive packets that were sent directly to the
-		// interface IP instead of the multicast address.
+		// If multiple bus instances bind the same unicast ip:port combination (allowed as the socket is marked as reusable), 
+		// then for each packet sent to that ip:port combination, only one of the instances (at the discretion of the OS) will receive it. 
+		// The instance that receives the packet may vary over time, seemingly based on the congestion of its socket.
+		// This isn't the intended usage.
+		// To allow traffic to be sent directly to unicast for discovery, set the interface and port for the unicast endpoint
+		// However for legacy reason, keep binding this as well, although it might be unreliable in some cases
 		.BoundToAddress(UnicastEndpoint.Address)
 #endif
 		.BoundToPort(MulticastEndpoint.Port)
