Expose a Server Port settings for the Multi-User server unicast endpoint over the passed UDP messaging settings when booting the server from the Editor

diff --git a/Engine/Plugins/Developer/Concert/ConcertMain/Source/Concert/Public/ConcertSettings.h b/Engine/Plugins/Developer/Concert/ConcertMain/Source/Concert/Public/ConcertSettings.h
--- a/Engine/Plugins/Developer/Concert/ConcertMain/Source/Concert/Public/ConcertSettings.h
+++ b/Engine/Plugins/Developer/Concert/ConcertMain/Source/Concert/Public/ConcertSettings.h
@@ -194,6 +194,7 @@ struct FConcertClientSettings
 		, AvatarColor(1.0f, 1.0f, 1.0f, 1.0f)
 		, DesktopAvatarActorClass(TEXT("/ConcertSyncClient/DesktopPresence.DesktopPresence_C"))
 		, VRAvatarActorClass(TEXT("/ConcertSyncClient/VRPresence.VRPresence_C"))
+		, ServerPort(0)
 		, DiscoveryTimeoutSeconds(5)
 		, SessionTickFrequencySeconds(1)
 		, LatencyCompensationMs(0)
@@ -218,7 +219,11 @@ struct FConcertClientSettings
 	UPROPERTY(config, EditAnywhere, NoClear, Category = "Client Settings", meta = (MetaClass = "ConcertClientVRPresenceActor", DisplayName = "VR Avatar Actor Class"))
 	FSoftClassPath VRAvatarActorClass;
 
-	/** The timespan at which discovered Concert server are considered stale if they haven't answered back */
+	/** The port to use to reach the server with static endpoints when launched through the editor. This port will be used over the unicast endpoint port in the UDP Messagging settings if non 0 when transferring the editor settings to the launched server. */
+	UPROPERTY(config, EditAnywhere, AdvancedDisplay, Category = "Client Settings")
+	uint16 ServerPort;
+
+	/** The timespan at which discovered Multi-User server are considered stale if they haven't answered back */
 	UPROPERTY(config, EditAnywhere, DisplayName="Discovery Timeout", AdvancedDisplay, Category="Client Settings", meta=(ForceUnits=s))
 	int32 DiscoveryTimeoutSeconds;
 
diff --git a/Engine/Plugins/Developer/Concert/MultiUserClient/Source/MultiUserClient/Private/MultiUserClientModule.cpp b/Engine/Plugins/Developer/Concert/MultiUserClient/Source/MultiUserClient/Private/MultiUserClientModule.cpp
--- a/Engine/Plugins/Developer/Concert/MultiUserClient/Source/MultiUserClient/Private/MultiUserClientModule.cpp
+++ b/Engine/Plugins/Developer/Concert/MultiUserClient/Source/MultiUserClient/Private/MultiUserClientModule.cpp
@@ -1303,14 +1303,18 @@ private:
 			// if the unicast endpoint port is bound, add 1 to the port for server 
 			if (Setting.ParseIntoArray(Settings, TEXT(":"), false) == 2)
 			{
-				if (Settings[1] != TEXT("0"))
+				const UConcertClientConfig* ClientConfig = GetDefault<UConcertClientConfig>();
+				if (ClientConfig && ClientConfig->ClientSettings.ServerPort != 0)
+				{
+					Setting = FString::Printf(TEXT("%s:%d"), *Settings[0], ClientConfig->ClientSettings.ServerPort);
+				}
+				else if (Settings[1] != TEXT("0"))
 				{
 					int32 Port = FCString::Atoi(*Settings[1]);
 					Port += 1;
-					Setting = Settings[0] + TEXT(":") + FString::FromInt(Port);
+					Setting = FString::Printf(TEXT("%s:%d"), *Settings[0], Port);
 				}
 			}
-
 			CmdLine = TEXT("-UDPMESSAGING_TRANSPORT_UNICAST=") + Setting;
 
 			// Multicast endpoint setting
diff --git a/Engine/Plugins/Developer/Concert/MultiUserClient/Source/MultiUserClient/Private/Widgets/SConcertBrowser.cpp b/Engine/Plugins/Developer/Concert/MultiUserClient/Source/MultiUserClient/Private/Widgets/SConcertBrowser.cpp
--- a/Engine/Plugins/Developer/Concert/MultiUserClient/Source/MultiUserClient/Private/Widgets/SConcertBrowser.cpp
+++ b/Engine/Plugins/Developer/Concert/MultiUserClient/Source/MultiUserClient/Private/Widgets/SConcertBrowser.cpp
@@ -2113,7 +2113,7 @@ void SConcertSessionBrowser::Construct(const FArguments& InArgs, IConcertClientP
 		.ButtonStyle(FEditorStyle::Get(), "FlatButton.Default")
 		.ButtonIcon(FConcertFrontendStyle::Get()->GetBrush("Concert.NewServer.Small"))
 		.ButtonText(LOCTEXT("LaunchLocalServer", "Launch a Server"))
-		.ButtonToolTip(LOCTEXT("LaunchServerTooltip", "Launch a Multi-User server on your computer unless one is already running.\nThe editor UDP messaging settings will be passed to the launching server. if the unicast port is bound, the server will use that port number + 1."))
+		.ButtonToolTip(LOCTEXT("LaunchServerTooltip", "Launch a Multi-User server on your computer unless one is already running.\nThe editor UDP messaging settings will be passed to the launching server.\nThe server will use the server port found in Multi-User client settings if non 0 or use the editor udp messaging port number + 1, if non 0."))
 		.OnButtonClicked(this, &SConcertSessionBrowser::OnLaunchServerButtonClicked);
 
 	// Controls the text displayed in the 'No sessions' panel.
@@ -2534,7 +2534,7 @@ TSharedRef<SWidget> SConcertSessionBrowser::MakeButtonBar()
 		.AutoWidth()
 		.Padding(0, 0, PaddingBetweenButtons, 0)
 		[
-			ConcertBrowserUtils::MakeIconButton(TEXT("FlatButton"), FConcertFrontendStyle::Get()->GetBrush("Concert.NewServer"), LOCTEXT("LaunchServerTooltip", "Launch a Multi-User server on your computer unless one is already running.\nThe editor UDP messaging settings will be passed to the launching server. if the unicast port is bound, the server will use that port number + 1."),
+			ConcertBrowserUtils::MakeIconButton(TEXT("FlatButton"), FConcertFrontendStyle::Get()->GetBrush("Concert.NewServer"), LOCTEXT("LaunchServerTooltip", "Launch a Multi-User server on your computer unless one is already running.\nThe editor UDP messaging settings will be passed to the launching server.\nThe server will use the server port found in Multi-User client settings if non 0 or use the editor udp messaging port number + 1, if non 0."),
 				TAttribute<bool>(this, &SConcertSessionBrowser::IsLaunchServerButtonEnabled),
 				FOnClicked::CreateSP(this, &SConcertSessionBrowser::OnLaunchServerButtonClicked),
 				TAttribute<EVisibility>::Create([this]() { return IsLaunchServerButtonEnabled() ? EVisibility::Visible : EVisibility::Collapsed; }))
