Fix multi-user toolbar button not appearing in new toolbar if enabled.

diff --git a/Engine/Plugins/Developer/Concert/MultiUserClient/Source/MultiUserClient/Private/MultiUserClientModule.cpp b/Engine/Plugins/Developer/Concert/MultiUserClient/Source/MultiUserClient/Private/MultiUserClientModule.cpp
--- a/Engine/Plugins/Developer/Concert/MultiUserClient/Source/MultiUserClient/Private/MultiUserClientModule.cpp
+++ b/Engine/Plugins/Developer/Concert/MultiUserClient/Source/MultiUserClient/Private/MultiUserClientModule.cpp
@@ -1194,30 +1194,26 @@ private:
 			);
 
 			TSharedPtr<FExtender> ToolbarExtender = MakeShared<FExtender>();
-			ToolbarExtender->AddToolBarExtension("Game", EExtensionHook::After, CommandList, FToolBarExtensionDelegate::CreateLambda([this, CommandList](FToolBarBuilder& ToolbarBuilder)
+			ToolbarExtender->AddToolBarExtension("Play", EExtensionHook::Before, CommandList, FToolBarExtensionDelegate::CreateLambda([this, CommandList](FToolBarBuilder& ToolbarBuilder)
 			{
-				ToolbarBuilder.BeginSection("Concert");
-				{
-					ToolbarBuilder.AddToolBarButton
-					(
-						FConcertUICommands::Get().TriggerToolbarButtonCmd,
-						NAME_None,
-						TAttribute<FText>::Create([this]() { return GetToolbarButtonIconTitle(); }),
-						TAttribute<FText>::Create([this]() { return GetToolbarButtonTooltip(); }),
-						TAttribute<FSlateIcon>::Create([this]() { return GetToolbarButtonIcon(); })
-					);
-
-					// Add a simple drop-down menu (no label, no icon for the drop-down button itself) to list
-					ToolbarBuilder.AddComboButton(
-						FUIAction(),
-						FOnGetContent::CreateStatic(&GenerateConcertMenuContent, CommandList.ToSharedRef()),
-						LOCTEXT("ConcertToolbarMenu_Label", "Multi-User Utilities"),
-						LOCTEXT("ConcertToolbarMenu_Tooltip", "Multi-User Commands"),
-						FSlateIcon(),
-						true
-					);
-				}
-				ToolbarBuilder.EndSection();
+				ToolbarBuilder.AddToolBarButton
+				(
+					FConcertUICommands::Get().TriggerToolbarButtonCmd,
+					NAME_None,
+					TAttribute<FText>::Create([this]() { return GetToolbarButtonIconTitle(); }),
+					TAttribute<FText>::Create([this]() { return GetToolbarButtonTooltip(); }),
+					TAttribute<FSlateIcon>::Create([this]() { return GetToolbarButtonIcon(); })
+				);
+
+				// Add a simple drop-down menu (no label, no icon for the drop-down button itself) to list
+				ToolbarBuilder.AddComboButton(
+					FUIAction(),
+					FOnGetContent::CreateStatic(&GenerateConcertMenuContent, CommandList.ToSharedRef()),
+					LOCTEXT("ConcertToolbarMenu_Label", "Multi-User Utilities"),
+					LOCTEXT("ConcertToolbarMenu_Tooltip", "Multi-User Commands"),
+					FSlateIcon(),
+					true
+				);
 			}));
 
 			LevelEditorModule.GetToolBarExtensibilityManager()->AddExtender(ToolbarExtender);
diff --git a/Engine/Source/Editor/UnrealEd/Classes/Editor/EditorPerProjectUserSettings.h b/Engine/Source/Editor/UnrealEd/Classes/Editor/EditorPerProjectUserSettings.h
--- a/Engine/Source/Editor/UnrealEd/Classes/Editor/EditorPerProjectUserSettings.h
+++ b/Engine/Source/Editor/UnrealEd/Classes/Editor/EditorPerProjectUserSettings.h
@@ -17,7 +17,7 @@ class UEditorPerProjectUserSettings : public UObject
 	// The following options are exposed in the Preferences Editor 
 
 	/** If enabled, any newly opened UI menus, menu bars, and toolbars will show the developer hooks that would accept extensions */
-	UPROPERTY(EditAnywhere, config, Category=DeveloperTools, meta = (DisplayName = "Display UI Extension Points"))
+	UPROPERTY(EditAnywhere, config, Category=DeveloperTools, meta = (DisplayName = "Display UI Extension Points", ConfigRestartRequired=true))
 	uint32 bDisplayUIExtensionPoints:1;
 
 	/** If enabled, tooltips linked to documentation will show the developer the link bound to that UI item */
