Fixed the most obvious styling errors for Multi-Users and Disaster Recovery UI.
Replaced usage FEditorStyle by FAppStyle as documentation states in FAppStyle.h

diff --git a/Engine/Plugins/Developer/Concert/ConcertSync/ConcertSyncClient/Source/ConcertSyncUI/Private/ConcertFrontendStyle.cpp b/Engine/Plugins/Developer/Concert/ConcertSync/ConcertSyncClient/Source/ConcertSyncUI/Private/ConcertFrontendStyle.cpp
--- a/Engine/Plugins/Developer/Concert/ConcertSync/ConcertSyncClient/Source/ConcertSyncUI/Private/ConcertFrontendStyle.cpp
+++ b/Engine/Plugins/Developer/Concert/ConcertSync/ConcertSyncClient/Source/ConcertSyncUI/Private/ConcertFrontendStyle.cpp
@@ -42,8 +42,8 @@ void FConcertFrontendStyle::Initialize()
 	const FVector2D Icon40x40(40.0f, 40.0f);
 	const FVector2D Icon48x48(48.0f, 48.0f);
 
-	// Most icons were designed to be used at 80% opacity.
-	const FLinearColor IconColorAndOpacity(FLinearColor(1.f, 1.f, 1.f, 0.8f));
+	// Use this to change the opacity. Ex: In UE4, the icon looked fine at 80%, in UE5, icons looks homogeneous with others at 100%.
+	const FLinearColor IconColorAndOpacity(FLinearColor(1.f, 1.f, 1.f, 1.f));
 
 	// 16x16
 	StyleSet->Set("Concert.Persist",         new IMAGE_PLUGIN_BRUSH("Icons/icon_ConcertPersist_16x", Icon16x16));
diff --git a/Engine/Plugins/Developer/Concert/ConcertSync/ConcertSyncClient/Source/ConcertSyncUI/Private/SConcertSessionActivities.cpp b/Engine/Plugins/Developer/Concert/ConcertSync/ConcertSyncClient/Source/ConcertSyncUI/Private/SConcertSessionActivities.cpp
--- a/Engine/Plugins/Developer/Concert/ConcertSync/ConcertSyncClient/Source/ConcertSyncUI/Private/SConcertSessionActivities.cpp
+++ b/Engine/Plugins/Developer/Concert/ConcertSync/ConcertSyncClient/Source/ConcertSyncUI/Private/SConcertSessionActivities.cpp
@@ -3,7 +3,6 @@
 #include "SConcertSessionActivities.h"
 
 #include "IConcertClientWorkspace.h"
-#include "EditorStyleSet.h"
 #include "Algo/Transform.h"
 #include "Misc/AsyncTaskNotification.h"
 #include "Misc/ITransaction.h"
@@ -23,6 +22,7 @@
 #include "Widgets/Text/SRichTextBlock.h"
 #include "Widgets/Text/STextBlock.h"
 #include "Widgets/Views/STableRow.h"
+#include "Styling/AppStyle.h"
 #include "ConcertFrontendStyle.h"
 #include "ConcertFrontendUtils.h"
 #include "ConcertLogGlobal.h"
@@ -428,7 +428,7 @@ void SConcertSessionActivities::Construct(const FArguments& InArgs)
 			+SOverlay::Slot() // Activity list itself.
 			[
 				SNew(SBorder)
-				.BorderImage(FEditorStyle::GetBrush("ToolPanel.GroupBorder"))
+				.BorderImage(FAppStyle::Get().GetBrush("ToolPanel.GroupBorder"))
 				.BorderBackgroundColor(FSlateColor(FLinearColor(0.6, 0.6, 0.6)))
 				.Padding(0)
 				[
@@ -464,7 +464,7 @@ void SConcertSessionActivities::Construct(const FArguments& InArgs)
 			.InitiallyCollapsed(true)
 			.BorderBackgroundColor(FLinearColor(0.6f, 0.6f, 0.6f, 1.0f))
 			.BorderImage_Lambda([this]() { return ConcertFrontendUtils::GetExpandableAreaBorderImage(*ExpandableDetails); })
-			.BodyBorderImage(FEditorStyle::GetBrush("ToolPanel.GroupBorder"))
+			.BodyBorderImage(FAppStyle::Get().GetBrush("ToolPanel.GroupBorder"))
 			.BodyBorderBackgroundColor(FLinearColor::White)
 			.OnAreaExpansionChanged(this, &SConcertSessionActivities::OnDetailsAreaExpansionChanged)
 			.Padding(0.0f)
@@ -472,7 +472,7 @@ void SConcertSessionActivities::Construct(const FArguments& InArgs)
 			[
 				SNew(STextBlock)
 				.Text(LOCTEXT("Details", "Details"))
-				.Font(FEditorStyle::GetFontStyle("DetailsView.CategoryFontStyle"))
+				.Font(FAppStyle::Get().GetFontStyle("DetailsView.CategoryFontStyle"))
 				.ShadowOffset(FVector2D(1.0f, 1.0f))
 			]
 			.BodyContent()
@@ -1075,7 +1075,7 @@ TSharedRef<SWidget> FConcertSessionActivitiesOptions::MakeMenuWidget()
 TSharedRef<SWidget> FConcertSessionActivitiesOptions::MakeViewOptionsWidget()
 {
 	return SNew(SComboButton)
-	.ComboButtonStyle(FEditorStyle::Get(), "GenericFilters.ComboButtonStyle")
+	.ComboButtonStyle(FAppStyle::Get(), "GenericFilters.ComboButtonStyle")
 	.ForegroundColor(FLinearColor::White)
 	.ContentPadding(0)
 	.OnGetMenuContent(this, &FConcertSessionActivitiesOptions::MakeMenuWidget)
@@ -1089,7 +1089,7 @@ TSharedRef<SWidget> FConcertSessionActivitiesOptions::MakeViewOptionsWidget()
 		.AutoWidth()
 		.VAlign(VAlign_Center)
 		[
-			SNew(SImage).Image(FEditorStyle::GetBrush("GenericViewButton")) // The eye ball image.
+			SNew(SImage).Image(FAppStyle::Get().GetBrush("GenericViewButton")) // The eye ball image.
 		]
 
 		+SHorizontalBox::Slot()
diff --git a/Engine/Plugins/Developer/Concert/ConcertSync/ConcertSyncClient/Source/ConcertSyncUI/Public/ConcertFrontendUtils.h b/Engine/Plugins/Developer/Concert/ConcertSync/ConcertSyncClient/Source/ConcertSyncUI/Public/ConcertFrontendUtils.h
--- a/Engine/Plugins/Developer/Concert/ConcertSync/ConcertSyncClient/Source/ConcertSyncUI/Public/ConcertFrontendUtils.h
+++ b/Engine/Plugins/Developer/Concert/ConcertSync/ConcertSyncClient/Source/ConcertSyncUI/Public/ConcertFrontendUtils.h
@@ -4,7 +4,7 @@
 
 #include "CoreMinimal.h"
 #include "ConcertActionDefinition.h"
-#include "EditorStyleSet.h"
+#include "Styling/AppStyle.h"
 #include "Widgets/Layout/SBorder.h"
 #include "Widgets/Images/SImage.h"
 #include "Widgets/Input/SButton.h"
@@ -22,7 +22,7 @@ namespace ConcertFrontendUtils
 	static const FName ButtonIconSyle = TEXT("FontAwesome.10");
 	static const float MinDesiredWidthForBtnAndIcon = 29.f;
 	static const FName ButtonStyleNames[(int32)EConcertActionType::NUM] = {
-		TEXT("FlatButton"),
+		TEXT("SimpleButton"),
 		TEXT("FlatButton.Primary"),
 		TEXT("FlatButton.Info"),
 		TEXT("FlatButton.Success"),
@@ -38,27 +38,26 @@ namespace ConcertFrontendUtils
 	inline TSharedRef<SWidget> CreateDisplayName(const TAttribute<FText>& InDisplayName)
 	{
 		return SNew(SBorder)
-			.BorderImage(FEditorStyle::GetBrush("NoBorder"))
+			.BorderImage(FAppStyle::Get().GetBrush("NoBorder"))
 			.ColorAndOpacity(FLinearColor(0.75f, 0.75f, 0.75f))
 			.Padding(FMargin(6.0f, 4.0f))
 			[
 				SNew(STextBlock)
-				.Font(FEditorStyle::GetFontStyle("BoldFont"))
+				.Font(FAppStyle::Get().GetFontStyle("BoldFont"))
 				.Text(InDisplayName)
 			];
 	}
 
 	inline TSharedRef<SButton> CreateTextButton(const FConcertActionDefinition& InDef)
 	{
-		const FButtonStyle* ButtonStyle = &FEditorStyle::Get().GetWidgetStyle<FButtonStyle>(ButtonStyleNames[(int32)InDef.Type]);
+		const FButtonStyle* ButtonStyle = &FAppStyle::Get().GetWidgetStyle<FButtonStyle>(ButtonStyleNames[(int32)InDef.Type]);
 		check(ButtonStyle);
 		const float ButtonContentWidthPadding = 6.f;
 		const float PaddingCompensation = (ButtonStyle->NormalPadding.Left + ButtonStyle->NormalPadding.Right + ButtonContentWidthPadding * 2);
 
 		return SNew(SButton)
-			.ToolTipText(InDef.ToolTipText)
 			.ButtonStyle(ButtonStyle)
-			.ForegroundColor(FLinearColor::White)
+			.ToolTipText(InDef.ToolTipText)
 			.ContentPadding(FMargin(ButtonContentWidthPadding, 2.f))
 			.IsEnabled(InDef.IsEnabled)
 			.Visibility_Lambda([IsVisible = InDef.IsVisible]() { return IsVisible.Get() ? EVisibility::Visible : EVisibility::Collapsed; })
@@ -69,7 +68,7 @@ namespace ConcertFrontendUtils
 				.HAlign(HAlign_Center)
 				[
 					SNew(STextBlock)
-					.Font(FEditorStyle::Get().GetFontStyle(ButtonIconSyle))
+					.Font(FAppStyle::Get().GetFontStyle(ButtonIconSyle))
 					.Text(InDef.Text)
 					.Justification(ETextJustify::Center)
 				]
@@ -78,12 +77,8 @@ namespace ConcertFrontendUtils
 
 	inline TSharedRef<SButton> CreateIconButton(const FConcertActionDefinition& InDef)
 	{
-		const FButtonStyle* ButtonStyle = &FEditorStyle::Get().GetWidgetStyle<FButtonStyle>(ButtonStyleNames[(int32)InDef.Type]);
-		check(ButtonStyle);
-
 		return SNew(SButton)
-			.ButtonStyle(ButtonStyle)
-			.ForegroundColor(FSlateColor::UseForeground())
+			.ButtonStyle(FAppStyle::Get(), ButtonStyleNames[(int32)InDef.Type])
 			.ToolTipText(InDef.ToolTipText)
 			.ContentPadding(FMargin(0, 0))
 			.IsEnabled(InDef.IsEnabled)
@@ -94,6 +89,7 @@ namespace ConcertFrontendUtils
 			[
 				SNew(SImage)
 				.Image(TAttribute<const FSlateBrush*>::Create([IconStyleAttr = InDef.IconStyle]() { return FConcertFrontendStyle::Get()->GetBrush(IconStyleAttr.Get()); }))
+				.ColorAndOpacity(FSlateColor::UseForeground())
 			];
 	}
 
@@ -174,9 +170,9 @@ namespace ConcertFrontendUtils
 	{
 		if (Area.IsTitleHovered())
 		{
-			return Area.IsExpanded() ? FEditorStyle::GetBrush("DetailsView.CategoryTop_Hovered") : FEditorStyle::GetBrush("DetailsView.CollapsedCategory_Hovered");
+			return Area.IsExpanded() ? FAppStyle::Get().GetBrush("DetailsView.CategoryTop_Hovered") : FAppStyle::Get().GetBrush("DetailsView.CollapsedCategory_Hovered");
 		}
-		return Area.IsExpanded() ? FEditorStyle::GetBrush("DetailsView.CategoryTop") : FEditorStyle::GetBrush("DetailsView.CollapsedCategory");
+		return Area.IsExpanded() ? FAppStyle::Get().GetBrush("DetailsView.CategoryTop") : FAppStyle::Get().GetBrush("DetailsView.CollapsedCategory");
 	}
 
 	static FText FormatRelativeTime(const FDateTime& EventTime, const FDateTime* CurrTime = nullptr)
diff --git a/Engine/Plugins/Developer/Concert/DisasterRecoveryClient/Source/DisasterRecoveryClient/Private/SDisasterRecoveryHub.cpp b/Engine/Plugins/Developer/Concert/DisasterRecoveryClient/Source/DisasterRecoveryClient/Private/SDisasterRecoveryHub.cpp
--- a/Engine/Plugins/Developer/Concert/DisasterRecoveryClient/Source/DisasterRecoveryClient/Private/SDisasterRecoveryHub.cpp
+++ b/Engine/Plugins/Developer/Concert/DisasterRecoveryClient/Source/DisasterRecoveryClient/Private/SDisasterRecoveryHub.cpp
@@ -10,7 +10,7 @@
 #include "SConcertSessionRecovery.h"
 
 #include "DesktopPlatformModule.h"
-#include "EditorStyleSet.h"
+#include "Styling/AppStyle.h"
 #include "EditorDirectories.h"
 #include "EditorFontGlyphs.h"
 #include "Framework/Application/SlateApplication.h"
@@ -84,7 +84,7 @@ public:
 	typedef TFunction<void(TSharedPtr<FDisasterRecoverySessionTreeNode>)> FDoubleClickFunc;
 
 	SLATE_BEGIN_ARGS(SDisasterRecoverySessionTreeNodeWidget)
-		: _Font(FEditorStyle::GetFontStyle(TEXT("NormalFont")))
+		: _Font(FAppStyle::Get().GetFontStyle(TEXT("NormalFont")))
 		{
 		}
 		SLATE_ATTRIBUTE(FSlateFontInfo, Font)
@@ -184,7 +184,7 @@ void SDisasterRecoveryHub::Construct(const FArguments& InArgs, const TSharedPtr<
 	ChildSlot
 	[
 		SNew(SBorder)
-		.BorderImage(FEditorStyle::GetBrush("ToolPanel.GroupBorder"))
+		.BorderImage(FAppStyle::Get().GetBrush("ToolPanel.GroupBorder"))
 		.BorderBackgroundColor(FSlateColor(FLinearColor(0.6, 0.6, 0.6)))
 		.Padding(0)
 		[
@@ -206,7 +206,7 @@ void SDisasterRecoveryHub::Construct(const FArguments& InArgs, const TSharedPtr<
 			.AutoHeight()
 			[
 				SNew(SBorder)
-				.BorderImage(FEditorStyle::GetBrush("ToolPanel.GroupBorder"))
+				.BorderImage(FAppStyle::Get().GetBrush("ToolPanel.GroupBorder"))
 				.Padding(FMargin(2.0f, 2.0f))
 				[
 					MakeToolbarWidget()
@@ -222,7 +222,7 @@ void SDisasterRecoveryHub::Construct(const FArguments& InArgs, const TSharedPtr<
 				.Value(0.25)
 				[
 					SNew(SBorder)
-					.BorderImage(FEditorStyle::GetBrush("ToolPanel.GroupBorder"))
+					.BorderImage(FAppStyle::Get().GetBrush("ToolPanel.GroupBorder"))
 					.Padding(FMargin(2.0f, 2.0f))
 					[
 						MakeSessionTreeView()
@@ -249,12 +249,13 @@ void SDisasterRecoveryHub::Construct(const FArguments& InArgs, const TSharedPtr<
 			[
 				SNew(SUniformGridPanel)
 				.SlotPadding(FMargin(2.0f, 0.0f))
+				.Visibility(InArgs._IsRecoveryMode ? EVisibility::Visible : EVisibility::Collapsed)
 
 				+SUniformGridPanel::Slot(0, 0)
 				[
 					SNew(SButton)
 					.ForegroundColor(FLinearColor::White)
-					.ButtonStyle(FEditorStyle::Get(), TEXT("FlatButton.Success"))
+					.ButtonStyle(FAppStyle::Get(), TEXT("FlatButton.Success"))
 					.ToolTipText(this, &SDisasterRecoveryHub::GetRecoverAllButtonTooltip)
 					.OnClicked(this, &SDisasterRecoveryHub::OnRecoverAllButtonClicked)
 					.HAlign(HAlign_Center)
@@ -263,7 +264,7 @@ void SDisasterRecoveryHub::Construct(const FArguments& InArgs, const TSharedPtr<
 					[
 						SNew(STextBlock)
 						.Text(LOCTEXT("RecoverAll", "Recover All"))
-						.Font( FEditorStyle::GetFontStyle("BoldFont"))
+						.Font( FAppStyle::Get().GetFontStyle("BoldFont"))
 						.ShadowOffset( FVector2D( 1.0f, 1.0f ) )
 					]
 				]
@@ -272,7 +273,7 @@ void SDisasterRecoveryHub::Construct(const FArguments& InArgs, const TSharedPtr<
 				[
 					SNew(SButton)
 					.ForegroundColor(FLinearColor::White)
-					.ButtonStyle(FEditorStyle::Get(), TEXT("FlatButton.Danger"))
+					.ButtonStyle(FAppStyle::Get(), TEXT("FlatButton.Danger"))
 					.ToolTipText(LOCTEXT("CancelRecoveryTooltip", "Discard any recoverable data for your assets and continue with their last saved state"))
 					.OnClicked(this, &SDisasterRecoveryHub::OnCancelButtonClicked)
 					.HAlign(HAlign_Center)
@@ -280,7 +281,7 @@ void SDisasterRecoveryHub::Construct(const FArguments& InArgs, const TSharedPtr<
 					[
 						SNew(STextBlock)
 						.Text(LOCTEXT("Cancel", "Cancel"))
-						.Font( FEditorStyle::GetFontStyle("BoldFont"))
+						.Font(FAppStyle::Get().GetFontStyle("BoldFont"))
 						.ShadowOffset( FVector2D( 1.0f, 1.0f ) )
 					]
 				]
@@ -324,16 +325,15 @@ TSharedRef<SWidget> SDisasterRecoveryHub::MakeToolbarWidget()
 		.AutoWidth()
 		[
 			SNew(SButton)
-			.ForegroundColor(FSlateColor::UseForeground())
-			.ButtonStyle(FEditorStyle::Get(), TEXT("FlatButton"))
+			.ButtonStyle(FAppStyle::Get(), TEXT("SimpleButton"))
 			.OnClicked(this, &SDisasterRecoveryHub::OnImportClicked)
 			.ToolTipText(LOCTEXT("ImportTooltip", "Import a crashed session for inspection."))
-			.ContentPadding(FMargin(4, 4))
+			.ContentPadding(FMargin(6, 4))
 			.VAlign(VAlign_Center)
 			.HAlign(HAlign_Center)
 			[
 				SNew(STextBlock)
-				.Font(FEditorStyle::Get().GetFontStyle(IconFontStyleName))
+				.Font(FAppStyle::Get().GetFontStyle(IconFontStyleName))
 				.Text(FEditorFontGlyphs::Download)
 			]
 		]
@@ -342,15 +342,14 @@ TSharedRef<SWidget> SDisasterRecoveryHub::MakeToolbarWidget()
 		.AutoWidth()
 		[
 			SNew(SButton)
-			.ForegroundColor(FSlateColor::UseForeground())
-			.ButtonStyle(FEditorStyle::Get(), TEXT("FlatButton"))
+			.ButtonStyle(FAppStyle::Get(), TEXT("SimpleButton"))
 			.OnClicked(this, &SDisasterRecoveryHub::OnDeleteClicked)
 			.IsEnabled(this, &SDisasterRecoveryHub::IsDeleteButtonEnabled)
 			.ToolTipText(LOCTEXT("DeleteTooltip", "Delete the selected session."))
-			.ContentPadding(FMargin(4, 4))
+			.ContentPadding(FMargin(6, 4))
 			[
 				SNew(STextBlock)
-				.Font(FEditorStyle::Get().GetFontStyle(IconFontStyleName))
+				.Font(FAppStyle::Get().GetFontStyle(IconFontStyleName))
 				.Text(FEditorFontGlyphs::Trash)
 			]
 		]
@@ -365,13 +364,13 @@ TSharedRef<SWidget> SDisasterRecoveryHub::MakeToolbarWidget()
 		.AutoWidth()
 		[
 			SNew(SButton)
-			.ForegroundColor(FSlateColor::UseForeground())
-			.ButtonStyle(FEditorStyle::Get(), TEXT("FlatButton"))
+			.ButtonStyle(FAppStyle::Get(), TEXT("SimpleButton"))
 			.OnClicked(this, &SDisasterRecoveryHub::OnConfigClicked)
 			.ToolTipText(LOCTEXT("ConfigureTooltip", "Open the configuration tab."))
+			.ContentPadding(FMargin(6, 4))
 			[
 				SNew(STextBlock)
-				.Font(FEditorStyle::Get().GetFontStyle(IconFontStyleName))
+				.Font(FAppStyle::Get().GetFontStyle(IconFontStyleName))
 				.Text(FEditorFontGlyphs::Cogs)
 			]
 		];
@@ -455,7 +454,7 @@ TSharedRef<SWidget> SDisasterRecoveryHub::MakeSessionTreeView()
 TSharedRef<ITableRow> SDisasterRecoveryHub::OnGenerateSessionTreeNodeWidget(TSharedPtr<FDisasterRecoverySessionTreeNode> TreeNode, const TSharedRef<STableViewBase>& OwnerTable)
 {
 	return SNew(SDisasterRecoverySessionTreeNodeWidget, TreeNode, OwnerTable)
-		.Font(TreeNode == UnreviewedCrashCategoryRootNode ? FEditorStyle::GetFontStyle(TEXT("BoldFont")) : FEditorStyle::GetFontStyle( TEXT("NormalFont")));
+		.Font(TreeNode == UnreviewedCrashCategoryRootNode ? FAppStyle::Get().GetFontStyle(TEXT("BoldFont")) : FAppStyle::Get().GetFontStyle( TEXT("NormalFont")));
 }
 
 void SDisasterRecoveryHub::OnGetSessionTreeNodeChildren(TSharedPtr<FDisasterRecoverySessionTreeNode> InParent, TArray<TSharedPtr<FDisasterRecoverySessionTreeNode>>& OutChildren) const
diff --git a/Engine/Plugins/Developer/Concert/MultiUserClient/Source/MultiUserClient/Private/Widgets/SActiveSession.cpp b/Engine/Plugins/Developer/Concert/MultiUserClient/Source/MultiUserClient/Private/Widgets/SActiveSession.cpp
--- a/Engine/Plugins/Developer/Concert/MultiUserClient/Source/MultiUserClient/Private/Widgets/SActiveSession.cpp
+++ b/Engine/Plugins/Developer/Concert/MultiUserClient/Source/MultiUserClient/Private/Widgets/SActiveSession.cpp
@@ -12,7 +12,7 @@
 #include "Algo/Transform.h"
 #include "Misc/PackageName.h"
 #include "EditorFontGlyphs.h"
-#include "EditorStyleSet.h"
+#include "Styling/AppStyle.h"
 #include "FileHelpers.h"
 #include "IConcertClient.h"
 #include "IConcertSyncClient.h"
@@ -141,7 +141,7 @@ public:
 				.AutoWidth()
 				[
 					SNew(SImage)
-					.Image(FEditorStyle::GetBrush("LevelEditor.Tabs.Levels"))
+					.Image(FAppStyle::Get().GetBrush("LevelEditor.Tabs.Levels"))
 				]
 				+SHorizontalBox::Slot()
 				.AutoWidth()
@@ -288,7 +288,7 @@ public:
 	FSlateFontInfo GetAvatarFont() const
 	{
 		// This font is used to render a small square box filled with the avatar color.
-		FSlateFontInfo ClientIconFontInfo = FEditorStyle::Get().GetFontStyle(ConcertFrontendUtils::ButtonIconSyle);
+		FSlateFontInfo ClientIconFontInfo = FAppStyle::Get().GetFontStyle(ConcertFrontendUtils::ButtonIconSyle);
 		ClientIconFontInfo.Size = 8;
 		ClientIconFontInfo.OutlineSettings.OutlineSize = 1;
 
@@ -363,12 +363,12 @@ void SActiveSession::Construct(const FArguments& InArgs, TSharedPtr<IConcertSync
 		.Padding(FMargin(4.0f, 1.0f))
 		[
 			SNew(SBorder)
-			.BorderImage(FEditorStyle::GetBrush("NoBorder"))
+			.BorderImage(FAppStyle::Get().GetBrush("NoBorder"))
 			.ColorAndOpacity(FLinearColor(0.75f, 0.75f, 0.75f))
 			.Padding(FMargin(0.0f, 4.0f, 6.0f, 4.0f))
 			[
 				SNew(STextBlock)
-				.Font(FEditorStyle::GetFontStyle("BoldFont"))
+				.Font(FAppStyle::Get().GetFontStyle("BoldFont"))
 				.Text(this, &SActiveSession::GetConnectionStatusText)
 			]
 		];
@@ -417,7 +417,7 @@ void SActiveSession::Construct(const FArguments& InArgs, TSharedPtr<IConcertSync
 		.Padding(0.0f, 0.0f, 0.0f, 4.0f) // Add space between the status bar and the clients list.
 		[
 			SNew(SBorder)
-			.BorderImage(FEditorStyle::GetBrush("ToolPanel.GroupBorder"))
+			.BorderImage(FAppStyle::Get().GetBrush("ToolPanel.GroupBorder"))
 			.Padding(0.0f)
 			[
 				StatusBar
@@ -437,13 +437,13 @@ void SActiveSession::Construct(const FArguments& InArgs, TSharedPtr<IConcertSync
 			.Value(0.3)
 			[
 				SNew(SBorder)
-				.BorderImage(FEditorStyle::GetBrush("ToolPanel.GroupBorder"))
+				.BorderImage(FAppStyle::Get().GetBrush("ToolPanel.GroupBorder"))
 				.Padding(0.0f)
 				[
 					SAssignNew(ClientArea, SExpandableArea)
 					.BorderBackgroundColor(FLinearColor(0.6f, 0.6f, 0.6f, 1.0f))
 					.BorderImage_Lambda([this]() { return ConcertFrontendUtils::GetExpandableAreaBorderImage(*ClientArea); })
-					.BodyBorderImage(FEditorStyle::GetBrush("ToolPanel.GroupBorder"))
+					.BodyBorderImage(FAppStyle::Get().GetBrush("ToolPanel.GroupBorder"))
 					.BodyBorderBackgroundColor(FLinearColor::White)
 					.OnAreaExpansionChanged(this, &SActiveSession::OnClientAreaExpansionChanged)
 					.Padding(0.0f)
@@ -451,7 +451,7 @@ void SActiveSession::Construct(const FArguments& InArgs, TSharedPtr<IConcertSync
 					[
 						SNew(STextBlock)
 						.Text(LOCTEXT("SessionConnectedClients", "Clients"))
-						.Font(FEditorStyle::GetFontStyle("DetailsView.CategoryFontStyle"))
+						.Font(FAppStyle::Get().GetFontStyle("DetailsView.CategoryFontStyle"))
 						.ShadowOffset(FVector2D(1.0f, 1.0f))
 					]
 					.BodyContent()
@@ -481,13 +481,13 @@ void SActiveSession::Construct(const FArguments& InArgs, TSharedPtr<IConcertSync
 			.Value(0.7)
 			[
 				SNew(SBorder)
-				.BorderImage(FEditorStyle::GetBrush("ToolPanel.GroupBorder"))
+				.BorderImage(FAppStyle::Get().GetBrush("ToolPanel.GroupBorder"))
 				.Padding(0.0f)
 				[
 					SAssignNew(HistoryArea, SExpandableArea)
 					.BorderBackgroundColor(FLinearColor(0.6f, 0.6f, 0.6f, 1.0f))
 					.BorderImage_Lambda([this]() { return ConcertFrontendUtils::GetExpandableAreaBorderImage(*HistoryArea); })
-					.BodyBorderImage(FEditorStyle::GetBrush("ToolPanel.GroupBorder"))
+					.BodyBorderImage(FAppStyle::Get().GetBrush("ToolPanel.GroupBorder"))
 					.BodyBorderBackgroundColor(FLinearColor::White)
 					.OnAreaExpansionChanged(this, &SActiveSession::OnHistoryAreaExpansionChanged)
 					.Padding(0.0f)
@@ -495,7 +495,7 @@ void SActiveSession::Construct(const FArguments& InArgs, TSharedPtr<IConcertSync
 					[
 						SNew(STextBlock)
 						.Text(LOCTEXT("SessionHistory", "History"))
-						.Font(FEditorStyle::GetFontStyle("DetailsView.CategoryFontStyle"))
+						.Font(FAppStyle::Get().GetFontStyle("DetailsView.CategoryFontStyle"))
 						.ShadowOffset(FVector2D(1.0f, 1.0f))
 					]
 					.BodyContent()
@@ -709,7 +709,7 @@ const FButtonStyle& SActiveSession::GetConnectionIconStyle() const
 		}
 	}
 	
-	return FEditorStyle::Get().GetWidgetStyle<FButtonStyle>(ConcertFrontendUtils::ButtonStyleNames[(int32)ButtonStyle]);
+	return FAppStyle::Get().GetWidgetStyle<FButtonStyle>(ConcertFrontendUtils::ButtonStyleNames[(int32)ButtonStyle]);
 }
 
 FSlateColor SActiveSession::GetConnectionIconColor() const
@@ -719,7 +719,7 @@ FSlateColor SActiveSession::GetConnectionIconColor() const
 
 FSlateFontInfo SActiveSession::GetConnectionIconFontInfo() const
 {
-	FSlateFontInfo ConnectionIconFontInfo = FEditorStyle::Get().GetFontStyle(ConcertFrontendUtils::ButtonIconSyle);
+	FSlateFontInfo ConnectionIconFontInfo = FAppStyle::Get().GetFontStyle(ConcertFrontendUtils::ButtonIconSyle);
 	ConnectionIconFontInfo.OutlineSettings.OutlineSize = 1;
 	ConnectionIconFontInfo.OutlineSettings.OutlineColor = GetConnectionIconStyle().Pressed.TintColor.GetSpecifiedColor();
 
diff --git a/Engine/Plugins/Developer/Concert/MultiUserClient/Source/MultiUserClient/Private/Widgets/SConcertBrowser.cpp b/Engine/Plugins/Developer/Concert/MultiUserClient/Source/MultiUserClient/Private/Widgets/SConcertBrowser.cpp
--- a/Engine/Plugins/Developer/Concert/MultiUserClient/Source/MultiUserClient/Private/Widgets/SConcertBrowser.cpp
+++ b/Engine/Plugins/Developer/Concert/MultiUserClient/Source/MultiUserClient/Private/Widgets/SConcertBrowser.cpp
@@ -17,7 +17,7 @@
 #include "MultiUserClientUtils.h"
 
 #include "Algo/Transform.h"
-#include "EditorStyleSet.h"
+#include "Styling/AppStyle.h"
 #include "Framework/Commands/GenericCommands.h"
 #include "Framework/MultiBox/MultiBoxBuilder.h"
 #include "Internationalization/Regex.h"
@@ -74,8 +74,7 @@ TSharedRef<SButton> MakeIconButton(const FName& ButtonStyle, const TAttribute<FT
 	const FSlateColor& ForegroundColor, const TAttribute<EVisibility>& Visibility = EVisibility::Visible, const TAttribute<FMargin>& ContentPadding = FMargin(3.0, 2.0), const FName FontStyle = IconColumnFontName)
 {
 	return SNew(SButton)
-		.ForegroundColor(ForegroundColor)
-		.ButtonStyle(FEditorStyle::Get(), ButtonStyle)
+		.ButtonStyle(FAppStyle::Get(), ButtonStyle)
 		.OnClicked(OnClicked)
 		.ToolTipText(Tooltip)
 		.ContentPadding(ContentPadding)
@@ -85,7 +84,7 @@ TSharedRef<SButton> MakeIconButton(const FName& ButtonStyle, const TAttribute<FT
 		.HAlign(HAlign_Center)
 		[
 			SNew(STextBlock)
-			.Font(FEditorStyle::Get().GetFontStyle(FontStyle))
+			.Font(FAppStyle::Get().GetFontStyle(FontStyle))
 			.Text(GlyphIcon)
 		];
 }
@@ -94,8 +93,7 @@ TSharedRef<SButton> MakeIconButton(const FName& ButtonStyle, const TAttribute<FT
 TSharedRef<SButton> MakeIconButton(const FName& ButtonStyle, const TAttribute<const FSlateBrush*>& Icon, const TAttribute<FText>& Tooltip, const TAttribute<bool>& EnabledAttribute, const FOnClicked& OnClicked, const TAttribute<EVisibility>& Visibility = EVisibility::Visible)
 {
 	return SNew(SButton)
-		.ForegroundColor(FSlateColor::UseForeground())
-		.ButtonStyle(FEditorStyle::Get(), ButtonStyle)
+		.ButtonStyle(FAppStyle::Get(), ButtonStyle)
 		.OnClicked(OnClicked)
 		.ToolTipText(Tooltip)
 		.ContentPadding(FMargin(0, 0))
@@ -104,7 +102,9 @@ TSharedRef<SButton> MakeIconButton(const FName& ButtonStyle, const TAttribute<co
 		.VAlign(VAlign_Center)
 		.HAlign(HAlign_Center)
 		[
-			SNew(SImage).Image(Icon)
+			SNew(SImage)
+			.Image(Icon)
+			.ColorAndOpacity(FSlateColor::UseForeground())
 		];
 }
 
@@ -118,11 +118,11 @@ FText GetServerVersionIgnoredTooltip()
 TSharedRef<SWidget> MakeServerVersionIgnoredWidget(EConcertServerFlags InServerFlags)
 {
 	return SNew(SBorder)
-		.BorderImage(FEditorStyle::GetBrush("NoBorder"))
-		.ColorAndOpacity(FEditorStyle::Get().GetWidgetStyle<FButtonStyle>("FlatButton.Warning").Normal.TintColor.GetSpecifiedColor())
+		.BorderImage(FAppStyle::Get().GetBrush("NoBorder"))
+		.ColorAndOpacity(FAppStyle::Get().GetWidgetStyle<FButtonStyle>("FlatButton.Warning").Normal.TintColor.GetSpecifiedColor())
 		[
 			SNew(STextBlock)
-			.Font(FEditorStyle::Get().GetFontStyle("FontAwesome.9"))
+			.Font(FAppStyle::Get().GetFontStyle("FontAwesome.9"))
 			.Text(FEditorFontGlyphs::Exclamation_Triangle)
 			.ToolTipText(GetServerVersionIgnoredTooltip())
 			.Visibility((InServerFlags & EConcertServerFlags::IgnoreSessionRequirement) != EConcertServerFlags::None ? EVisibility::Visible : EVisibility::Collapsed)
@@ -1035,9 +1035,9 @@ TSharedRef<SWidget> SSessionRow::GenerateWidgetForColumn(const FName& ColumnName
 		.ToolTipText(ItemPin->Type == FConcertSessionItem::EType::ActiveSession ? LOCTEXT("ActiveIconTooltip", "Active session") : LOCTEXT("ArchivedIconTooltip", "Archived Session"))
 		[
 			SNew(STextBlock)
-			.Font(FEditorStyle::Get().GetFontStyle(ConcertBrowserUtils::IconColumnFontName))
+			.Font(FAppStyle::Get().GetFontStyle(ConcertBrowserUtils::IconColumnFontName))
 			.Text(ItemPin->Type == FConcertSessionItem::EType::ActiveSession ? FEditorFontGlyphs::Circle : FEditorFontGlyphs::Archive)
-			.ColorAndOpacity(ItemPin->Type == FConcertSessionItem::EType::ActiveSession ? FEditorStyle::Get().GetWidgetStyle<FButtonStyle>("FlatButton.Success").Normal.TintColor : FSlateColor::UseSubduedForeground())
+			.ColorAndOpacity(ItemPin->Type == FConcertSessionItem::EType::ActiveSession ? FAppStyle::Get().GetWidgetStyle<FButtonStyle>("FlatButton.Success").Normal.TintColor : FSlateColor::UseSubduedForeground())
 		];
 	}
 
@@ -1047,7 +1047,7 @@ TSharedRef<SWidget> SSessionRow::GenerateWidgetForColumn(const FName& ColumnName
 	if (ItemPin->Type == FConcertSessionItem::EType::ActiveSession)
 	{
 		FontColor = bIsDefaultConfig ? FSlateColor(FLinearColor::White) : FSlateColor(FLinearColor::White * 0.8f);
-		FontInfo = FEditorStyle::Get().GetFontStyle("NormalFont");
+		FontInfo = FAppStyle::Get().GetFontStyle("NormalFont");
 	}
 	else
 	{
@@ -1277,7 +1277,7 @@ TSharedRef<SWidget> SNewSessionRow::GenerateWidgetForColumn(const FName& ColumnN
 			.HAlign(HAlign_Center)
 			[
 				SNew(STextBlock)
-					.Font(FEditorStyle::Get().GetFontStyle(ConcertBrowserUtils::IconColumnFontName))
+					.Font(FAppStyle::Get().GetFontStyle(ConcertBrowserUtils::IconColumnFontName))
 					.Text(FEditorFontGlyphs::Plus_Circle)
 			];
 	}
@@ -1372,7 +1372,7 @@ TSharedRef<SWidget> SNewSessionRow::OnGenerateServersComboOptionWidget(TSharedPt
 		.AutoWidth()
 		[
 			SNew(STextBlock)
-			.Font(bIsDefaultServer ? FEditorStyle::GetFontStyle("BoldFont") : FEditorStyle::GetFontStyle("NormalFont"))
+			.Font(bIsDefaultServer ? FAppStyle::Get().GetFontStyle("BoldFont") : FAppStyle::Get().GetFontStyle("NormalFont"))
 			.Text(GetServerDisplayName(ServerItem->ServerName))
 		]
 		+SHorizontalBox::Slot()
@@ -1473,7 +1473,7 @@ TSharedRef<SWidget> SNewSessionRow::MakeSelectedServerWidget()
 		.Padding(2, 0, 0, 0)
 		[
 			SNew(STextBlock)
-			.Font(FEditorStyle::Get().GetFontStyle("FontAwesome.9"))
+			.Font(FAppStyle::Get().GetFontStyle("FontAwesome.9"))
 			.Text(this, &SNewSessionRow::GetSelectedServerIgnoreVersionText)
 			.ToolTipText(this, &SNewSessionRow::GetSelectedServerIgnoreVersionTooltip)
 		];
@@ -1669,7 +1669,7 @@ TSharedRef<SWidget> SSaveRestoreSessionRow::GenerateWidgetForColumn(const FName&
 			.Padding(8, 0)
 			[
 				SNew(SExpanderArrow, SharedThis(this))
-				.StyleSet(&FEditorStyle::Get())
+				.StyleSet(&FAppStyle::Get())
 				.ShouldDrawWires(true)
 			];
 	}
@@ -1822,7 +1822,7 @@ public:
 		, _ThrobberVisibility(EVisibility::Visible)
 		, _ButtonVisibility(EVisibility::Visible)
 		, _IsButtonEnabled(true)
-		, _ButtonStyle(&FEditorStyle::Get().GetWidgetStyle<FButtonStyle>("FlatButton"))
+		, _ButtonStyle(&FAppStyle::Get().GetWidgetStyle<FButtonStyle>("FlatButton"))
 		, _ButtonIcon()
 		, _ButtonText()
 		, _ButtonToolTip()
@@ -2110,7 +2110,7 @@ void SConcertSessionBrowser::Construct(const FArguments& InArgs, IConcertClientP
 		.Text(LOCTEXT("LookingForServer", "Looking for Multi-User Servers..."))
 		.Visibility_Lambda([this]() { return Controller->GetServers().Num() == 0 ? EVisibility::Visible : EVisibility::Hidden; })
 		.IsButtonEnabled(this, &SConcertSessionBrowser::IsLaunchServerButtonEnabled)
-		.ButtonStyle(FEditorStyle::Get(), "FlatButton.Default")
+		.ButtonStyle(FAppStyle::Get(), "FlatButton.Default")
 		.ButtonIcon(FConcertFrontendStyle::Get()->GetBrush("Concert.NewServer.Small"))
 		.ButtonText(LOCTEXT("LaunchLocalServer", "Launch a Server"))
 		.ButtonToolTip(LOCTEXT("LaunchServerTooltip", "Launch a Multi-User server on your computer unless one is already running.\nThe editor UDP messaging settings will be passed to the launching server.\nThe server will use the server port found in Multi-User client settings if non 0 or use the editor udp messaging port number + 1, if non 0."))
@@ -2135,7 +2135,7 @@ void SConcertSessionBrowser::Construct(const FArguments& InArgs, IConcertClientP
 		.Visibility_Lambda([this]() { return Controller->GetServers().Num() > 0 && Sessions.Num() == 0 && !Controller->IsCreatingSession() ? EVisibility::Visible : EVisibility::Hidden; })
 		.ThrobberVisibility_Lambda([this]() { return !Controller->HasReceivedInitialSessionList() ? EVisibility::Visible : EVisibility::Collapsed; })
 		.ButtonVisibility_Lambda([this]() { return Controller->HasReceivedInitialSessionList() && Controller->GetActiveSessions().Num() == 0 && Controller->GetArchivedSessions().Num() == 0 ? EVisibility::Visible : EVisibility::Collapsed; })
-		.ButtonStyle(FEditorStyle::Get(), "FlatButton.Default")
+		.ButtonStyle(FAppStyle::Get(), "FlatButton.Default")
 		.ButtonIcon(FConcertFrontendStyle::Get()->GetBrush("Concert.NewSession.Small"))
 		.ButtonText(LOCTEXT("CreateSession", "Create Session"))
 		.ButtonToolTip(LOCTEXT("CreateSessionTooltip", "Create a new session"))
@@ -2191,7 +2191,7 @@ TSharedRef<SWidget> SConcertSessionBrowser::MakeBrowserContent()
 			.Value(0.6)
 			[
 				SNew(SBorder)
-				.BorderImage(FEditorStyle::GetBrush("ToolPanel.GroupBorder"))
+				.BorderImage(FAppStyle::Get().GetBrush("ToolPanel.GroupBorder"))
 				.Padding(FMargin(1.0f, 2.0f))
 				[
 					SNew(SVerticalBox)
@@ -2246,7 +2246,7 @@ TSharedRef<SWidget> SConcertSessionBrowser::MakeBrowserContent()
 			.Value(0.4)
 			[
 				SAssignNew(SessionDetailsView, SBorder)
-				.BorderImage(FEditorStyle::GetBrush("ToolPanel.GroupBorder"))
+				.BorderImage(FAppStyle::Get().GetBrush("ToolPanel.GroupBorder"))
 				[
 					NoSessionSelectedPanel.ToSharedRef()
 				]
@@ -2467,7 +2467,7 @@ TSharedRef<SWidget> SConcertSessionBrowser::MakeControlBar()
 		.AutoWidth()
 		[
 			SNew(STextBlock)
-			.Font(FEditorStyle::Get().GetFontStyle("FontAwesome.8"))
+			.Font(FAppStyle::Get().GetFontStyle("FontAwesome.8"))
 			.ColorAndOpacity_Lambda([this]() { return Controller->GetConcertClient()->GetClientInfo().AvatarColor; })
 			.Text(FEditorFontGlyphs::Square)
 		]
@@ -2479,24 +2479,22 @@ TSharedRef<SWidget> SConcertSessionBrowser::MakeControlBar()
 		.Padding(3, 0, 2, 0)
 		[
 			SNew(STextBlock)
-			.Text_Lambda([this]() {return FText::FromString(Controller->GetConcertClient()->GetClientInfo().DisplayName);} )
+			.Text_Lambda([this]() { return FText::FromString(Controller->GetConcertClient()->GetClientInfo().DisplayName);} )
 		]
 
 		// The "Settings" icons.
 		+SHorizontalBox::Slot()
 		.VAlign(VAlign_Fill)
 		.AutoWidth()
-		.Padding(0, 0, 1, 0)
 		[
 			SNew(SButton)
-			.ForegroundColor(FSlateColor::UseForeground())
-			.ButtonStyle(FEditorStyle::Get(), TEXT("FlatButton"))
+			.ButtonStyle(FAppStyle::Get(), TEXT("SimpleButton"))
 			.OnClicked_Lambda([](){FModuleManager::GetModulePtr<ISettingsModule>("Settings")->ShowViewer("Project", "Plugins", "Concert"); return FReply::Handled(); })
 			.VAlign(VAlign_Center)
 			.HAlign(HAlign_Fill)
 			[
 				SNew(STextBlock)
-				.Font(FEditorStyle::Get().GetFontStyle("FontAwesome.16"))
+				.Font(FAppStyle::Get().GetFontStyle("FontAwesome.16"))
 				.Text(FEditorFontGlyphs::Cogs)
 			]
 		];
@@ -2534,7 +2532,7 @@ TSharedRef<SWidget> SConcertSessionBrowser::MakeButtonBar()
 		.AutoWidth()
 		.Padding(0, 0, PaddingBetweenButtons, 0)
 		[
-			ConcertBrowserUtils::MakeIconButton(TEXT("FlatButton"), FConcertFrontendStyle::Get()->GetBrush("Concert.NewServer"), LOCTEXT("LaunchServerTooltip", "Launch a Multi-User server on your computer unless one is already running.\nThe editor UDP messaging settings will be passed to the launching server.\nThe server will use the server port found in Multi-User client settings if non 0 or use the editor udp messaging port number + 1, if non 0."),
+			ConcertBrowserUtils::MakeIconButton(TEXT("SimpleButton"), FConcertFrontendStyle::Get()->GetBrush("Concert.NewServer"), LOCTEXT("LaunchServerTooltip", "Launch a Multi-User server on your computer unless one is already running.\nThe editor UDP messaging settings will be passed to the launching server.\nThe server will use the server port found in Multi-User client settings if non 0 or use the editor udp messaging port number + 1, if non 0."),
 				TAttribute<bool>(this, &SConcertSessionBrowser::IsLaunchServerButtonEnabled),
 				FOnClicked::CreateSP(this, &SConcertSessionBrowser::OnLaunchServerButtonClicked),
 				TAttribute<EVisibility>::Create([this]() { return IsLaunchServerButtonEnabled() ? EVisibility::Visible : EVisibility::Collapsed; }))
@@ -2543,7 +2541,7 @@ TSharedRef<SWidget> SConcertSessionBrowser::MakeButtonBar()
 		// Stop server.
 		+SHorizontalBox::Slot()
 		[
-			ConcertBrowserUtils::MakeIconButton(TEXT("FlatButton"), FConcertFrontendStyle::Get()->GetBrush("Concert.CloseServer"), LOCTEXT("ShutdownServerTooltip", "Shutdown the Multi-User server running on this computer."),
+			ConcertBrowserUtils::MakeIconButton(TEXT("SimpleButton"), FConcertFrontendStyle::Get()->GetBrush("Concert.CloseServer"), LOCTEXT("ShutdownServerTooltip", "Shutdown the Multi-User server running on this computer."),
 				true, // Always enabled, but might be collapsed.
 				FOnClicked::CreateSP(this, &SConcertSessionBrowser::OnShutdownServerButtonClicked),
 				TAttribute<EVisibility>::Create([this]() { return IsLaunchServerButtonEnabled() ? EVisibility::Collapsed : EVisibility::Visible; }))
@@ -2554,7 +2552,7 @@ TSharedRef<SWidget> SConcertSessionBrowser::MakeButtonBar()
 		.AutoWidth()
 		.Padding(0, 0, PaddingBetweenButtons, 0)
 		[
-			ConcertBrowserUtils::MakeIconButton(TEXT("FlatButton"), FConcertFrontendStyle::Get()->GetBrush("Concert.NewSession"), LOCTEXT("NewButtonTooltip", "Create a new session"),
+			ConcertBrowserUtils::MakeIconButton(TEXT("SimpleButton"), FConcertFrontendStyle::Get()->GetBrush("Concert.NewSession"), LOCTEXT("NewButtonTooltip", "Create a new session"),
 				TAttribute<bool>(this, &SConcertSessionBrowser::IsNewButtonEnabled),
 				FOnClicked::CreateSP(this, &SConcertSessionBrowser::OnNewButtonClicked))
 		]
@@ -2572,7 +2570,7 @@ TSharedRef<SWidget> SConcertSessionBrowser::MakeButtonBar()
 		.AutoWidth()
 		.Padding(0, 0, PaddingBetweenButtons, 0)
 		[
-			ConcertBrowserUtils::MakeIconButton(TEXT("FlatButton"), FConcertFrontendStyle::Get()->GetBrush("Concert.JoinDefaultSession"), AutoJoinTooltip,
+			ConcertBrowserUtils::MakeIconButton(TEXT("SimpleButton"), FConcertFrontendStyle::Get()->GetBrush("Concert.JoinDefaultSession"), AutoJoinTooltip,
 				TAttribute<bool>(this, &SConcertSessionBrowser::IsAutoJoinButtonEnabled),
 				FOnClicked::CreateSP(this, &SConcertSessionBrowser::OnAutoJoinButtonClicked),
 				TAttribute<EVisibility>::Create([this]() { return !IsCancelAutoJoinButtonEnabled() ? EVisibility::Visible : EVisibility::Collapsed; })) // Default button shown if both auto-join/cancel are disabled.
@@ -2583,7 +2581,7 @@ TSharedRef<SWidget> SConcertSessionBrowser::MakeButtonBar()
 		.AutoWidth()
 		.Padding(0, 0, PaddingBetweenButtons, 0)
 		[
-			ConcertBrowserUtils::MakeIconButton(TEXT("FlatButton"), FConcertFrontendStyle::Get()->GetBrush("Concert.CancelAutoJoin"), CancelAutoJoinTooltip,
+			ConcertBrowserUtils::MakeIconButton(TEXT("SimpleButton"), FConcertFrontendStyle::Get()->GetBrush("Concert.CancelAutoJoin"), CancelAutoJoinTooltip,
 				TAttribute<bool>(this, &SConcertSessionBrowser::IsCancelAutoJoinButtonEnabled),
 				FOnClicked::CreateSP(this, &SConcertSessionBrowser::OnCancelAutoJoinButtonClicked),
 				TAttribute<EVisibility>::Create([this]() { return IsCancelAutoJoinButtonEnabled() ? EVisibility::Visible : EVisibility::Collapsed; }))
@@ -2594,7 +2592,7 @@ TSharedRef<SWidget> SConcertSessionBrowser::MakeButtonBar()
 		.AutoWidth()
 		.Padding(0, 0, PaddingBetweenButtons, 0)
 		[
-			ConcertBrowserUtils::MakeIconButton(TEXT("FlatButton"), FConcertFrontendStyle::Get()->GetBrush("Concert.JoinSession"), LOCTEXT("JoinButtonTooltip", "Join the selected session"),
+			ConcertBrowserUtils::MakeIconButton(TEXT("SimpleButton"), FConcertFrontendStyle::Get()->GetBrush("Concert.JoinSession"), LOCTEXT("JoinButtonTooltip", "Join the selected session"),
 				TAttribute<bool>(this, &SConcertSessionBrowser::IsJoinButtonEnabled),
 				FOnClicked::CreateSP(this, &SConcertSessionBrowser::OnJoinButtonClicked),
 				TAttribute<EVisibility>::Create([this]() { return !IsRestoreButtonEnabled() ? EVisibility::Visible : EVisibility::Collapsed; })) // Default button shown if both join/restore are disabled.
@@ -2605,7 +2603,7 @@ TSharedRef<SWidget> SConcertSessionBrowser::MakeButtonBar()
 		.AutoWidth()
 		.Padding(0, 0, PaddingBetweenButtons, 0)
 		[
-			ConcertBrowserUtils::MakeIconButton(TEXT("FlatButton"), FConcertFrontendStyle::Get()->GetBrush("Concert.RestoreSession"), LOCTEXT("RestoreButtonTooltip", "Restore the selected session"),
+			ConcertBrowserUtils::MakeIconButton(TEXT("SimpleButton"), FConcertFrontendStyle::Get()->GetBrush("Concert.RestoreSession"), LOCTEXT("RestoreButtonTooltip", "Restore the selected session"),
 				TAttribute<bool>(this, &SConcertSessionBrowser::IsRestoreButtonEnabled),
 				FOnClicked::CreateSP(this, &SConcertSessionBrowser::OnRestoreButtonClicked),
 				TAttribute<EVisibility>::Create([this]() { return IsRestoreButtonEnabled() ? EVisibility::Visible : EVisibility::Collapsed; }))
@@ -2616,7 +2614,7 @@ TSharedRef<SWidget> SConcertSessionBrowser::MakeButtonBar()
 		.AutoWidth()
 		.Padding(0, 0, PaddingBetweenButtons, 0)
 		[
-			ConcertBrowserUtils::MakeIconButton(TEXT("FlatButton"), FConcertFrontendStyle::Get()->GetBrush("Concert.ArchiveSession"), LOCTEXT("ArchiveButtonTooltip", "Archive the selected session"),
+			ConcertBrowserUtils::MakeIconButton(TEXT("SimpleButton"), FConcertFrontendStyle::Get()->GetBrush("Concert.ArchiveSession"), LOCTEXT("ArchiveButtonTooltip", "Archive the selected session"),
 				TAttribute<bool>(this, &SConcertSessionBrowser::IsArchiveButtonEnabled),
 				FOnClicked::CreateSP(this, &SConcertSessionBrowser::OnArchiveButtonClicked))
 		]
@@ -2625,7 +2623,7 @@ TSharedRef<SWidget> SConcertSessionBrowser::MakeButtonBar()
 		+SHorizontalBox::Slot()
 		.AutoWidth()
 		[
-			ConcertBrowserUtils::MakeIconButton(TEXT("FlatButton"), FConcertFrontendStyle::Get()->GetBrush("Concert.DeleteSession"), LOCTEXT("DeleteButtonTooltip", "Delete the selected session if permitted"),
+			ConcertBrowserUtils::MakeIconButton(TEXT("SimpleButton"), FConcertFrontendStyle::Get()->GetBrush("Concert.DeleteSession"), LOCTEXT("DeleteButtonTooltip", "Delete the selected session if permitted"),
 				TAttribute<bool>(this, &SConcertSessionBrowser::IsDeleteButtonEnabled),
 				FOnClicked::CreateSP(this, &SConcertSessionBrowser::OnDeleteButtonClicked))
 		];
@@ -3022,7 +3020,7 @@ TSharedRef<SWidget> SConcertSessionBrowser::MakeSessionViewOptionsBar()
 		.AutoWidth()
 		[
 			SNew(SComboButton)
-			.ComboButtonStyle(FEditorStyle::Get(), "GenericFilters.ComboButtonStyle")
+			.ComboButtonStyle(FAppStyle::Get(), "GenericFilters.ComboButtonStyle")
 			.ForegroundColor(FLinearColor::White)
 			.ContentPadding(0)
 			.OnGetMenuContent_Lambda(AddFilterMenu)
@@ -3036,7 +3034,7 @@ TSharedRef<SWidget> SConcertSessionBrowser::MakeSessionViewOptionsBar()
 				.AutoWidth()
 				.VAlign(VAlign_Center)
 				[
-					SNew(SImage).Image(FEditorStyle::GetBrush("GenericViewButton")) // The eye ball image.
+					SNew(SImage).Image(FAppStyle::Get().GetBrush("GenericViewButton")) // The eye ball image.
 				]
 
 				+SHorizontalBox::Slot()
@@ -3215,7 +3213,7 @@ TSharedRef<SWidget> SConcertSessionBrowser::MakeActiveSessionDetails(TSharedPtr<
 			SAssignNew(DetailsArea, SExpandableArea)
 			.BorderBackgroundColor(FLinearColor(0.6f, 0.6f, 0.6f, 1.0f))
 			.BorderImage_Lambda([this]() { return ConcertFrontendUtils::GetExpandableAreaBorderImage(*DetailsArea); })
-			.BodyBorderImage(FEditorStyle::GetBrush("ToolPanel.GroupBorder"))
+			.BodyBorderImage(FAppStyle::Get().GetBrush("ToolPanel.GroupBorder"))
 			.BodyBorderBackgroundColor(FLinearColor::White)
 			.OnAreaExpansionChanged(FOnBooleanValueChanged::CreateLambda(OnDetailsAreaExpansionChanged))
 			.InitiallyCollapsed(!(*bDetailsAreaExpanded))
@@ -3223,7 +3221,7 @@ TSharedRef<SWidget> SConcertSessionBrowser::MakeActiveSessionDetails(TSharedPtr<
 			[
 				SNew(STextBlock)
 				.Text(FText::Format(LOCTEXT("Details", "Details"), FText::FromString(Item->SessionName)))
-				.Font(FEditorStyle::GetFontStyle("DetailsView.CategoryFontStyle"))
+				.Font(FAppStyle::Get().GetFontStyle("DetailsView.CategoryFontStyle"))
 				.ShadowOffset(FVector2D(1.0f, 1.0f))
 			]
 			.BodyContent()
@@ -3249,7 +3247,7 @@ TSharedRef<SWidget> SConcertSessionBrowser::MakeActiveSessionDetails(TSharedPtr<
 			SAssignNew(ClientsArea, SExpandableArea)
 			.BorderBackgroundColor(FLinearColor(0.6f, 0.6f, 0.6f, 1.0f))
 			.BorderImage_Lambda([this]() { return ConcertFrontendUtils::GetExpandableAreaBorderImage(*ClientsArea); })
-			.BodyBorderImage(FEditorStyle::GetBrush("ToolPanel.GroupBorder"))
+			.BodyBorderImage(FAppStyle::Get().GetBrush("ToolPanel.GroupBorder"))
 			.BodyBorderBackgroundColor(FLinearColor::White)
 			.OnAreaExpansionChanged(FOnBooleanValueChanged::CreateLambda(OnClientsAreaExpansionChanged))
 			.InitiallyCollapsed(!(*bClientsAreaExpanded))
@@ -3257,7 +3255,7 @@ TSharedRef<SWidget> SConcertSessionBrowser::MakeActiveSessionDetails(TSharedPtr<
 			[
 				SNew(STextBlock)
 				.Text(LOCTEXT("Clients", "Clients"))
-				.Font(FEditorStyle::GetFontStyle("DetailsView.CategoryFontStyle"))
+				.Font(FAppStyle::Get().GetFontStyle("DetailsView.CategoryFontStyle"))
 				.ShadowOffset(FVector2D(1.0f, 1.0f))
 			]
 			.BodyContent()
@@ -3296,14 +3294,14 @@ TSharedRef<SWidget> SConcertSessionBrowser::MakeArchivedSessionDetails(TSharedPt
 	TSharedRef<SExpandableArea> Widget = SAssignNew(DetailsArea, SExpandableArea)
 	.BorderBackgroundColor(FLinearColor(0.6f, 0.6f, 0.6f, 1.0f))
 	.BorderImage_Lambda([this]() { return ConcertFrontendUtils::GetExpandableAreaBorderImage(*DetailsArea); })
-	.BodyBorderImage(FEditorStyle::GetBrush("ToolPanel.GroupBorder"))
+	.BodyBorderImage(FAppStyle::Get().GetBrush("ToolPanel.GroupBorder"))
 	.BodyBorderBackgroundColor(FLinearColor::White)
 	.InitiallyCollapsed(true)
 	.HeaderContent()
 	[
 		SNew(STextBlock)
 		.Text(FText::Format(LOCTEXT("Details", "Details"), FText::FromString(Item->SessionName)))
-		.Font(FEditorStyle::GetFontStyle("DetailsView.CategoryFontStyle"))
+		.Font(FAppStyle::Get().GetFontStyle("DetailsView.CategoryFontStyle"))
 		.ShadowOffset(FVector2D(1.0f, 1.0f))
 	]
 	.BodyContent()
@@ -3381,7 +3379,7 @@ TSharedRef<ITableRow> SConcertSessionBrowser::OnGenerateClientRowWidget(TSharedP
 		.AutoWidth()
 		[
 			SNew(STextBlock)
-			.Font(FEditorStyle::Get().GetFontStyle("FontAwesome.8"))
+			.Font(FAppStyle::Get().GetFontStyle("FontAwesome.8"))
 			.ColorAndOpacity(Item->ClientInfo.AvatarColor)
 			.Text(FEditorFontGlyphs::Square)
 		]
diff --git a/Engine/Plugins/Developer/Concert/MultiUserClient/Source/MultiUserClient/Private/Widgets/SSessionHistory.cpp b/Engine/Plugins/Developer/Concert/MultiUserClient/Source/MultiUserClient/Private/Widgets/SSessionHistory.cpp
--- a/Engine/Plugins/Developer/Concert/MultiUserClient/Source/MultiUserClient/Private/Widgets/SSessionHistory.cpp
+++ b/Engine/Plugins/Developer/Concert/MultiUserClient/Source/MultiUserClient/Private/Widgets/SSessionHistory.cpp
@@ -8,11 +8,9 @@
 #include "ConcertMessageData.h"
 #include "ConcertMessages.h"
 #include "Editor/Transactor.h"
-#include "EditorStyleSet.h"
 #include "Widgets/Input/SSearchBox.h"
 #include "Widgets/Layout/SSeparator.h"
 #include "SConcertSessionActivities.h"
-#include "EditorStyleSet.h"
 
 #define LOCTEXT_NAMESPACE "SSessionHistory"
 
