Do not store a weak ptr to the session.  Depending on timing this could get reset after the UI is constructed.  We can just refer to the global client to get current session info.

diff --git a/Engine/Plugins/Developer/Concert/MultiUserClient/Source/MultiUserClient/Private/ConcertWorkspaceUI.cpp b/Engine/Plugins/Developer/Concert/MultiUserClient/Source/MultiUserClient/Private/ConcertWorkspaceUI.cpp
--- a/Engine/Plugins/Developer/Concert/MultiUserClient/Source/MultiUserClient/Private/ConcertWorkspaceUI.cpp
+++ b/Engine/Plugins/Developer/Concert/MultiUserClient/Source/MultiUserClient/Private/ConcertWorkspaceUI.cpp
@@ -2,6 +2,7 @@
 
 #include "ConcertWorkspaceUI.h"
 
+#include "IConcertSyncClientModule.h"
 #include "IConcertClientWorkspace.h"
 #include "IConcertSyncClient.h"
 #include "IConcertSession.h"
@@ -362,7 +363,7 @@ public:
 	 */
 	void Construct(const FArguments& InArgs, TSharedPtr<FConcertWorkspaceUI> InWorkspaceFrontend)
 	{
-		WorkspaceFrontend = MoveTemp(InWorkspaceFrontend);	
+		WorkspaceFrontend = MoveTemp(InWorkspaceFrontend);
 		SetVisibility(MakeAttributeSP(this, &SConcertWorkspaceSequencerToolbarExtension::GetVisibility));
 
 		FToolBarBuilder ToolbarBuilder(TSharedPtr<const FUICommandList>(), FMultiBoxCustomization::None, TSharedPtr<FExtender>(), true);
@@ -372,9 +373,9 @@ public:
 			// Toggle playback sync
 			ToolbarBuilder.AddToolBarButton(
 				FUIAction(
-					FExecuteAction::CreateLambda([SyncClient = WorkspaceFrontend->SyncClient]()
+					FExecuteAction::CreateLambda([]()
 					{
-						TSharedPtr<IConcertSyncClient> SyncClientPin = SyncClient.Pin();
+						TSharedPtr<IConcertSyncClient> SyncClientPin = IConcertSyncClientModule::Get().GetClient(TEXT("MultiUser"));
 						if (SyncClientPin && SyncClientPin->GetSequencerManager())
 						{
 							IConcertClientSequencerManager* SequencerManager = SyncClientPin->GetSequencerManager();
@@ -382,9 +383,9 @@ public:
 						}
 					}),
 					FCanExecuteAction(),
-					FIsActionChecked::CreateLambda([SyncClient = WorkspaceFrontend->SyncClient]()
+					FIsActionChecked::CreateLambda([]()
 					{
-						TSharedPtr<IConcertSyncClient> SyncClientPin = SyncClient.Pin();
+						TSharedPtr<IConcertSyncClient> SyncClientPin = IConcertSyncClientModule::Get().GetClient(TEXT("MultiUser"));
 						return SyncClientPin.IsValid() && SyncClientPin->GetSequencerManager() && SyncClientPin->GetSequencerManager()->IsSequencerPlaybackSyncEnabled();
 					})
 				),
@@ -398,9 +399,9 @@ public:
 			// Toggle unrelated timeline sync
 			ToolbarBuilder.AddToolBarButton(
 				FUIAction(
-					FExecuteAction::CreateLambda([SyncClient = WorkspaceFrontend->SyncClient]()
+					FExecuteAction::CreateLambda([]()
 					{
-						TSharedPtr<IConcertSyncClient> SyncClientPin = SyncClient.Pin();
+						TSharedPtr<IConcertSyncClient> SyncClientPin = IConcertSyncClientModule::Get().GetClient(TEXT("MultiUser"));
 						if (SyncClientPin && SyncClientPin->GetSequencerManager())
 						{
 							IConcertClientSequencerManager* SequencerManager = SyncClientPin->GetSequencerManager();
@@ -408,9 +409,9 @@ public:
 						}
 					}),
 					FCanExecuteAction(),
-					FIsActionChecked::CreateLambda([SyncClient = WorkspaceFrontend->SyncClient]()
+					FIsActionChecked::CreateLambda([]()
 					{
-						TSharedPtr<IConcertSyncClient> SyncClientPin = SyncClient.Pin();
+						TSharedPtr<IConcertSyncClient> SyncClientPin = IConcertSyncClientModule::Get().GetClient(TEXT("MultiUser"));
 						return SyncClientPin.IsValid() && SyncClientPin->GetSequencerManager() && SyncClientPin->GetSequencerManager()->IsUnrelatedSequencerTimelineSyncEnabled();
 					})
 				),
@@ -424,9 +425,9 @@ public:
 			// Toggle remote open
 			ToolbarBuilder.AddToolBarButton(
 				FUIAction(
-					FExecuteAction::CreateLambda([SyncClient = WorkspaceFrontend->SyncClient]() 
+					FExecuteAction::CreateLambda([]()
 					{
-						TSharedPtr<IConcertSyncClient> SyncClientPin = SyncClient.Pin();
+						TSharedPtr<IConcertSyncClient> SyncClientPin = IConcertSyncClientModule::Get().GetClient(TEXT("MultiUser"));
 						if (SyncClientPin && SyncClientPin->GetSequencerManager())
 						{
 							IConcertClientSequencerManager* SequencerManager = SyncClientPin->GetSequencerManager();
@@ -434,14 +435,14 @@ public:
 						}
 					}),
 					FCanExecuteAction(),
-					FIsActionChecked::CreateLambda([SyncClient = WorkspaceFrontend->SyncClient]()
+					FIsActionChecked::CreateLambda([]()
 					{
-						TSharedPtr<IConcertSyncClient> SyncClientPin = SyncClient.Pin();
+						TSharedPtr<IConcertSyncClient> SyncClientPin = IConcertSyncClientModule::Get().GetClient(TEXT("MultiUser"));
 						return SyncClientPin.IsValid() && SyncClientPin->GetSequencerManager() && SyncClientPin->GetSequencerManager()->IsSequencerRemoteOpenEnabled();
 					})
 				),
 				NAME_None,
-				LOCTEXT("ToggleRemoteOpenLabel", "Remote Open"), 
+				LOCTEXT("ToggleRemoteOpenLabel", "Remote Open"),
 				LOCTEXT("ToggleRemoteOpenTooltip", "Toggle Multi-User Remote Open. If the option is enabled, opening a sequence will open the same sequence on all users in the Multi-User session that also have this option enabled."),
 				FSlateIcon(FConcertFrontendStyle::GetStyleSetName(), "Concert.Sequencer.SyncSequence", "Concert.Sequencer.SyncSequence.Small"),
 				EUserInterfaceActionType::ToggleButton
