Resolve uninitialized variables in UE5/Main CIS

diff --git a/Engine/Plugins/Developer/Concert/MultiUserClient/Source/MultiUserClientLibrary/Public/MultiUserClientStatics.h b/Engine/Plugins/Developer/Concert/MultiUserClient/Source/MultiUserClientLibrary/Public/MultiUserClientStatics.h
--- a/Engine/Plugins/Developer/Concert/MultiUserClient/Source/MultiUserClientLibrary/Public/MultiUserClientStatics.h
+++ b/Engine/Plugins/Developer/Concert/MultiUserClient/Source/MultiUserClientLibrary/Public/MultiUserClientStatics.h
@@ -25,7 +25,7 @@ struct FMultiUserClientInfo
 
 	/** Holds the color of the user avatar in a session. */
 	UPROPERTY(VisibleAnywhere, BlueprintReadOnly, Category = "Client Info")
-	FLinearColor AvatarColor;
+	FLinearColor AvatarColor = FLinearColor(1.0f, 1.0f, 1.0f, 1.0f);
 
 	/** Holds an array of tags that can be used for grouping and categorizing. */
 	UPROPERTY(VisibleAnywhere, BlueprintReadOnly, AdvancedDisplay, Category = "Client Info")
@@ -93,7 +93,7 @@ struct FMultiUserClientConfig
 	FString DefaultSessionToRestore;
 
 	UPROPERTY(BlueprintReadWrite, Category = "Source Control Settings")
-	EMultiUserSourceValidationMode ValidationMode;
+	EMultiUserSourceValidationMode ValidationMode = EMultiUserSourceValidationMode::Hard;
 };
 
 /** Connection status for Multi-User client sessions */
