Add ValidationMode support case when p4 command fails

diff --git a/Engine/Plugins/Developer/Concert/MultiUserClient/Source/MultiUserClient/Private/MultiUserClientModule.cpp b/Engine/Plugins/Developer/Concert/MultiUserClient/Source/MultiUserClient/Private/MultiUserClientModule.cpp
--- a/Engine/Plugins/Developer/Concert/MultiUserClient/Source/MultiUserClient/Private/MultiUserClientModule.cpp
+++ b/Engine/Plugins/Developer/Concert/MultiUserClient/Source/MultiUserClient/Private/MultiUserClientModule.cpp
@@ -180,7 +180,7 @@ protected:
 		FConcertConnectionError Error;
 	};
 
-	static EConcertResponseCode GetValidationModeResultOnFailure(const UConcertClientConfig* InClientConfig)
+	static EConcertResponseCode GetValidationModeResultOnFailure(const UConcertClientConfig* InClientConfig, bool bDoNotUseSoftAuto = false)
 	{
 		switch (InClientConfig->SourceControlSettings.ValidationMode)
 		{
@@ -190,6 +190,10 @@ protected:
 			// InvalidRequest is used here to signify a prompt
 			return EConcertResponseCode::InvalidRequest;
 		case EConcertSourceValidationMode::SoftAutoProceed:
+			if (bDoNotUseSoftAuto)
+			{
+				return EConcertResponseCode::InvalidRequest;
+			}
 			return EConcertResponseCode::Success;
 		}
 		return EConcertResponseCode::Failed;
@@ -310,9 +314,12 @@ private:
 			InSharedState->Error.ErrorText = LOCTEXT("ValidatingWorkspace_Canceled", "The workspace validation request was canceled.");
 			break;
 		default:
-			InSharedState->Result = EConcertResponseCode::Failed;
+			InSharedState->Result = GetValidationModeResultOnFailure(InClientConfig, true);
+			InSharedState->PromptText = LOCTEXT("ValidatingWorkspace_SCContinue", "Continue");
 			InSharedState->Error.ErrorCode = MultiUserClientUtil::SourceControlValidationGenericErrorCode;
-			InSharedState->Error.ErrorText = LOCTEXT("ValidatingWorkspace_Failed", "The workspace validation request failed. Please check your source control settings.");
+			InSharedState->Error.ErrorText = InClientConfig->SourceControlSettings.ValidationMode == EConcertSourceValidationMode::Hard ?
+				LOCTEXT("ValidatingWorkspace_Failed", "The workspace validation request failed. Please check your source control settings.") :
+				LOCTEXT("ValidatingWorkspace_FailedSoft", "This workspace validation request failed. Please check your source control settings. Local changes will not available to other users.");
 			break;
 		}
 	}
