Extract session table to re-usabe UI widget

diff --git a/Engine/Plugins/Developer/Concert/ConcertUI/ConcertSharedSlate/Source/ConcertSharedSlate/Private/SessionBrowser/SConcertSessionBrowser.cpp b/Engine/Plugins/Developer/Concert/ConcertUI/ConcertSharedSlate/Source/ConcertSharedSlate/Private/SessionBrowser/SConcertSessionBrowser.cpp
new file mode 100644
--- /dev/null
+++ b/Engine/Plugins/Developer/Concert/ConcertUI/ConcertSharedSlate/Source/ConcertSharedSlate/Private/SessionBrowser/SConcertSessionBrowser.cpp
@@ -0,0 +1,1047 @@
+﻿// Copyright Epic Games, Inc. All Rights Reserved.
+
+#include "SessionBrowser/SConcertSessionBrowser.h"
+
+#include "ConcertFrontendStyle.h"
+#include "ConcertFrontendUtils.h"
+#include "SessionBrowser/ConcertBrowserUtils.h"
+#include "SessionBrowser/ConcertSessionBrowserSettings.h"
+#include "SessionBrowser/ConcertSessionItem.h"
+#include "SessionBrowser/IConcertSessionBrowserController.h"
+#include "SessionBrowser/SNewSessionRow.h"
+#include "SessionBrowser/SSaveRestoreSessionRow.h"
+#include "SessionBrowser/SSessionRow.h"
+
+#include "Styling/AppStyle.h"
+#include "Framework/MultiBox/MultiBoxBuilder.h"
+#include "Internationalization/Regex.h"
+#include "Misc/AsyncTaskNotification.h"
+#include "Misc/MessageDialog.h"
+#include "Misc/TextFilter.h"
+#include "Widgets/Images/SImage.h"
+#include "Widgets/Input/SEditableTextBox.h"
+#include "Widgets/Input/SSearchBox.h"
+#include "Widgets/Layout/SBorder.h"
+#include "Widgets/Layout/SSpacer.h"
+#include "Widgets/Layout/SSeparator.h"
+#include "Widgets/Layout/SUniformGridPanel.h"
+#include "Widgets/Text/STextBlock.h"
+#include "Widgets/Views/SListView.h"
+#include "Widgets/Views/STableRow.h"
+
+#define LOCTEXT_NAMESPACE "SConcertBrowser"
+
+const FName SConcertSessionBrowser::ControlButtonExtensionHooks::BeforeSeparator("BeforeSeparator");
+const FName SConcertSessionBrowser::ControlButtonExtensionHooks::Separator("Separator");
+const FName SConcertSessionBrowser::ControlButtonExtensionHooks::AfterSeparator("AfterSeparator");
+
+const FName SConcertSessionBrowser::SessionContextMenuExtensionHooks::ManageSession("ManageSession");
+
+void SConcertSessionBrowser::Construct(const FArguments& InArgs, TSharedRef<IConcertSessionBrowserController> InController, TSharedPtr<FText> InSearchText)
+{
+	Controller = InController;
+
+	// Reload the persistent settings, such as the filters.
+	PersistentSettings = GetMutableDefault<UConcertSessionBrowserSettings>();
+
+	DefaultSessionName = InArgs._DefaultSessionName;
+	DefaultServerUrl = InArgs._DefaultServerURL;
+
+	ExtendSessionContextMenu = InArgs._ExtendSessionContextMenu;
+	OnSessionClicked = InArgs._OnSessionClicked;
+	OnSessionDoubleClicked = InArgs._OnSessionDoubleClicked;
+	OnRequestedDeleteSession = InArgs._OnRequestedDeleteSession;
+
+	// Setup search filter.
+	SearchedText = InSearchText; // Reload a previous search text (in any). Useful to remember searched text between join/leave sessions, but not persistent if the tab is closed.
+	SearchTextFilter = MakeShared<TTextFilter<const FConcertSessionItem&>>(TTextFilter<const FConcertSessionItem&>::FItemToStringArray::CreateSP(this, &SConcertSessionBrowser::PopulateSearchStrings));
+	SearchTextFilter->OnChanged().AddSP(this, &SConcertSessionBrowser::RefreshSessionList);
+
+	ChildSlot
+	[
+		MakeBrowserContent(InArgs)
+	];
+
+	if (!SearchedText->IsEmpty())
+	{
+		SearchBox->SetText(*SearchedText); // This trigger the chain of actions to apply the search filter.
+	}
+}
+
+TSharedRef<SWidget> SConcertSessionBrowser::MakeBrowserContent(const FArguments& InArgs)
+{
+	TSharedRef<SWidget> SessionTable = MakeSessionTableView();
+	return SNew(SBorder)
+		.BorderImage(FAppStyle::Get().GetBrush("ToolPanel.GroupBorder"))
+		.Padding(FMargin(1.0f, 2.0f))
+		[
+			SNew(SVerticalBox)
+
+			+SVerticalBox::Slot()
+			.AutoHeight()
+			[
+				MakeControlBar(InArgs)
+			]
+
+			// Session list.
+			+SVerticalBox::Slot()
+			.FillHeight(1.0f)
+			.Padding(1.0f, 2.0f)
+			[
+				InArgs._ExtendSessionTable.IsBound() ? InArgs._ExtendSessionTable.Execute(SessionTable) : SessionTable
+			]
+
+			+SVerticalBox::Slot()
+			.AutoHeight()
+			.Padding(FMargin(2.0f, 0.0f))
+			[
+				SNew(SSeparator)
+			]
+
+			// Session Count/View options filter.
+			+SVerticalBox::Slot()
+			.AutoHeight()
+			.Padding(FMargin(2.0f, 0.0f))
+			[
+				MakeSessionViewOptionsBar()
+			]
+		];
+}
+
+void SConcertSessionBrowser::RefreshSessionList()
+{
+	// Remember the selected instances (if any).
+	TArray<TSharedPtr<FConcertSessionItem>> SelectedItems = SessionsView->GetSelectedItems();
+	TArray<TSharedPtr<FConcertSessionItem>> ReselectedItems;
+	TSharedPtr<FConcertSessionItem> NewEditableRowParent;
+
+	// Predicate returning true if the specified item should be re-selected.
+	auto IsSelected = [&SelectedItems](const FConcertSessionItem& Item)
+	{
+		return SelectedItems.ContainsByPredicate([Item](const TSharedPtr<FConcertSessionItem>& Visited) { return *Visited == Item; });
+	};
+
+	// Matches the object instances before the update to the new instance after the update.
+	auto ReconcileObjectInstances = [&](TSharedPtr<FConcertSessionItem> NewItem)
+	{
+		if (IsSelected(*NewItem))
+		{
+			ReselectedItems.Add(NewItem);
+		}
+		else if (EditableSessionRowParent.IsValid() && !NewEditableRowParent.IsValid() && *EditableSessionRowParent == *NewItem)
+		{
+			NewEditableRowParent = NewItem;
+		}
+	};
+
+	// Clear sessions.
+	Sessions.Reset();
+
+	// Populate the live sessions.
+	for (const TSharedPtr<IConcertSessionBrowserController::FActiveSessionInfo>& ActiveSession : GetController()->GetActiveSessions())
+	{
+		FConcertSessionItem NewItem(FConcertSessionItem::EType::ActiveSession, ActiveSession->SessionInfo.SessionName, ActiveSession->SessionInfo.SessionId, ActiveSession->ServerInfo.ServerName, ActiveSession->ServerInfo.AdminEndpointId, ActiveSession->ServerInfo.ServerFlags);
+		if (!IsFilteredOut(NewItem))
+		{
+			Sessions.Emplace(MakeShared<FConcertSessionItem>(MoveTemp(NewItem)));
+			ReconcileObjectInstances(Sessions.Last());
+		}
+	}
+
+	// Populate the archived.
+	for (const TSharedPtr<IConcertSessionBrowserController::FArchivedSessionInfo>& ArchivedSession : GetController()->GetArchivedSessions())
+	{
+		FConcertSessionItem NewItem(FConcertSessionItem::EType::ArchivedSession, ArchivedSession->SessionInfo.SessionName, ArchivedSession->SessionInfo.SessionId, ArchivedSession->ServerInfo.ServerName, ArchivedSession->ServerInfo.AdminEndpointId, ArchivedSession->ServerInfo.ServerFlags);
+		if (!IsFilteredOut(NewItem))
+		{
+			Sessions.Emplace(MakeShared<FConcertSessionItem>(MoveTemp(NewItem)));
+			ReconcileObjectInstances(Sessions.Last());
+		}
+	}
+
+	// Restore the editable row state. (SortSessionList() below will ensure the parent/child relationship)
+	EditableSessionRowParent = NewEditableRowParent;
+	if (EditableSessionRow.IsValid())
+	{
+		if (EditableSessionRow->Type == FConcertSessionItem::EType::NewSession)
+		{
+			Sessions.Insert(EditableSessionRow, 0); // Always put 'new session' row at the top.
+		}
+		else if (EditableSessionRowParent.IsValid())
+		{
+			Sessions.Add(EditableSessionRow); // SortSessionList() called below will ensure the correct parent/child order.
+		}
+	}
+
+	// Restore previous selection.
+	if (ReselectedItems.Num())
+	{
+		for (const TSharedPtr<FConcertSessionItem>& Item : ReselectedItems)
+		{
+			SessionsView->SetItemSelection(Item, true);
+		}
+	}
+
+	SortSessionList();
+	SessionsView->RequestListRefresh();
+}
+
+TSharedPtr<IConcertSessionBrowserController> SConcertSessionBrowser::GetController() const
+{
+	TSharedPtr<IConcertSessionBrowserController> Result = Controller.Pin();
+	check(Result);
+	return Result;
+}
+
+void SConcertSessionBrowser::OnSearchTextChanged(const FText& InFilterText)
+{
+	SearchTextFilter->SetRawFilterText(InFilterText);
+	SearchBox->SetError( SearchTextFilter->GetFilterErrorText() );
+	*SearchedText = InFilterText;
+
+	bRefreshSessionFilter = true;
+}
+
+void SConcertSessionBrowser::OnSearchTextCommitted(const FText& InFilterText, ETextCommit::Type CommitType)
+{
+	if (!InFilterText.EqualTo(*SearchedText))
+	{
+		OnSearchTextChanged(InFilterText);
+	}
+}
+
+void SConcertSessionBrowser::PopulateSearchStrings(const FConcertSessionItem& Item, TArray<FString>& OutSearchStrings) const
+{
+	OutSearchStrings.Add(Item.ServerName);
+	OutSearchStrings.Add(Item.SessionName);
+}
+
+bool SConcertSessionBrowser::IsFilteredOut(const FConcertSessionItem& Item) const
+{
+	bool bIsDefaultServer = LastDefaultServerUrl.IsEmpty() || Item.ServerName == LastDefaultServerUrl;
+
+	return (!PersistentSettings->bShowActiveSessions && (Item.Type == FConcertSessionItem::EType::ActiveSession || Item.Type == FConcertSessionItem::EType::SaveSession)) ||
+	       (!PersistentSettings->bShowArchivedSessions && (Item.Type == FConcertSessionItem::EType::ArchivedSession || Item.Type == FConcertSessionItem::EType::RestoreSession)) ||
+	       (PersistentSettings->bShowDefaultServerSessionsOnly && !bIsDefaultServer) ||
+	       !SearchTextFilter->PassesFilter(Item);
+}
+
+FText SConcertSessionBrowser::HighlightSearchText() const
+{
+	return *SearchedText;
+}
+
+TSharedRef<SWidget> SConcertSessionBrowser::MakeControlBar(const FArguments& InArgs)
+{
+	return SNew(SHorizontalBox)
+		// The New/Join/Restore/Delete/Archive buttons
+		+SHorizontalBox::Slot()
+		.AutoWidth()
+		[
+			MakeButtonBar(InArgs)
+		]
+
+		// The search text.
+		+SHorizontalBox::Slot()
+		.FillWidth(1.0)
+		.Padding(4.0f, 5.0f, 8.0f, 5.0f)
+		[
+			SAssignNew(SearchBox, SSearchBox)
+			.HintText(LOCTEXT("SearchHint", "Search Session"))
+			.OnTextChanged(this, &SConcertSessionBrowser::OnSearchTextChanged)
+			.OnTextCommitted(this, &SConcertSessionBrowser::OnSearchTextCommitted)
+			.DelayChangeNotificationsWhileTyping(true)
+		]
+
+		// Optional: everything to the right of the search bar, e.g. user name and settings combo button
+		+SHorizontalBox::Slot()
+		.VAlign(VAlign_Center)
+		.AutoWidth()
+		[
+			InArgs._RightOfSearchBar.Widget
+		];
+}
+
+TSharedRef<SWidget> SConcertSessionBrowser::MakeButtonBar(const FArguments& InArgs)
+{
+	TSharedRef<FExtender> Extender = MakeShared<FExtender>();
+	InArgs._ExtendControllButtons.ExecuteIfBound(Extender.Get());
+	FToolBarBuilder RowBuilder(nullptr, FMultiBoxCustomization::None, Extender);
+	
+	RowBuilder.BeginSection(ControlButtonExtensionHooks::BeforeSeparator);
+	// New Session
+	RowBuilder.AddWidget(
+		ConcertBrowserUtils::MakeIconButton(
+			TEXT("SimpleButton"),
+			FConcertFrontendStyle::Get()->GetBrush("Concert.NewSession"),
+			LOCTEXT("NewButtonTooltip", "Create a new session"),
+			TAttribute<bool>(this, &SConcertSessionBrowser::IsNewButtonEnabledInternal),
+			FOnClicked::CreateSP(this, &SConcertSessionBrowser::OnNewButtonClicked)
+		)
+	);
+	RowBuilder.EndSection();
+	
+	RowBuilder.AddSeparator(ControlButtonExtensionHooks::Separator);
+	RowBuilder.BeginSection(ControlButtonExtensionHooks::AfterSeparator);
+	
+	// Restore (Share the same slot as Join)
+	RowBuilder.AddWidget(
+		ConcertBrowserUtils::MakeIconButton(
+			TEXT("SimpleButton"),
+			FConcertFrontendStyle::Get()->GetBrush("Concert.RestoreSession"),
+			LOCTEXT("RestoreButtonTooltip", "Restore the selected session"),
+			TAttribute<bool>(this, &SConcertSessionBrowser::IsRestoreButtonEnabledInternal),
+			FOnClicked::CreateSP(this, &SConcertSessionBrowser::OnRestoreButtonClicked),
+			TAttribute<EVisibility>::Create([this]() { return IsRestoreButtonEnabledInternal() ? EVisibility::Visible : EVisibility::Collapsed; })
+			)
+	);
+	// Archive.
+	RowBuilder.AddWidget(
+		ConcertBrowserUtils::MakeIconButton(TEXT("SimpleButton"), FConcertFrontendStyle::Get()->GetBrush("Concert.ArchiveSession"), LOCTEXT("ArchiveButtonTooltip", "Archive the selected session"),
+			TAttribute<bool>(this, &SConcertSessionBrowser::IsArchiveButtonEnabledInternal),
+			FOnClicked::CreateSP(this, &SConcertSessionBrowser::OnArchiveButtonClicked))
+	);
+	// Delete.
+	RowBuilder.AddWidget(
+		ConcertBrowserUtils::MakeIconButton(TEXT("SimpleButton"), FConcertFrontendStyle::Get()->GetBrush("Concert.DeleteSession"), LOCTEXT("DeleteButtonTooltip", "Delete the selected session if permitted"),
+			TAttribute<bool>(this, &SConcertSessionBrowser::IsDeleteButtonEnabledInternal),
+			FOnClicked::CreateSP(this, &SConcertSessionBrowser::OnDeleteButtonClicked))
+	);
+	RowBuilder.EndSection();
+
+	return RowBuilder.MakeWidget();
+}
+
+TSharedRef<SWidget> SConcertSessionBrowser::MakeSessionTableView()
+{
+	PrimarySortedColumn = ConcertBrowserUtils::IconColName;
+	PrimarySortMode = EColumnSortMode::Ascending;
+	SecondarySortedColumn = ConcertBrowserUtils::SessionColName;
+	SecondarySortMode = EColumnSortMode::Ascending;
+
+	return SAssignNew(SessionsView, SListView<TSharedPtr<FConcertSessionItem>>)
+		.SelectionMode(ESelectionMode::Single)
+		.ListItemsSource(&Sessions)
+		.OnGenerateRow(this, &SConcertSessionBrowser::OnGenerateSessionRowWidget)
+		.SelectionMode(ESelectionMode::Single)
+		.OnSelectionChanged(this, &SConcertSessionBrowser::OnSessionSelectionChanged)
+		.OnContextMenuOpening(this, &SConcertSessionBrowser::MakeContextualMenu)
+		.HeaderRow(
+			SNew(SHeaderRow)
+
+			+SHeaderRow::Column(ConcertBrowserUtils::IconColName)
+			.DefaultLabel(FText::GetEmpty())
+			.SortPriority(this, &SConcertSessionBrowser::GetColumnSortPriority, ConcertBrowserUtils::IconColName)
+			.SortMode(this, &SConcertSessionBrowser::GetColumnSortMode, ConcertBrowserUtils::IconColName)
+			.OnSort(this, &SConcertSessionBrowser::OnColumnSortModeChanged)
+			.FixedWidth(20)
+
+			+SHeaderRow::Column(ConcertBrowserUtils::SessionColName)
+			.DefaultLabel(LOCTEXT("SessioName", "Session"))
+			.SortPriority(this, &SConcertSessionBrowser::GetColumnSortPriority, ConcertBrowserUtils::SessionColName)
+			.SortMode(this, &SConcertSessionBrowser::GetColumnSortMode, ConcertBrowserUtils::SessionColName)
+			.OnSort(this, &SConcertSessionBrowser::OnColumnSortModeChanged)
+
+			+SHeaderRow::Column(ConcertBrowserUtils::ServerColName)
+			.DefaultLabel(LOCTEXT("Server", "Server"))
+			.SortPriority(this, &SConcertSessionBrowser::GetColumnSortPriority, ConcertBrowserUtils::ServerColName)
+			.SortMode(this, &SConcertSessionBrowser::GetColumnSortMode, ConcertBrowserUtils::ServerColName)
+			.OnSort(this, &SConcertSessionBrowser::OnColumnSortModeChanged));
+}
+
+EColumnSortMode::Type SConcertSessionBrowser::GetColumnSortMode(const FName ColumnId) const
+{
+	if (ColumnId == PrimarySortedColumn)
+	{
+		return PrimarySortMode;
+	}
+	else if (ColumnId == SecondarySortedColumn)
+	{
+		return SecondarySortMode;
+	}
+
+	return EColumnSortMode::None;
+}
+
+EColumnSortPriority::Type SConcertSessionBrowser::GetColumnSortPriority(const FName ColumnId) const
+{
+	if (ColumnId == PrimarySortedColumn)
+	{
+		return EColumnSortPriority::Primary;
+	}
+	else if (ColumnId == SecondarySortedColumn)
+	{
+		return EColumnSortPriority::Secondary;
+	}
+
+	return EColumnSortPriority::Max; // No specific priority.
+}
+
+void SConcertSessionBrowser::OnColumnSortModeChanged(const EColumnSortPriority::Type SortPriority, const FName& ColumnId, const EColumnSortMode::Type InSortMode)
+{
+	if (SortPriority == EColumnSortPriority::Primary)
+	{
+		PrimarySortedColumn = ColumnId;
+		PrimarySortMode = InSortMode;
+
+		if (ColumnId == SecondarySortedColumn) // Cannot be primary and secondary at the same time.
+		{
+			SecondarySortedColumn = FName();
+			SecondarySortMode = EColumnSortMode::None;
+		}
+	}
+	else if (SortPriority == EColumnSortPriority::Secondary)
+	{
+		SecondarySortedColumn = ColumnId;
+		SecondarySortMode = InSortMode;
+	}
+
+	SortSessionList();
+	SessionsView->RequestListRefresh();
+}
+
+void SConcertSessionBrowser::SortSessionList()
+{
+	check(!PrimarySortedColumn.IsNone()); // Should always have a primary column. User cannot clear this one.
+
+	auto Compare = [this](const TSharedPtr<FConcertSessionItem>& Lhs, const TSharedPtr<FConcertSessionItem>& Rhs, const FName& ColName, EColumnSortMode::Type SortMode)
+	{
+		if (Lhs->Type == FConcertSessionItem::EType::NewSession) // Always keep editable 'new session' row at the top.
+		{
+			return true;
+		}
+		else if (Rhs->Type == FConcertSessionItem::EType::NewSession)
+		{
+			return false;
+		}
+
+		if (ColName == ConcertBrowserUtils::IconColName)
+		{
+			return SortMode == EColumnSortMode::Ascending ? Lhs->Type < Rhs->Type : Lhs->Type > Rhs->Type;
+		}
+		else if (ColName == ConcertBrowserUtils::SessionColName)
+		{
+			return SortMode == EColumnSortMode::Ascending ? Lhs->SessionName < Rhs->SessionName : Lhs->SessionName > Rhs->SessionName;
+		}
+		else
+		{
+			return SortMode == EColumnSortMode::Ascending ? Lhs->ServerName < Rhs->ServerName : Lhs->ServerName > Rhs->ServerName;
+		}
+	};
+
+	Sessions.Sort([&](const TSharedPtr<FConcertSessionItem>& Lhs, const TSharedPtr<FConcertSessionItem>& Rhs)
+	{
+		if (Compare(Lhs, Rhs, PrimarySortedColumn, PrimarySortMode))
+		{
+			return true; // Lhs must be before Rhs based on the primary sort order.
+		}
+		else if (Compare(Rhs, Lhs, PrimarySortedColumn, PrimarySortMode)) // Invert operands order (goal is to check if operands are equal or not)
+		{
+			return false; // Rhs must be before Lhs based on the primary sort.
+		}
+		else // Lhs == Rhs on the primary column, need to order accoding the secondary column if one is set.
+		{
+			return SecondarySortedColumn.IsNone() ? false : Compare(Lhs, Rhs, SecondarySortedColumn, SecondarySortMode);
+		}
+	});
+
+	EnsureEditableParentChildOrder();
+}
+
+void SConcertSessionBrowser::EnsureEditableParentChildOrder()
+{
+	// This is for Archiving or Restoring a session. We keep the editable row below the session to archive or restore and visually link them with small wires in UI.
+	if (EditableSessionRowParent.IsValid())
+	{
+		check(EditableSessionRow.IsValid());
+		Sessions.Remove(EditableSessionRow);
+
+		int32 ParentIndex = Sessions.IndexOfByKey(EditableSessionRowParent);
+		if (ParentIndex != INDEX_NONE)
+		{
+			Sessions.Insert(EditableSessionRow, ParentIndex + 1); // Insert the 'child' below its parent.
+		}
+	}
+}
+
+TSharedRef<ITableRow> SConcertSessionBrowser::OnGenerateSessionRowWidget(TSharedPtr<FConcertSessionItem> Item, const TSharedRef<STableViewBase>& OwnerTable)
+{
+	switch (Item->Type)
+	{
+		case FConcertSessionItem::EType::ActiveSession:
+			return MakeActiveSessionRowWidget(Item, OwnerTable);
+
+		case FConcertSessionItem::EType::ArchivedSession:
+			return MakeArchivedSessionRowWidget(Item, OwnerTable);
+
+		case FConcertSessionItem::EType::NewSession:
+			return MakeNewSessionRowWidget(Item, OwnerTable);
+
+		case FConcertSessionItem::EType::RestoreSession:
+			return MakeRestoreSessionRowWidget(Item, OwnerTable);
+
+		default:
+			check(Item->Type == FConcertSessionItem::EType::SaveSession);
+			return MakeSaveSessionRowWidget(Item, OwnerTable);
+	}
+}
+
+TSharedRef<ITableRow> SConcertSessionBrowser::MakeActiveSessionRowWidget(const TSharedPtr<FConcertSessionItem>& ActiveItem, const TSharedRef<STableViewBase>& OwnerTable)
+{
+	const FConcertSessionInfo* SessionInfo = GetController()->GetActiveSessionInfo(ActiveItem->ServerAdminEndpointId, ActiveItem->SessionId);
+
+	// Add an 'Active Session' row. Clicking the row icon joins the session.
+	return SNew(SSessionRow, ActiveItem, OwnerTable)
+		.OnDoubleClickFunc([this](const TSharedPtr<FConcertSessionItem>& Item) { OnSessionDoubleClicked.ExecuteIfBound(Item); })
+		.OnRenameFunc([this](const TSharedPtr<FConcertSessionItem>& Item, const FString& NewName) { RequestRenameSession(Item, NewName); })
+		.IsDefaultSession([this](TSharedPtr<FConcertSessionItem> ItemPin)
+		{
+			return ItemPin->Type == FConcertSessionItem::EType::ActiveSession
+				&& DefaultSessionName.IsSet() && ItemPin->SessionName == DefaultSessionName.Get()
+				&& DefaultServerUrl.IsSet() && ItemPin->ServerName == DefaultServerUrl.Get();
+		})
+		.HighlightText(this, &SConcertSessionBrowser::HighlightSearchText)
+		.ToolTipText(SessionInfo ? SessionInfo->ToDisplayString() : FText::GetEmpty())
+		.IsSelected_Lambda([this, ActiveItem]() { return SessionsView->GetSelectedItems().Num() == 1 && SessionsView->GetSelectedItems()[0] == ActiveItem; });
+}
+
+TSharedRef<ITableRow> SConcertSessionBrowser::MakeArchivedSessionRowWidget(const TSharedPtr<FConcertSessionItem>& ArchivedItem, const TSharedRef<STableViewBase>& OwnerTable)
+{
+	const FConcertSessionInfo* SessionInfo = GetController()->GetArchivedSessionInfo(ArchivedItem->ServerAdminEndpointId, ArchivedItem->SessionId);
+
+	// Add an 'Archived Session' row. Clicking the row icon adds a 'Restore as' row to the table.
+	return SNew(SSessionRow, ArchivedItem, OwnerTable)
+		.OnDoubleClickFunc([this](const TSharedPtr<FConcertSessionItem>& Item) { InsertRestoreSessionAsEditableRowInternal(Item); })
+		.OnRenameFunc([this](const TSharedPtr<FConcertSessionItem>& Item, const FString& NewName) { RequestRenameSession(Item, NewName); })
+		.IsDefaultSession([this](TSharedPtr<FConcertSessionItem> ItemPin)
+		{
+			return ItemPin->Type == FConcertSessionItem::EType::ActiveSession
+				&& DefaultSessionName.IsSet() && ItemPin->SessionName == DefaultSessionName.Get()
+				&& DefaultServerUrl.IsSet() && ItemPin->ServerName == DefaultServerUrl.Get();
+		})
+		.HighlightText(this, &SConcertSessionBrowser::HighlightSearchText)
+		.ToolTipText(SessionInfo ? SessionInfo->ToDisplayString() : FText::GetEmpty())
+		.IsSelected_Lambda([this, ArchivedItem]() { return SessionsView->GetSelectedItems().Num() == 1 && SessionsView->GetSelectedItems()[0] == ArchivedItem; });
+}
+
+TSharedRef<ITableRow> SConcertSessionBrowser::MakeNewSessionRowWidget(const TSharedPtr<FConcertSessionItem>& NewItem, const TSharedRef<STableViewBase>& OwnerTable)
+{
+	// Add an editable 'New Session' row in the table to let user pick a name and a server.
+	TAttribute<FString> DefaultServerUrlArg = DefaultServerUrl.IsSet()
+		? TAttribute<FString>::Create([this](){ return DefaultServerUrl.Get(); }) : TAttribute<FString>();
+	return SNew(SNewSessionRow, NewItem, OwnerTable)
+		.GetServerFunc([this]() { return GetController()->GetServers(); }) // Let the row pull the servers for the combo box.
+		.OnAcceptFunc([this](const TSharedPtr<FConcertSessionItem>& Item) { RequestCreateSession(Item); }) // Accepting creates the session.
+		.OnDeclineFunc([this](const TSharedPtr<FConcertSessionItem>& Item) { RemoveSessionRow(Item); })  // Declining removes the editable 'new session' row from the view.
+		.HighlightText(this, &SConcertSessionBrowser::HighlightSearchText)
+		.DefaultServerURL(DefaultServerUrlArg);
+}
+
+TSharedRef<ITableRow> SConcertSessionBrowser::MakeSaveSessionRowWidget(const TSharedPtr<FConcertSessionItem>& SaveItem, const TSharedRef<STableViewBase>& OwnerTable)
+{
+	// Add an editable 'Save Session' row in the table to let the user enter an archive name.
+	return SNew(SSaveRestoreSessionRow, SaveItem, OwnerTable)
+		.OnAcceptFunc([this](const TSharedPtr<FConcertSessionItem>& Item, const FString& ArchiveName) { RequestArchiveSession(Item, ArchiveName); }) // Accepting archive the session.
+		.OnDeclineFunc([this](const TSharedPtr<FConcertSessionItem>& Item) { RemoveSessionRow(Item); }) // Declining removes the editable 'save session as' row from the view.
+		.HighlightText(this, &SConcertSessionBrowser::HighlightSearchText);
+}
+
+TSharedRef<ITableRow> SConcertSessionBrowser::MakeRestoreSessionRowWidget(const TSharedPtr<FConcertSessionItem>& RestoreItem, const TSharedRef<STableViewBase>& OwnerTable)
+{
+	// Add an editable 'Restore Session' row in the table to let the user enter a session name.
+	return SNew(SSaveRestoreSessionRow, RestoreItem, OwnerTable)
+		.OnAcceptFunc([this](const TSharedPtr<FConcertSessionItem>& Item, const FString& SessionName) { RequestRestoreSession(Item, SessionName); }) // Accepting restores the session.
+		.OnDeclineFunc([this](const TSharedPtr<FConcertSessionItem>& Item) { RemoveSessionRow(Item); }) // Declining removes the editable 'restore session as' row from the view.
+		.HighlightText(this, &SConcertSessionBrowser::HighlightSearchText);
+}
+
+void SConcertSessionBrowser::InsertNewSessionEditableRowInternal()
+{
+	// Insert a 'new session' editable row.
+	InsertEditableSessionRow(MakeShared<FConcertSessionItem>(FConcertSessionItem::EType::NewSession, FString(), FGuid(), FString(), FGuid(), EConcertServerFlags::None), nullptr);
+}
+
+void SConcertSessionBrowser::InsertRestoreSessionAsEditableRowInternal(const TSharedPtr<FConcertSessionItem>& ArchivedItem)
+{
+	// Insert the 'restore session as ' editable row just below the 'archived' item to restore.
+	InsertEditableSessionRow(MakeShared<FConcertSessionItem>(FConcertSessionItem::EType::RestoreSession, ArchivedItem->SessionName, ArchivedItem->SessionId, ArchivedItem->ServerName, ArchivedItem->ServerAdminEndpointId, ArchivedItem->ServerFlags), ArchivedItem);
+}
+
+void SConcertSessionBrowser::InsertArchiveSessionAsEditableRow(const TSharedPtr<FConcertSessionItem>& LiveItem)
+{
+	// Insert the 'save session as' editable row just below the 'active' item to save.
+	InsertEditableSessionRow(MakeShared<FConcertSessionItem>(FConcertSessionItem::EType::SaveSession, LiveItem->SessionName, LiveItem->SessionId, LiveItem->ServerName, LiveItem->ServerAdminEndpointId, LiveItem->ServerFlags), LiveItem);
+}
+
+void SConcertSessionBrowser::InsertEditableSessionRow(TSharedPtr<FConcertSessionItem> EditableItem, TSharedPtr<FConcertSessionItem> ParentItem)
+{
+	// Insert the new row below its parent (if any).
+	int32 ParentIndex = Sessions.IndexOfByKey(ParentItem);
+	Sessions.Insert(EditableItem, ParentIndex != INDEX_NONE ? ParentIndex + 1 : 0);
+
+	// Ensure there is only one editable row at the time, removing the row being edited (if any).
+	Sessions.Remove(EditableSessionRow);
+	EditableSessionRow = EditableItem;
+	EditableSessionRowParent = ParentItem;
+
+	// Ensure the editable row added is selected and visible.
+	SessionsView->SetSelection(EditableItem);
+	SessionsView->RequestListRefresh();
+
+	// NOTE: Ideally, I would only use RequestScrollIntoView() to scroll the new item into view, but it did not work. If an item was added into an hidden part,
+	//       it was not always scrolled correctly into view. RequestNavigateToItem() worked much better, except when inserting the very first row in the list, in
+	//       such case calling the function would give the focus to the list view (showing a white dashed line around it).
+	if (ParentIndex == INDEX_NONE)
+	{
+		SessionsView->ScrollToTop(); // Item is inserted at 0. (New session)
+	}
+	else
+	{
+		SessionsView->RequestNavigateToItem(EditableItem);
+	}
+}
+
+void SConcertSessionBrowser::RemoveSessionRow(const TSharedPtr<FConcertSessionItem>& Item)
+{
+	Sessions.Remove(Item);
+
+	// Don't keep the editable row if its 'parent' is removed. (if the session to restore or archive gets deleted in the meantime)
+	if (Item == EditableSessionRowParent)
+	{
+		Sessions.Remove(EditableSessionRow);
+		EditableSessionRow.Reset();
+	}
+
+	// Clear the editable row state if its the one removed.
+	if (Item == EditableSessionRow)
+	{
+		EditableSessionRow.Reset();
+		EditableSessionRowParent.Reset();
+	}
+
+	SessionsView->RequestListRefresh();
+}
+
+TSharedRef<SWidget> SConcertSessionBrowser::MakeSessionViewOptionsBar()
+{
+	auto AddFilterMenu = [this]()
+	{
+		FMenuBuilder MenuBuilder(/*bInShouldCloseWindowAfterMenuSelection=*/true, nullptr);
+
+		MenuBuilder.AddMenuEntry(
+			LOCTEXT("ActiveSessions_Label", "Active Sessions"),
+			LOCTEXT("ActiveSessions_Tooltip", "Displays Active Sessions"),
+			FSlateIcon(),
+			FUIAction(
+				FExecuteAction::CreateSP(this, &SConcertSessionBrowser::OnFilterMenuChecked, ConcertBrowserUtils::ActiveSessionsCheckBoxMenuName),
+				FCanExecuteAction::CreateLambda([] { return true; }),
+				FIsActionChecked::CreateLambda([this] { return PersistentSettings->bShowActiveSessions; })),
+			NAME_None,
+			EUserInterfaceActionType::ToggleButton
+		);
+
+		MenuBuilder.AddMenuEntry(
+			LOCTEXT("ArchivedSessions_Label", "Archived Sessions"),
+			LOCTEXT("ArchivedSessions_Tooltip", "Displays Archived Sessions"),
+			FSlateIcon(),
+			FUIAction(
+				FExecuteAction::CreateSP(this, &SConcertSessionBrowser::OnFilterMenuChecked, ConcertBrowserUtils::ArchivedSessionsCheckBoxMenuName),
+				FCanExecuteAction::CreateLambda([] { return true; }),
+				FIsActionChecked::CreateLambda([this] { return PersistentSettings->bShowArchivedSessions; })),
+			NAME_None,
+			EUserInterfaceActionType::ToggleButton
+		);
+
+		MenuBuilder.AddMenuEntry(
+			LOCTEXT("DefaultServer_Label", "Default Server Sessions"),
+			LOCTEXT("DefaultServer_Tooltip", "Displays Sessions Hosted By the Default Server"),
+			FSlateIcon(),
+			FUIAction(
+				FExecuteAction::CreateSP(this, &SConcertSessionBrowser::OnFilterMenuChecked, ConcertBrowserUtils::DefaultServerCheckBoxMenuName),
+				FCanExecuteAction::CreateLambda([] { return true; }),
+				FIsActionChecked::CreateLambda([this] { return PersistentSettings->bShowDefaultServerSessionsOnly; })),
+			NAME_None,
+			EUserInterfaceActionType::ToggleButton
+		);
+
+		return MenuBuilder.MakeWidget();
+	};
+
+	return SNew(SHorizontalBox)
+		+SHorizontalBox::Slot()
+		.AutoWidth()
+		.VAlign(VAlign_Center)
+		[
+			SNew(STextBlock).Text_Lambda([this]()
+			{
+				// Don't count the 'New Session', 'Restore Session' and 'Archive Session' editable row, they are transient rows used for inline input only.
+				int32 DisplayedSessionNum = Sessions.Num() - (EditableSessionRow.IsValid() ? 1 : 0);
+				int32 AvailableSessionNum = GetController()->GetActiveSessions().Num() + GetController()->GetArchivedSessions().Num();
+				int32 ServerNum = GetController()->GetServers().Num();
+
+				// If all discovered session are displayed (none excluded by a filter).
+				if (DisplayedSessionNum == AvailableSessionNum)
+				{
+					if (GetController()->GetServers().Num() == 0)
+					{
+						return LOCTEXT("NoServerNoFilter", "No servers found");
+					}
+					else
+					{
+						return FText::Format(LOCTEXT("NSessionNServerNoFilter", "{0} {0}|plural(one=session,other=sessions) on {1} {1}|plural(one=server,other=servers)"),
+							DisplayedSessionNum, ServerNum);
+					}
+				}
+				else // A filter is excluding at least one session.
+				{
+					if (DisplayedSessionNum == 0)
+					{
+						return FText::Format(LOCTEXT("NoSessionMatchNServer", "No matching sessions ({0} total on {1} {1}|plural(one=server,other=servers))"),
+							AvailableSessionNum, ServerNum);
+					}
+					else
+					{
+						return FText::Format(LOCTEXT("NSessionNServer", "Showing {0} of {1} {1}|plural(one=session,other=sessions) on {2} {2}|plural(one=server,other=servers)"),
+							DisplayedSessionNum, AvailableSessionNum, ServerNum);
+					}
+				}
+			})
+		]
+
+		+SHorizontalBox::Slot()
+		.FillWidth(1.0)
+		[
+			SNew(SSpacer)
+		]
+
+		+SHorizontalBox::Slot()
+		.AutoWidth()
+		[
+			SNew(SComboButton)
+			.ComboButtonStyle(FAppStyle::Get(), "GenericFilters.ComboButtonStyle")
+			.ForegroundColor(FLinearColor::White)
+			.ContentPadding(0)
+			.OnGetMenuContent_Lambda(AddFilterMenu)
+			.HasDownArrow(true)
+			.ContentPadding(FMargin(1, 0))
+			.ButtonContent()
+			[
+				SNew(SHorizontalBox)
+
+				+SHorizontalBox::Slot()
+				.AutoWidth()
+				.VAlign(VAlign_Center)
+				[
+					SNew(SImage).Image(FAppStyle::Get().GetBrush("GenericViewButton")) // The eye ball image.
+				]
+
+				+SHorizontalBox::Slot()
+				.AutoWidth()
+				.Padding(2, 0, 0, 0)
+				.VAlign(VAlign_Center)
+				[
+					SNew(STextBlock).Text(LOCTEXT("ViewOptions", "View Options"))
+				]
+			]
+		];
+}
+
+void SConcertSessionBrowser::OnFilterMenuChecked(const FName MenuName)
+{
+	if (MenuName == ConcertBrowserUtils::ActiveSessionsCheckBoxMenuName)
+	{
+		PersistentSettings->bShowActiveSessions = !PersistentSettings->bShowActiveSessions;
+	}
+	else if (MenuName == ConcertBrowserUtils::ArchivedSessionsCheckBoxMenuName)
+	{
+		PersistentSettings->bShowArchivedSessions = !PersistentSettings->bShowArchivedSessions;
+	}
+	else if (MenuName == ConcertBrowserUtils::DefaultServerCheckBoxMenuName)
+	{
+		PersistentSettings->bShowDefaultServerSessionsOnly = !PersistentSettings->bShowDefaultServerSessionsOnly;
+	}
+	bRefreshSessionFilter = true;
+
+	PersistentSettings->SaveConfig();
+}
+
+TSharedPtr<SWidget> SConcertSessionBrowser::MakeContextualMenu()
+{
+	TArray<TSharedPtr<FConcertSessionItem>> SelectedItems = SessionsView->GetSelectedItems();
+	if (SelectedItems.Num() == 0 || (SelectedItems[0]->Type != FConcertSessionItem::EType::ActiveSession && SelectedItems[0]->Type != FConcertSessionItem::EType::ArchivedSession))
+	{
+		return nullptr; // No menu for editable rows.
+	}
+
+	TSharedPtr<FConcertSessionItem> Item = SelectedItems[0];
+
+	TSharedRef<FExtender> Extender = MakeShared<FExtender>();
+	ExtendSessionContextMenu.ExecuteIfBound(Item, Extender.Get());
+	FMenuBuilder MenuBuilder(/*bInShouldCloseWindowAfterMenuSelection=*/true, nullptr, Extender);
+
+	// Section title.
+	MenuBuilder.BeginSection(SessionContextMenuExtensionHooks::ManageSession, Item->Type == FConcertSessionItem::EType::ActiveSession ?
+		LOCTEXT("ActiveSessionSection", "Active Session"):
+		LOCTEXT("ArchivedSessionSection", "Archived Session"));
+
+	if (Item->Type == FConcertSessionItem::EType::ActiveSession)
+	{
+		MenuBuilder.AddMenuEntry(
+			LOCTEXT("CtxMenuArchive", "Archive"),
+			LOCTEXT("CtxMenuArchive_Tooltip", "Archived the Session"),
+			FSlateIcon(),
+			FUIAction(
+				FExecuteAction::CreateLambda([this](){ OnArchiveButtonClicked(); }),
+				FCanExecuteAction::CreateLambda([SelectedCount = SelectedItems.Num()] { return SelectedCount == 1; }),
+				FIsActionChecked::CreateLambda([this] { return false; })),
+			NAME_None,
+			EUserInterfaceActionType::Button);
+	}
+	else // Archive
+	{
+		MenuBuilder.AddMenuEntry(
+			LOCTEXT("CtxMenuRestore", "Restore"),
+			LOCTEXT("CtxMenuRestore_Tooltip", "Restore the Session"),
+			FSlateIcon(),
+			FUIAction(
+				FExecuteAction::CreateLambda([this](){ OnRestoreButtonClicked(); }),
+				FCanExecuteAction::CreateLambda([SelectedCount = SelectedItems.Num()] { return SelectedCount == 1; }),
+				FIsActionChecked::CreateLambda([this] { return false; })),
+			NAME_None,
+			EUserInterfaceActionType::Button);
+	}
+
+	MenuBuilder.AddMenuEntry(
+		LOCTEXT("CtxMenuRename", "Rename"),
+		LOCTEXT("CtxMenuRename_Tooltip", "Rename the Session"),
+		FSlateIcon(),
+		FUIAction(
+			FExecuteAction::CreateLambda([this, Item](){ OnBeginEditingSessionName(Item); }),
+			FCanExecuteAction::CreateLambda([this] { return IsRenameButtonEnabledInternal(); }),
+			FIsActionChecked::CreateLambda([this] { return false; })),
+		NAME_None,
+		EUserInterfaceActionType::Button);
+
+	MenuBuilder.AddMenuEntry(
+		LOCTEXT("CtxMenuDelete", "Delete"),
+		FText::Format(LOCTEXT("CtxMenuDelete_Tooltip", "Delete the {0}|plural(one=Session,other=Sessions)"), SelectedItems.Num()),
+		FSlateIcon(),
+		FUIAction(
+			FExecuteAction::CreateLambda([this](){ OnDeleteButtonClicked(); }),
+			FCanExecuteAction::CreateLambda([this] { return IsDeleteButtonEnabledInternal(); }),
+			FIsActionChecked::CreateLambda([this] { return false; })),
+		NAME_None,
+		EUserInterfaceActionType::Button);
+
+	MenuBuilder.EndSection();
+
+	return MenuBuilder.MakeWidget();
+}
+
+void SConcertSessionBrowser::OnSessionSelectionChanged(TSharedPtr<FConcertSessionItem> SelectedSession, ESelectInfo::Type SelectInfo)
+{
+	// Cancel editing the row to create, archive or restore a session (if any), unless the row was selected in code.
+	if (EditableSessionRow.IsValid() && SelectInfo != ESelectInfo::Direct)
+	{
+		check (SelectedSession != EditableSessionRow); // User should not be able to reselect an editable row as we remove it as soon as it is unselected.
+		RemoveSessionRow(EditableSessionRow);
+		check(!EditableSessionRow.IsValid() && !EditableSessionRowParent.IsValid()); // Expect to be cleared by RemoveSessionRow().
+	}
+
+	// Clear the list of clients (if any)
+	Clients.Reset();
+
+	OnSessionClicked.ExecuteIfBound(SelectedSession);
+}
+
+bool SConcertSessionBrowser::IsNewButtonEnabledInternal() const
+{
+	return GetController()->GetServers().Num() > 0;
+}
+
+bool SConcertSessionBrowser::IsRestoreButtonEnabledInternal() const
+{
+	TArray<TSharedPtr<FConcertSessionItem>> SelectedItems = SessionsView->GetSelectedItems();
+	return SelectedItems.Num() == 1 && SelectedItems[0]->Type == FConcertSessionItem::EType::ArchivedSession;
+}
+
+bool SConcertSessionBrowser::IsArchiveButtonEnabledInternal() const
+{
+	TArray<TSharedPtr<FConcertSessionItem>> SelectedItems = SessionsView->GetSelectedItems();
+	return SelectedItems.Num() == 1 && SelectedItems[0]->Type == FConcertSessionItem::EType::ActiveSession;
+}
+
+bool SConcertSessionBrowser::IsRenameButtonEnabledInternal() const
+{
+	TArray<TSharedPtr<FConcertSessionItem>> SelectedItems = SessionsView->GetSelectedItems();
+	if (SelectedItems.Num() != 1)
+	{
+		return false;
+	}
+
+	return (SelectedItems[0]->Type == FConcertSessionItem::EType::ActiveSession && GetController()->CanRenameActiveSession(SelectedItems[0]->ServerAdminEndpointId, SelectedItems[0]->SessionId)) ||
+	       (SelectedItems[0]->Type == FConcertSessionItem::EType::ArchivedSession && GetController()->CanRenameArchivedSession(SelectedItems[0]->ServerAdminEndpointId, SelectedItems[0]->SessionId));
+}
+
+bool SConcertSessionBrowser::IsDeleteButtonEnabledInternal() const
+{
+	TArray<TSharedPtr<FConcertSessionItem>> SelectedItems = SessionsView->GetSelectedItems();
+	if (SelectedItems.Num() == 0)
+	{
+		return false;
+	}
+
+	return (SelectedItems[0]->Type == FConcertSessionItem::EType::ActiveSession && GetController()->CanDeleteActiveSession(SelectedItems[0]->ServerAdminEndpointId, SelectedItems[0]->SessionId)) ||
+	       (SelectedItems[0]->Type == FConcertSessionItem::EType::ArchivedSession && GetController()->CanDeleteArchivedSession(SelectedItems[0]->ServerAdminEndpointId, SelectedItems[0]->SessionId));
+}
+
+FReply SConcertSessionBrowser::OnNewButtonClicked()
+{
+	InsertNewSessionEditableRowInternal();
+	return FReply::Handled();
+}
+
+FReply SConcertSessionBrowser::OnRestoreButtonClicked()
+{
+	TArray<TSharedPtr<FConcertSessionItem>> SelectedItems = SessionsView->GetSelectedItems();
+	if (SelectedItems.Num() == 1)
+	{
+		InsertRestoreSessionAsEditableRowInternal(SelectedItems[0]);
+	}
+
+	return FReply::Handled();
+}
+
+FReply SConcertSessionBrowser::OnArchiveButtonClicked()
+{
+	TArray<TSharedPtr<FConcertSessionItem>> SelectedItems = SessionsView->GetSelectedItems();
+	if (SelectedItems.Num() == 1)
+	{
+		InsertArchiveSessionAsEditableRow(SelectedItems[0]);
+	}
+
+	return FReply::Handled();
+}
+
+FReply SConcertSessionBrowser::OnDeleteButtonClicked()
+{
+	TArray<TSharedPtr<FConcertSessionItem>> SelectedItems = SessionsView->GetSelectedItems();
+	for (const TSharedPtr<FConcertSessionItem>& Item : SelectedItems)
+	{
+		RequestDeleteSession(Item);
+	}
+
+	return FReply::Handled();
+}
+
+void SConcertSessionBrowser::OnBeginEditingSessionName(TSharedPtr<FConcertSessionItem> Item)
+{
+	Item->OnBeginEditSessionNameRequest.Broadcast(); // Signal the row widget to enter in edit mode.
+}
+
+void SConcertSessionBrowser::Tick(const FGeometry& AllottedGeometry, const double InCurrentTime, const float InDeltaTime)
+{
+	// Ensure the 'default server' filter is updated when the configuration of the default server changes.
+	if (DefaultServerUrl.IsSet() && LastDefaultServerUrl != DefaultServerUrl.Get())
+	{
+		LastDefaultServerUrl = DefaultServerUrl.Get();
+		bRefreshSessionFilter = true;
+	}
+
+	// Should refresh the session filter?
+	if (bRefreshSessionFilter)
+	{
+		RefreshSessionList();
+		bRefreshSessionFilter = false;
+	}
+}
+
+FReply SConcertSessionBrowser::OnKeyDown(const FGeometry& MyGeometry, const FKeyEvent& InKeyEvent)
+{
+	// NOTE: When an 'editable row' text box has the focus the keys are grabbed by the text box but
+	//       if the editable row is still selected, but the text field doesn't have the focus anymore
+	//       the keys will end up here if the browser has the focus.
+
+	if (InKeyEvent.GetKey() == EKeys::Delete && !EditableSessionRow.IsValid()) // Delete selected row(s) unless the selected row is an 'editable' one.
+	{
+		for (TSharedPtr<FConcertSessionItem>& Item : SessionsView->GetSelectedItems())
+		{
+			RequestDeleteSession(Item);
+		}
+		return FReply::Handled();
+	}
+	else if (InKeyEvent.GetKey() == EKeys::Escape && EditableSessionRow.IsValid()) // Cancel 'new session', 'archive session' or 'restore session' action.
+	{
+		RemoveSessionRow(EditableSessionRow);
+		check(!EditableSessionRow.IsValid() && !EditableSessionRowParent.IsValid()); // Expect to be cleared by RemoveSessionRow().
+		return FReply::Handled();
+	}
+	else if (InKeyEvent.GetKey() == EKeys::F2 && !EditableSessionRow.IsValid())
+	{
+		TArray<TSharedPtr<FConcertSessionItem>> SelectedItems = SessionsView->GetSelectedItems();
+		if (SelectedItems.Num() == 1)
+		{
+			SelectedItems[0]->OnBeginEditSessionNameRequest.Broadcast(); // Broadcast the request.
+		}
+	}
+
+	return FReply::Unhandled();
+}
+
+void SConcertSessionBrowser::RequestCreateSession(const TSharedPtr<FConcertSessionItem>& NewItem)
+{
+	GetController()->CreateSession(NewItem->ServerAdminEndpointId, NewItem->SessionName);
+	RemoveSessionRow(NewItem); // The row used to edit the session name and pick the server.
+}
+
+void SConcertSessionBrowser::RequestArchiveSession(const TSharedPtr<FConcertSessionItem>& SaveItem, const FString& ArchiveName)
+{
+	GetController()->ArchiveSession(SaveItem->ServerAdminEndpointId, SaveItem->SessionId, ArchiveName, FConcertSessionFilter());
+	RemoveSessionRow(SaveItem); // The row used to edit the archive name.
+}
+
+void SConcertSessionBrowser::RequestRestoreSession(const TSharedPtr<FConcertSessionItem>& RestoreItem, const FString& SessionName)
+{
+	GetController()->RestoreSession(RestoreItem->ServerAdminEndpointId, RestoreItem->SessionId, SessionName, FConcertSessionFilter());
+	RemoveSessionRow(RestoreItem); // The row used to edit the restore as name.
+}
+
+void SConcertSessionBrowser::RequestRenameSession(const TSharedPtr<FConcertSessionItem>& RenamedItem, const FString& NewName)
+{
+	if (RenamedItem->Type == FConcertSessionItem::EType::ActiveSession)
+	{
+		GetController()->RenameActiveSession(RenamedItem->ServerAdminEndpointId, RenamedItem->SessionId, NewName);
+	}
+	else if (RenamedItem->Type == FConcertSessionItem::EType::ArchivedSession)
+	{
+		GetController()->RenameArchivedSession(RenamedItem->ServerAdminEndpointId, RenamedItem->SessionId, NewName);
+	}
+
+	// Display the new name until the server response is received. If the server refuses the new name, the discovery will reset the
+	// name (like if another client renamed it back) and the user will get a toast saying the rename failed.
+	RenamedItem->SessionName = NewName;
+}
+
+void SConcertSessionBrowser::RequestDeleteSession(const TSharedPtr<FConcertSessionItem>& DeletedItem)
+{
+	const FText SessionNameInText = FText::FromString(DeletedItem->SessionName);
+	const FText SeverNameInText = FText::FromString(DeletedItem->ServerName);
+	const FText ConfirmationMessage = FText::Format(LOCTEXT("DeleteSessionConfirmationMessage", "Do you really want to delete the session \"{0}\" from the server \"{1}\"?"), SessionNameInText, SeverNameInText);
+	const FText ConfirmationTitle = LOCTEXT("DeleteSessionConfirmationTitle", "Delete Session Confirmation");
+
+	if (FMessageDialog::Open(EAppMsgType::YesNo, ConfirmationMessage, &ConfirmationTitle) == EAppReturnType::Yes) // Confirmed?
+	{
+		if (DeletedItem->Type == FConcertSessionItem::EType::ActiveSession)
+		{
+			GetController()->DeleteActiveSession(DeletedItem->ServerAdminEndpointId, DeletedItem->SessionId);
+		}
+		else if (DeletedItem->Type == FConcertSessionItem::EType::ArchivedSession)
+		{
+			GetController()->DeleteArchivedSession(DeletedItem->ServerAdminEndpointId, DeletedItem->SessionId);
+		}
+
+		OnRequestedDeleteSession.ExecuteIfBound(DeletedItem);
+	}
+}
+
+#undef LOCTEXT_NAMESPACE
\ No newline at end of file
diff --git a/Engine/Plugins/Developer/Concert/ConcertUI/ConcertSharedSlate/Source/ConcertSharedSlate/Private/SessionBrowser/SNewSessionRow.cpp b/Engine/Plugins/Developer/Concert/ConcertUI/ConcertSharedSlate/Source/ConcertSharedSlate/Private/SessionBrowser/SNewSessionRow.cpp
new file mode 100644
--- /dev/null
+++ b/Engine/Plugins/Developer/Concert/ConcertUI/ConcertSharedSlate/Source/ConcertSharedSlate/Private/SessionBrowser/SNewSessionRow.cpp
@@ -0,0 +1,349 @@
+﻿// Copyright Epic Games, Inc. All Rights Reserved.
+
+#include "SNewSessionRow.h"
+
+#include "SessionBrowser/ConcertBrowserUtils.h"
+#include "SessionBrowser/ConcertSessionItem.h"
+
+#include "EditorFontGlyphs.h"
+#include "Styling/AppStyle.h"
+#include "Widgets/Layout/SBox.h"
+#include "Widgets/Layout/SHeader.h"
+#include "Widgets/Layout/SUniformGridPanel.h"
+#include "Widgets/Text/STextBlock.h"
+#include "Widgets/Views/STableRow.h"
+
+#define LOCTEXT_NAMESPACE "SConcertBrowser"
+
+void SNewSessionRow::Construct(const FArguments& InArgs, TSharedPtr<FConcertSessionItem> InItem, const TSharedRef<STableViewBase>& InOwnerTableView)
+{
+	Item = MoveTemp(InItem);
+	GetServersFunc = InArgs._GetServerFunc;
+	AcceptFunc = InArgs._OnAcceptFunc;
+	DeclineFunc = InArgs._OnDeclineFunc;
+	HighlightText = InArgs._HighlightText;
+	DefaultServerURL = InArgs._DefaultServerURL;
+
+	GetServersFunc.CheckCallable();
+	AcceptFunc.CheckCallable();
+	DeclineFunc.CheckCallable();
+
+	// Construct base class
+	SMultiColumnTableRow<TSharedPtr<FConcertSessionItem>>::Construct(FSuperRowType::FArguments(), InOwnerTableView);
+
+	// Fill the server combo.
+	UpdateServerList();
+}
+
+void SNewSessionRow::Tick(const FGeometry& AllottedGeometry, const double InCurrentTime, const float InDeltaTime)
+{
+	// Performance could be improved by only updating it when the server list has changed but it's no biggy...
+	UpdateServerList();
+
+	// Should give the focus to an editable text.
+	if (!bInitialFocusTaken)
+	{
+		bInitialFocusTaken = FSlateApplication::Get().SetKeyboardFocus(EditableSessionName.ToSharedRef());
+	}
+}
+
+TSharedRef<SWidget> SNewSessionRow::GenerateWidgetForColumn(const FName& ColumnName)
+{
+	TSharedPtr<FConcertSessionItem> ItemPin = Item.Pin();
+
+	if (ColumnName == ConcertBrowserUtils::IconColName)
+	{
+		// 'New' icon
+		return SNew(SBox)
+			.VAlign(VAlign_Center)
+			.HAlign(HAlign_Center)
+			[
+				SNew(STextBlock)
+					.Font(FAppStyle::Get().GetFontStyle(ConcertBrowserUtils::IconColumnFontName))
+					.Text(FEditorFontGlyphs::Plus_Circle)
+			];
+	}
+	else if (ColumnName == ConcertBrowserUtils::SessionColName)
+	{
+		return SNew(SBox)
+			.VAlign(VAlign_Center)
+			.Padding(FMargin(0, 0, 2, 0))
+			[
+				SAssignNew(EditableSessionName, SEditableTextBox)
+				.HintText(LOCTEXT("EnterSessionNameHint", "Enter a session name"))
+				.OnTextCommitted(this, &SNewSessionRow::OnSessionNameCommitted)
+				.OnKeyDownHandler(this, &SNewSessionRow::OnKeyDownHandler)
+				.OnTextChanged(this, &SNewSessionRow::OnSessionNameChanged)
+			];
+	}
+	else
+	{
+		return SNew(SHorizontalBox)
+
+			// 'Server' combo
+			+SHorizontalBox::Slot()
+			.AutoWidth()
+			.Padding(0, 1)
+			[
+				SAssignNew(ServersComboBox, SComboBox<TSharedPtr<FConcertServerInfo>>)
+				.OptionsSource(&Servers)
+				.OnGenerateWidget(this, &SNewSessionRow::OnGenerateServersComboOptionWidget)
+				.ToolTipText(LOCTEXT("SelectServerTooltip", "Select the server on which the session should be created"))
+				[
+					MakeSelectedServerWidget()
+				]
+			]
+
+			+SHorizontalBox::Slot()
+			.AutoWidth()
+			.Padding(2.0f)
+			.HAlign(HAlign_Left)
+			[
+				SNew(SUniformGridPanel)
+				.SlotPadding(FMargin(1.0f, 0.0f))
+
+				// 'Accept' button
+				+SUniformGridPanel::Slot(0, 0)
+				[
+					ConcertBrowserUtils::MakeIconButton(
+						TEXT("FlatButton.Success"),
+						FEditorFontGlyphs::Check,
+						LOCTEXT("CreateCheckIconTooltip", "Create the session"),
+						TAttribute<bool>::Create([this]() { return !EditableSessionName->GetText().IsEmpty(); }),
+						FOnClicked::CreateRaw(this, &SNewSessionRow::OnAccept),
+						FSlateColor(FLinearColor::White))
+				]
+
+				// 'Decline' button
+				+SUniformGridPanel::Slot(1, 0)
+				[
+					ConcertBrowserUtils::MakeIconButton(
+						TEXT("FlatButton.Danger"),
+						FEditorFontGlyphs::Times,
+						LOCTEXT("CancelIconTooltip", "Cancel"),
+						true, // Always enabled.
+						FOnClicked::CreateRaw(this, &SNewSessionRow::OnDecline),
+						FSlateColor(FLinearColor::White))
+				]
+			];
+	}
+}
+
+TSharedRef<SWidget> SNewSessionRow::OnGenerateServersComboOptionWidget(TSharedPtr<FConcertServerInfo> ServerItem)
+{
+	bool bIsDefaultServer = ServerItem->ServerName == DefaultServerURL.Get();
+
+	FText Tooltip;
+	if (bIsDefaultServer)
+	{
+		Tooltip = LOCTEXT("DefaultServerTooltip", "Default Configured Server");
+	}
+	else if (ServerItem->ServerName == FPlatformProcess::ComputerName())
+	{
+		Tooltip = LOCTEXT("LocalServerTooltip", "Local Server Running on This Computer");
+	}
+	else
+	{
+		Tooltip = LOCTEXT("OnlineServerTooltip", "Online Server");
+	}
+
+	return SNew(SHorizontalBox)
+		.ToolTipText(Tooltip)
+
+		+SHorizontalBox::Slot()
+		.AutoWidth()
+		[
+			SNew(STextBlock)
+			.Font(bIsDefaultServer ? FAppStyle::Get().GetFontStyle("BoldFont") : FAppStyle::Get().GetFontStyle("NormalFont"))
+			.Text(GetServerDisplayName(ServerItem->ServerName))
+		]
+		+SHorizontalBox::Slot()
+		.AutoWidth()
+		[
+			ConcertBrowserUtils::MakeServerVersionIgnoredWidget(ServerItem->ServerFlags)
+		];
+}
+
+void SNewSessionRow::UpdateServerList()
+{
+	// Remember the currently selected item (if any).
+	TSharedPtr<FConcertServerInfo> SelectedItem = ServersComboBox->GetSelectedItem(); // Instance in current list.
+
+	// Clear the current list. The list is rebuilt from scratch.
+	Servers.Reset();
+
+	TSharedPtr<FConcertServerInfo> LocalServerInfo;
+	TSharedPtr<FConcertServerInfo> DefaultServerInfo;
+	TSharedPtr<FConcertServerInfo> SelectedServerInfo; // Instance in the new list.
+
+	// Convert to shared ptr (slate needs that) and find if the latest list contains a default/local server.
+	for (const FConcertServerInfo& ServerInfo : GetServersFunc())
+	{
+		TSharedPtr<FConcertServerInfo> ComboItem = MakeShared<FConcertServerInfo>(ServerInfo);
+		
+		// Default server is deemed more important than local server to display the icon aside the server.
+		if (ComboItem->ServerName == DefaultServerURL.Get()) 
+		{
+			DefaultServerInfo = ComboItem;
+		}
+		else if (ComboItem->ServerName == FPlatformProcess::ComputerName())
+		{
+			LocalServerInfo = ComboItem;
+		}
+
+		if (SelectedItem.IsValid() && SelectedItem->ServerName == ComboItem->ServerName)
+		{
+			SelectedServerInfo = ComboItem; // Preserve user selection using the new instance.
+		}
+
+		Servers.Emplace(MoveTemp(ComboItem));
+	}
+
+	// Sort the server list alphabetically.
+	Servers.Sort([](const TSharedPtr<FConcertServerInfo>& Lhs, const TSharedPtr<FConcertServerInfo>& Rhs) { return Lhs->ServerName < Rhs->ServerName; });
+
+	// If a server is running on this machine, put it first in the list.
+	if (LocalServerInfo.IsValid() && Servers[0] != LocalServerInfo)
+	{
+		Servers.Remove(LocalServerInfo); // Keep sort order.
+		Servers.Insert(LocalServerInfo, 0);
+	}
+
+	// If a 'default server' is configured and available, put it first in the list. (Possibly overruling the local one)
+	if (DefaultServerInfo.IsValid() && Servers[0] != DefaultServerInfo)
+	{
+		Servers.Remove(DefaultServerInfo); // Keep sort order.
+		Servers.Insert(DefaultServerInfo, 0);
+	}
+
+	// If a server was selected and is still in the updated list.
+	if (SelectedServerInfo.IsValid())
+	{
+		// Preserve user selection.
+		ServersComboBox->SetSelectedItem(SelectedServerInfo);
+	}
+	else if (Servers.Num())
+	{
+		// Select the very first item in the list which is most likely be the default or the local server as they were put first above.
+		ServersComboBox->SetSelectedItem(Servers[0]);
+	}
+	else // Server list is empty.
+	{
+		ServersComboBox->ClearSelection();
+		Servers.Reset();
+	}
+
+	ServersComboBox->RefreshOptions();
+}
+
+TSharedRef<SWidget> SNewSessionRow::MakeSelectedServerWidget()
+{
+	return SNew(SHorizontalBox)
+		+SHorizontalBox::Slot()
+		.AutoWidth()
+		[
+			SNew(STextBlock)
+			.Text(this, &SNewSessionRow::GetSelectedServerText)
+			.HighlightText(HighlightText)
+		]
+		+SHorizontalBox::Slot()
+		.AutoWidth()
+		.Padding(2, 0, 0, 0)
+		[
+			SNew(STextBlock)
+			.Font(FAppStyle::Get().GetFontStyle("FontAwesome.9"))
+			.Text(this, &SNewSessionRow::GetSelectedServerIgnoreVersionText)
+			.ToolTipText(this, &SNewSessionRow::GetSelectedServerIgnoreVersionTooltip)
+		];
+}
+
+FText SNewSessionRow::GetSelectedServerText() const
+{
+	TSharedPtr<FConcertServerInfo> SelectedServer = ServersComboBox->GetSelectedItem();
+	if (SelectedServer.IsValid())
+	{
+		return GetServerDisplayName(SelectedServer->ServerName);
+	}
+	return LOCTEXT("SelectAServer", "Select a Server");
+}
+
+FText SNewSessionRow::GetServerDisplayName(const FString& ServerName) const
+{
+	const bool bIsDefaultServer = ServerName == DefaultServerURL.Get();
+	if (bIsDefaultServer)
+	{
+		return FText::Format(LOCTEXT("DefaultServer", "{0} (Default)"), FText::FromString(FPlatformProcess::ComputerName()));
+	}
+	if (ServerName == FPlatformProcess::ComputerName())
+	{
+		return FText::Format(LOCTEXT("MyComputer", "{0} (My Computer)"), FText::FromString(FPlatformProcess::ComputerName()));
+	}
+	return FText::FromString(ServerName);
+}
+
+FText SNewSessionRow::GetSelectedServerIgnoreVersionText() const
+{
+	if (ServersComboBox->GetSelectedItem() && (ServersComboBox->GetSelectedItem()->ServerFlags & EConcertServerFlags::IgnoreSessionRequirement) != EConcertServerFlags::None)
+	{
+		return FEditorFontGlyphs::Exclamation_Triangle;
+	}
+	return FText();
+}
+
+FText SNewSessionRow::GetSelectedServerIgnoreVersionTooltip() const
+{
+	if (ServersComboBox->GetSelectedItem() && (ServersComboBox->GetSelectedItem()->ServerFlags & EConcertServerFlags::IgnoreSessionRequirement) != EConcertServerFlags::None)
+	{
+		return ConcertBrowserUtils::GetServerVersionIgnoredTooltip();
+	}
+	return FText();
+}
+
+FReply SNewSessionRow::OnAccept()
+{
+	TSharedPtr<FConcertSessionItem> ItemPin = Item.Pin();
+	FString NewSessionName = EditableSessionName->GetText().ToString();
+
+	FText InvalidNameErrorMsg = ConcertSettingsUtils::ValidateSessionName(NewSessionName);
+	if (InvalidNameErrorMsg.IsEmpty())
+	{
+		ItemPin->SessionName = EditableSessionName->GetText().ToString();
+		ItemPin->ServerName = ServersComboBox->GetSelectedItem()->ServerName;
+		ItemPin->ServerAdminEndpointId = ServersComboBox->GetSelectedItem()->AdminEndpointId;
+		AcceptFunc(ItemPin); // Delegate to create the session.
+	}
+	else
+	{
+		EditableSessionName->SetError(InvalidNameErrorMsg);
+		FSlateApplication::Get().SetKeyboardFocus(EditableSessionName);
+	}
+
+	return FReply::Handled();
+}
+
+FReply SNewSessionRow::OnDecline()
+{
+	DeclineFunc(Item.Pin()); // Decline the creation and remove the row from the table.
+	return FReply::Handled();
+}
+
+void SNewSessionRow::OnSessionNameChanged(const FText& NewName)
+{
+	EditableSessionName->SetError(ConcertSettingsUtils::ValidateSessionName(NewName.ToString()));
+}
+
+void SNewSessionRow::OnSessionNameCommitted(const FText& NewText, ETextCommit::Type CommitType)
+{
+	if (CommitType == ETextCommit::OnEnter)
+	{
+		OnAccept(); // Create the session.
+	}
+}
+
+FReply SNewSessionRow::OnKeyDownHandler(const FGeometry&, const FKeyEvent& KeyEvent)
+{
+	// NOTE: This is invoked when the editable text field has the focus.
+	return KeyEvent.GetKey() == EKeys::Escape ? OnDecline() : FReply::Unhandled();
+}
+
+#undef LOCTEXT_NAMESPACE
\ No newline at end of file
diff --git a/Engine/Plugins/Developer/Concert/ConcertUI/ConcertSharedSlate/Source/ConcertSharedSlate/Private/SessionBrowser/SNewSessionRow.h b/Engine/Plugins/Developer/Concert/ConcertUI/ConcertSharedSlate/Source/ConcertSharedSlate/Private/SessionBrowser/SNewSessionRow.h
new file mode 100644
--- /dev/null
+++ b/Engine/Plugins/Developer/Concert/ConcertUI/ConcertSharedSlate/Source/ConcertSharedSlate/Private/SessionBrowser/SNewSessionRow.h
@@ -0,0 +1,78 @@
+﻿// Copyright Epic Games, Inc. All Rights Reserved.
+
+#pragma once
+
+#include "CoreMinimal.h"
+#include "ConcertMessageData.h"
+#include "Widgets/DeclarativeSyntaxSupport.h"
+#include "Widgets/Views/STableRow.h"
+
+class FConcertSessionItem;
+
+/**
+ * The type of row used in the session list view to edit a new session (the session name + server).
+ */
+class SNewSessionRow : public SMultiColumnTableRow<TSharedPtr<FConcertSessionItem>>
+{
+public:
+	using FGetServersFunc = TFunction<TArray<FConcertServerInfo>()>;
+	// Should remove the editable 'new' row and creates the sessions.
+	using FAcceptFunc = TFunction<void(const TSharedPtr<FConcertSessionItem>&)>;
+	// Should just remove the editable 'new' row.
+	using FDeclineFunc = TFunction<void(const TSharedPtr<FConcertSessionItem>&)>; 
+
+	SLATE_BEGIN_ARGS(SNewSessionRow) // Needed to use Args because Construct() is limited to 5 arguments and 6 were required.
+		: _GetServerFunc()
+		, _OnAcceptFunc()
+		, _OnDeclineFunc()
+		, _HighlightText()
+	{
+	}
+
+	SLATE_ARGUMENT(FGetServersFunc, GetServerFunc)
+	SLATE_ARGUMENT(FAcceptFunc, OnAcceptFunc)
+	SLATE_ARGUMENT(FDeclineFunc, OnDeclineFunc)
+	
+	SLATE_ATTRIBUTE(FText, HighlightText)
+	SLATE_ATTRIBUTE(FString, DefaultServerURL)
+
+	SLATE_END_ARGS()
+
+	void Construct(const FArguments& InArgs, TSharedPtr<FConcertSessionItem> InItem, const TSharedRef<STableViewBase>& InOwnerTableView);
+	
+	virtual TSharedRef<SWidget> GenerateWidgetForColumn(const FName& ColumnName) override;
+	virtual void Tick(const FGeometry& AllottedGeometry, const double InCurrentTime, const float InDeltaTime) override;
+
+private:
+	
+	TSharedRef<SWidget> OnGenerateServersComboOptionWidget(TSharedPtr<FConcertServerInfo> Item);
+	TSharedRef<SWidget> MakeSelectedServerWidget();
+	FText GetSelectedServerText() const;
+	FText GetSelectedServerIgnoreVersionText() const;
+	FText GetSelectedServerIgnoreVersionTooltip() const;
+	FText GetServerDisplayName(const FString& ServerName) const;
+
+	void OnSessionNameChanged(const FText& NewName);
+	void OnSessionNameCommitted(const FText& NewText, ETextCommit::Type CommitType);
+	FReply OnAccept();
+	FReply OnDecline();
+	FReply OnKeyDownHandler(const FGeometry&, const FKeyEvent&); // Registered as handler to the editable text (vs. OnKeyDown() virtual method).
+
+	void UpdateServerList();
+
+	/** Holds the new item to fill with session name and server. */
+	TWeakPtr<FConcertSessionItem> Item;
+	/** Servers displayed in the server combo box. */
+	TArray<TSharedPtr<FConcertServerInfo>> Servers;
+	
+	TSharedPtr<SComboBox<TSharedPtr<FConcertServerInfo>>> ServersComboBox;
+	TSharedPtr<SEditableTextBox> EditableSessionName;
+	
+	FGetServersFunc GetServersFunc;
+	FAcceptFunc AcceptFunc;
+	FDeclineFunc DeclineFunc;
+	
+	TAttribute<FText> HighlightText;
+	TAttribute<FString> DefaultServerURL;
+	bool bInitialFocusTaken = false;
+};
diff --git a/Engine/Plugins/Developer/Concert/ConcertUI/ConcertSharedSlate/Source/ConcertSharedSlate/Private/SessionBrowser/SSaveRestoreSessionRow.cpp b/Engine/Plugins/Developer/Concert/ConcertUI/ConcertSharedSlate/Source/ConcertSharedSlate/Private/SessionBrowser/SSaveRestoreSessionRow.cpp
new file mode 100644
--- /dev/null
+++ b/Engine/Plugins/Developer/Concert/ConcertUI/ConcertSharedSlate/Source/ConcertSharedSlate/Private/SessionBrowser/SSaveRestoreSessionRow.cpp
@@ -0,0 +1,209 @@
+﻿// Copyright Epic Games, Inc. All Rights Reserved.
+
+#include "SSaveRestoreSessionRow.h"
+
+#include "SessionBrowser/ConcertBrowserUtils.h"
+#include "SessionBrowser/ConcertSessionItem.h"
+
+#include "EditorFontGlyphs.h"
+#include "Framework/Application/SlateApplication.h"
+#include "Internationalization/Regex.h"
+#include "Widgets/SBoxPanel.h"
+#include "Widgets/Layout/SUniformGridPanel.h"
+#include "Widgets/Views/SExpanderArrow.h"
+
+#define LOCTEXT_NAMESPACE "SConcertBrowser"
+
+void SSaveRestoreSessionRow::Construct(const FArguments& InArgs, TSharedPtr<FConcertSessionItem> InNode, const TSharedRef<STableViewBase>& InOwnerTableView)
+{
+	Item = MoveTemp(InNode);
+	AcceptFunc = InArgs._OnAcceptFunc;
+	DeclineFunc = InArgs._OnDeclineFunc;
+	HighlightText = InArgs._HighlightText;
+
+	AcceptFunc.CheckCallable();
+	DeclineFunc.CheckCallable();
+
+	// Construct base class
+	SMultiColumnTableRow<TSharedPtr<FConcertSessionItem>>::Construct(FSuperRowType::FArguments(), InOwnerTableView);
+}
+
+TBitArray<> SSaveRestoreSessionRow::GetWiresNeededByDepth() const
+{
+	TBitArray<> Bits;
+	Bits.Add(false);
+	return Bits;
+}
+
+FText SSaveRestoreSessionRow::GetDefaultName(const FConcertSessionItem& InItem) const
+{
+	if (InItem.Type == FConcertSessionItem::EType::SaveSession)
+	{
+		return FText::Format(LOCTEXT("DefaultName", "{0}.{1}"), FText::FromString(InItem.SessionName), FText::FromString(FDateTime::UtcNow().ToString()));
+	}
+
+	// Supposing the name of the archive has the dates as suffix, like SessionXYZ.2019.03.13-19.39.12, then extracts SessionXYZ
+	static FRegexPattern Pattern(TEXT(R"((.*)\.\d+\.\d+\.\d+\-\d+\.\d+\.\d+$)"));
+	FRegexMatcher Matcher(Pattern, InItem.SessionName);
+	if (Matcher.FindNext())
+	{
+		return FText::FromString(Matcher.GetCaptureGroup(1));
+	}
+
+	return FText::FromString(InItem.SessionName);
+}
+
+TSharedRef<SWidget> SSaveRestoreSessionRow::GenerateWidgetForColumn(const FName& ColumnName)
+{
+	TSharedPtr<FConcertSessionItem> ItemPin = Item.Pin();
+
+	if (ColumnName == ConcertBrowserUtils::IconColName)
+	{
+		return SNew(SHorizontalBox)
+			+SHorizontalBox::Slot()
+			.AutoWidth()
+			.Padding(8, 0)
+			[
+				SNew(SExpanderArrow, SharedThis(this))
+				.StyleSet(&FAppStyle::Get())
+				.ShouldDrawWires(true)
+			];
+	}
+	else if (ColumnName == ConcertBrowserUtils::SessionColName)
+	{
+		return SNew(SHorizontalBox)
+
+			// 'Restore as/Save as' text
+			+SHorizontalBox::Slot()
+			.VAlign(VAlign_Center)
+			.AutoWidth()
+			.Padding(2.0, 0.0)
+			[
+				SNew(STextBlock).Text(ItemPin->Type == FConcertSessionItem::EType::RestoreSession ? LOCTEXT("RestoreAs", "Restore as:") : LOCTEXT("ArchiveAs", "Archive as:"))
+			]
+
+			// Editable text.
+			+SHorizontalBox::Slot()
+			.VAlign(VAlign_Center)
+			.Padding(FMargin(0, 0, 2, 0))
+			[
+				SAssignNew(EditableSessionName, SEditableTextBox)
+				.HintText(ItemPin->Type == FConcertSessionItem::EType::RestoreSession ? LOCTEXT("RestoreSessionHint", "Enter a session name") : LOCTEXT("ArchivSessionHint", "Enter an archive name"))
+				.OnTextCommitted(this, &SSaveRestoreSessionRow::OnSessionNameCommitted)
+				.OnKeyDownHandler(this, &SSaveRestoreSessionRow::OnKeyDownHandler)
+				.OnTextChanged(this, &SSaveRestoreSessionRow::OnSessionNameChanged)
+				.Text(GetDefaultName(*ItemPin))
+				.SelectAllTextWhenFocused(true)
+			];
+	}
+	else
+	{
+		check(ColumnName == ConcertBrowserUtils::ServerColName);
+
+		// 'Server' text block.
+		return SNew(SHorizontalBox)
+		+SHorizontalBox::Slot()
+		.AutoWidth()
+		[
+			SNew(SBox)
+			.VAlign(VAlign_Center)
+			[
+				// Server.
+				SNew(STextBlock)
+				.Text(FText::FromString(ItemPin->ServerName))
+				.HighlightText(HighlightText)
+			]
+		]
+
+		+SHorizontalBox::Slot()
+		.AutoWidth()
+		.Padding(2.0f)
+		.HAlign(HAlign_Left)
+		[
+			SNew(SUniformGridPanel)
+			.SlotPadding(FMargin(1.0f, 0.0f))
+
+			// 'Accept' button
+			+SUniformGridPanel::Slot(0, 0)
+			[
+				ConcertBrowserUtils::MakeIconButton(
+					TEXT("FlatButton.Success"),
+					FEditorFontGlyphs::Check,
+					ItemPin->Type == FConcertSessionItem::EType::RestoreSession ? LOCTEXT("RestoreCheckIconTooltip", "Restore the session") : LOCTEXT("ArchiveCheckIconTooltip", "Archive the session"),
+					TAttribute<bool>::Create([this]() { return !EditableSessionName->GetText().IsEmpty(); }), // Enabled?
+					FOnClicked::CreateRaw(this, &SSaveRestoreSessionRow::OnAccept),
+					FSlateColor(FLinearColor::White))
+			]
+
+			// 'Cancel' button
+			+SUniformGridPanel::Slot(1, 0)
+			[
+				ConcertBrowserUtils::MakeIconButton(
+					TEXT("FlatButton.Danger"),
+					FEditorFontGlyphs::Times,
+					LOCTEXT("CancelTooltip", "Cancel"),
+					true, // Enabled?
+					FOnClicked::CreateRaw(this, &SSaveRestoreSessionRow::OnDecline),
+					FSlateColor(FLinearColor::White))
+			]
+		];
+	}
+}
+
+void SSaveRestoreSessionRow::Tick(const FGeometry& AllottedGeometry, const double InCurrentTime, const float InDeltaTime)
+{
+	// Should give the focus to an editable text.
+	if (!bInitialFocusTaken)
+	{
+		bInitialFocusTaken = FSlateApplication::Get().SetKeyboardFocus(EditableSessionName.ToSharedRef());
+	}
+}
+
+void SSaveRestoreSessionRow::OnSessionNameChanged(const FText& NewName)
+{
+	EditableSessionName->SetError(ConcertSettingsUtils::ValidateSessionName(NewName.ToString()));
+}
+
+void SSaveRestoreSessionRow::OnSessionNameCommitted(const FText& NewText, ETextCommit::Type CommitType)
+{
+	TSharedPtr<FConcertSessionItem> ItemPin = Item.Pin();
+	if (CommitType == ETextCommit::Type::OnEnter)
+	{
+		OnAccept();
+	}
+}
+
+FReply SSaveRestoreSessionRow::OnAccept()
+{
+	// Read the session name given by the user.
+	TSharedPtr<FConcertSessionItem> ItemPin = Item.Pin();
+	FString Name = EditableSessionName->GetText().ToString(); // Archive name or restored session name.
+
+	// Ensure the user provided a name.
+	FText InvalidNameErrorMsg = ConcertSettingsUtils::ValidateSessionName(Name);
+	if (InvalidNameErrorMsg.IsEmpty()) // Name is valid?
+	{
+		AcceptFunc(ItemPin, Name); // Delegate archiving/restoring operation.
+	}
+	else
+	{
+		EditableSessionName->SetError(InvalidNameErrorMsg);
+		FSlateApplication::Get().SetKeyboardFocus(EditableSessionName);
+	}
+
+	return FReply::Handled();
+}
+
+FReply SSaveRestoreSessionRow::OnDecline()
+{
+	DeclineFunc(Item.Pin()); // Remove the save/restore editable row.
+	return FReply::Handled();
+}
+
+FReply SSaveRestoreSessionRow::OnKeyDownHandler(const FGeometry&, const FKeyEvent& KeyEvent)
+{
+	// NOTE: This handler is to capture the 'Escape' while the text field has focus.
+	return KeyEvent.GetKey() == EKeys::Escape ? OnDecline() : FReply::Unhandled();
+}
+
+#undef LOCTEXT_NAMESPACE
\ No newline at end of file
diff --git a/Engine/Plugins/Developer/Concert/ConcertUI/ConcertSharedSlate/Source/ConcertSharedSlate/Private/SessionBrowser/SSaveRestoreSessionRow.h b/Engine/Plugins/Developer/Concert/ConcertUI/ConcertSharedSlate/Source/ConcertSharedSlate/Private/SessionBrowser/SSaveRestoreSessionRow.h
new file mode 100644
--- /dev/null
+++ b/Engine/Plugins/Developer/Concert/ConcertUI/ConcertSharedSlate/Source/ConcertSharedSlate/Private/SessionBrowser/SSaveRestoreSessionRow.h
@@ -0,0 +1,63 @@
+﻿// Copyright Epic Games, Inc. All Rights Reserved.
+
+#pragma once
+
+#include "CoreMinimal.h"
+#include "Widgets/DeclarativeSyntaxSupport.h"
+#include "Widgets/Views/STableRow.h"
+
+class FConcertSessionItem;
+
+/**
+ * The type of row used in the session list view to create and archive or restore a session (edit the information required for the operation).
+ */
+class SSaveRestoreSessionRow : public SMultiColumnTableRow<TSharedPtr<FConcertSessionItem>>
+{
+public:
+	/** Should remove the editable row and save or restore the session. */
+	using FAcceptFunc = TFunction<void(TSharedPtr<FConcertSessionItem>, const FString&)>;
+	/** Should only remove the editable row from the table. */
+	using FDeclineFunc = TFunction<void(TSharedPtr<FConcertSessionItem>)>;
+
+	SLATE_BEGIN_ARGS(SSaveRestoreSessionRow)
+		: _OnAcceptFunc()
+		, _OnDeclineFunc()
+		, _HighlightText()
+	{}
+
+	SLATE_ARGUMENT(FAcceptFunc, OnAcceptFunc)
+	SLATE_ARGUMENT(FDeclineFunc, OnDeclineFunc)
+	SLATE_ATTRIBUTE(FText, HighlightText)
+
+	SLATE_END_ARGS()
+	
+	void Construct(const FArguments& InArgs, TSharedPtr<FConcertSessionItem> Node, const TSharedRef<STableViewBase>& InOwnerTableView);
+	virtual TSharedRef<SWidget> GenerateWidgetForColumn(const FName& ColumnName) override;
+	virtual void Tick(const FGeometry& AllottedGeometry, const double InCurrentTime, const float InDeltaTime) override;
+
+	void OnSessionNameChanged(const FText& NewName);
+	void OnSessionNameCommitted(const FText& NewText, ETextCommit::Type CommitType);
+	FReply OnAccept();
+	FReply OnDecline();
+	FReply OnKeyDownHandler(const FGeometry&, const FKeyEvent&); // Registered as handler to the editable text (vs. OnKeyDown() virtual method for this widget).
+
+	// Override SMultiColumnTableRow functions to ensure a wire is drawn between the item to restore and the editable row to link them together.
+	virtual TBitArray<> GetWiresNeededByDepth() const override;
+	virtual bool IsLastChild() const override { return true; }
+	virtual int32 DoesItemHaveChildren() const override { return 0; }
+	virtual bool IsItemExpanded() const override { return false; }
+
+	/** Generates a default name for an archive or a restored session. */
+	FText GetDefaultName(const FConcertSessionItem& Item) const; 
+
+private:
+	
+	TWeakPtr<FConcertSessionItem> Item;
+	TSharedPtr<SEditableTextBox> EditableSessionName;
+	
+	FAcceptFunc AcceptFunc;
+	FDeclineFunc DeclineFunc;
+	
+	TAttribute<FText> HighlightText;
+	bool bInitialFocusTaken = false;
+};
diff --git a/Engine/Plugins/Developer/Concert/ConcertUI/ConcertSharedSlate/Source/ConcertSharedSlate/Private/SessionBrowser/SSessionRow.cpp b/Engine/Plugins/Developer/Concert/ConcertUI/ConcertSharedSlate/Source/ConcertSharedSlate/Private/SessionBrowser/SSessionRow.cpp
new file mode 100644
--- /dev/null
+++ b/Engine/Plugins/Developer/Concert/ConcertUI/ConcertSharedSlate/Source/ConcertSharedSlate/Private/SessionBrowser/SSessionRow.cpp
@@ -0,0 +1,181 @@
+﻿// Copyright Epic Games, Inc. All Rights Reserved.
+
+#include "SSessionRow.h"
+
+#include "SessionBrowser/ConcertBrowserUtils.h"
+#include "SessionBrowser/ConcertSessionItem.h"
+
+#include "EditorFontGlyphs.h"
+
+#define LOCTEXT_NAMESPACE "SConcertBrowser"
+
+void SSessionRow::Construct(const FArguments& InArgs, TSharedPtr<FConcertSessionItem> InItem, const TSharedRef<STableViewBase>& InOwnerTableView)
+{
+	Item = MoveTemp(InItem);
+	DoubleClickFunc = InArgs._OnDoubleClickFunc; // This function should join a session or add a row to restore an archive.
+	RenameFunc = InArgs._OnRenameFunc; // Function invoked to send a rename request to the server.
+	IsDefaultSession = InArgs._IsDefaultSession;
+	HighlightText = InArgs._HighlightText;
+	IsSelected = InArgs._IsSelected;
+
+	DoubleClickFunc.CheckCallable();
+	RenameFunc.CheckCallable();
+	IsDefaultSession.CheckCallable();
+
+	// Construct base class
+	SMultiColumnTableRow<TSharedPtr<FConcertSessionItem>>::Construct(FSuperRowType::FArguments(), InOwnerTableView);
+
+	// Listen and handle rename request.
+	InItem->OnBeginEditSessionNameRequest.AddSP(this, &SSessionRow::OnBeginEditingSessionName);
+}
+
+TSharedRef<SWidget> SSessionRow::GenerateWidgetForColumn(const FName& ColumnName)
+{
+	TSharedPtr<FConcertSessionItem> ItemPin = Item.Pin();
+
+	if (ColumnName == ConcertBrowserUtils::IconColName)
+	{
+		return SNew(SBox)
+		.VAlign(VAlign_Center)
+		.HAlign(HAlign_Center)
+		.Padding(2)
+		.ToolTipText(ItemPin->Type == FConcertSessionItem::EType::ActiveSession ? LOCTEXT("ActiveIconTooltip", "Active session") : LOCTEXT("ArchivedIconTooltip", "Archived Session"))
+		[
+			SNew(STextBlock)
+			.Font(FAppStyle::Get().GetFontStyle(ConcertBrowserUtils::IconColumnFontName))
+			.Text(ItemPin->Type == FConcertSessionItem::EType::ActiveSession ? FEditorFontGlyphs::Circle : FEditorFontGlyphs::Archive)
+			.ColorAndOpacity(ItemPin->Type == FConcertSessionItem::EType::ActiveSession ? FAppStyle::Get().GetWidgetStyle<FButtonStyle>("FlatButton.Success").Normal.TintColor : FSlateColor::UseSubduedForeground())
+		];
+	}
+	
+	const bool bIsDefaultConfig = IsDefaultSession(ItemPin);
+	FSlateFontInfo FontInfo;
+	FSlateColor FontColor;
+	if (ItemPin->Type == FConcertSessionItem::EType::ActiveSession)
+	{
+		FontColor = bIsDefaultConfig ? FSlateColor(FLinearColor::White) : FSlateColor(FLinearColor::White * 0.8f);
+		FontInfo = FAppStyle::Get().GetFontStyle("NormalFont");
+	}
+	else
+	{
+		FontColor = FSlateColor::UseSubduedForeground();
+		FontInfo = FCoreStyle::GetDefaultFontStyle("Italic", 9);
+	}
+
+	if (ColumnName == ConcertBrowserUtils::SessionColName)
+	{
+		return SNew(SBox)
+			.VAlign(VAlign_Center)
+			[
+				SAssignNew(SessionNameText, SInlineEditableTextBlock)
+				.Text_Lambda([this]() { return FText::AsCultureInvariant(Item.Pin()->SessionName); })
+				.HighlightText(HighlightText)
+				.OnTextCommitted(this, &SSessionRow::OnSessionNameCommitted)
+				.IsReadOnly(false)
+				.IsSelected(FIsSelected::CreateLambda([this]() { return IsSelected.Get(); }))
+				.OnVerifyTextChanged(this, &SSessionRow::OnValidatingSessionName)
+				.Font(FontInfo)
+				.ColorAndOpacity(FontColor)
+			];
+	}
+
+	check(ColumnName == ConcertBrowserUtils::ServerColName);
+
+	if (bIsDefaultConfig)
+	{
+		return SNew(SHorizontalBox)
+			+SHorizontalBox::Slot()
+			.AutoWidth()
+			.VAlign(VAlign_Center)
+			[
+				SNew(STextBlock)
+				.Text(FText::Format(INVTEXT("{0} * "), FText::AsCultureInvariant(ItemPin->ServerName)))
+				.HighlightText(HighlightText)
+				.Font(FontInfo)
+				.ColorAndOpacity(FontColor)
+			]
+			+SHorizontalBox::Slot()
+			.AutoWidth()
+			.VAlign(VAlign_Center)
+			[
+				SNew(STextBlock)
+				.Text(LOCTEXT("DefaultServerSession", "(Default Session/Server)"))
+				.HighlightText(HighlightText)
+				.Font(FCoreStyle::GetDefaultFontStyle("Regular", 9))
+				.ColorAndOpacity(FontColor)
+			]
+			+SHorizontalBox::Slot()
+			[
+				SNew(SSpacer)
+			]
+			+SHorizontalBox::Slot()
+			.AutoWidth()
+			.VAlign(VAlign_Center)
+			.HAlign(HAlign_Right)
+			[
+				ConcertBrowserUtils::MakeServerVersionIgnoredWidget(ItemPin->ServerFlags)
+			];
+	}
+	else
+	{
+		return SNew(SHorizontalBox)
+			+SHorizontalBox::Slot()
+			.AutoWidth()
+			.VAlign(VAlign_Center)
+			[
+				SNew(STextBlock)
+				.Text(FText::AsCultureInvariant(ItemPin->ServerName))
+				.HighlightText(HighlightText)
+				.Font(FontInfo)
+				.ColorAndOpacity(FontColor)
+			]
+			+SHorizontalBox::Slot()
+			[
+				SNew(SSpacer)
+			]
+			+SHorizontalBox::Slot()
+			.AutoWidth()
+			.VAlign(VAlign_Center)
+			.HAlign(HAlign_Right)
+			[
+				ConcertBrowserUtils::MakeServerVersionIgnoredWidget(ItemPin->ServerFlags)
+			];
+	}
+}
+
+FReply SSessionRow::OnMouseButtonDoubleClick(const FGeometry& InMyGeometry, const FPointerEvent& InMouseEvent)
+{
+	if (TSharedPtr<FConcertSessionItem> ItemPin = Item.Pin())
+	{
+		DoubleClickFunc(ItemPin);
+	}
+	return FReply::Handled();
+}
+
+bool SSessionRow::OnValidatingSessionName(const FText& NewSessionName, FText& OutError)
+{
+	OutError = ConcertSettingsUtils::ValidateSessionName(NewSessionName.ToString());
+	return OutError.IsEmpty();
+}
+
+void SSessionRow::OnSessionNameCommitted(const FText& NewSessionName, ETextCommit::Type CommitType)
+{
+	if (TSharedPtr<FConcertSessionItem> ItemPin = Item.Pin())
+	{
+		FString NewName = NewSessionName.ToString();
+		if (NewName != ItemPin->SessionName) // Was renamed?
+		{
+			if (ConcertSettingsUtils::ValidateSessionName(NewName).IsEmpty()) // Name is valid?
+			{
+				RenameFunc(ItemPin, NewName); // Send the rename request to the server. (Server may still refuse at this point)
+			}
+			else
+			{
+				// NOTE: Error are interactively detected and raised by OnValidatingSessionName()
+				FSlateApplication::Get().SetKeyboardFocus(SessionNameText);
+			}
+		}
+	}
+}
+
+#undef LOCTEXT_NAMESPACE
\ No newline at end of file
diff --git a/Engine/Plugins/Developer/Concert/ConcertUI/ConcertSharedSlate/Source/ConcertSharedSlate/Private/SessionBrowser/SSessionRow.h b/Engine/Plugins/Developer/Concert/ConcertUI/ConcertSharedSlate/Source/ConcertSharedSlate/Private/SessionBrowser/SSessionRow.h
new file mode 100644
--- /dev/null
+++ b/Engine/Plugins/Developer/Concert/ConcertUI/ConcertSharedSlate/Source/ConcertSharedSlate/Private/SessionBrowser/SSessionRow.h
@@ -0,0 +1,59 @@
+﻿// Copyright Epic Games, Inc. All Rights Reserved.
+
+#pragma once
+
+#include "CoreMinimal.h"
+#include "Widgets/DeclarativeSyntaxSupport.h"
+#include "Widgets/Text/SInlineEditableTextBlock.h"
+#include "Widgets/Views/STableRow.h"
+
+class FConcertSessionItem;
+
+class SSessionRow : public SMultiColumnTableRow<TSharedPtr<FConcertSessionItem>>
+{
+public:
+	using FDoubleClickFunc = TFunction<void(TSharedPtr<FConcertSessionItem>)>;
+	using FRenameFunc = TFunction<void(TSharedPtr<FConcertSessionItem>, const FString&)>;
+	using FIsDefaultSession = TFunction<bool(TSharedPtr<FConcertSessionItem>)>;
+
+	SLATE_BEGIN_ARGS(SSessionRow)
+		: _OnDoubleClickFunc()
+		, _OnRenameFunc()
+		, _HighlightText()
+		, _IsSelected(false)
+	{}
+
+	SLATE_ARGUMENT(FDoubleClickFunc, OnDoubleClickFunc)
+	SLATE_ARGUMENT(FRenameFunc, OnRenameFunc)
+	SLATE_ARGUMENT(FIsDefaultSession, IsDefaultSession)
+	
+	SLATE_ATTRIBUTE(FText, HighlightText)
+	SLATE_ATTRIBUTE(bool, IsSelected)
+
+	SLATE_END_ARGS()
+
+	void Construct(const FArguments& InArgs, TSharedPtr<FConcertSessionItem> InItem, const TSharedRef<STableViewBase>& InOwnerTableView);
+	
+	virtual TSharedRef<SWidget> GenerateWidgetForColumn(const FName& ColumnName) override;
+	virtual FReply OnMouseButtonDoubleClick( const FGeometry& InMyGeometry, const FPointerEvent& InMouseEvent) override;
+
+	void OnSessionNameCommitted(const FText& NewSessionName, ETextCommit::Type CommitType);
+
+private:
+	
+	void OnBeginEditingSessionName() { SessionNameText->EnterEditingMode(); }
+	bool OnValidatingSessionName(const FText& NewSessionName, FText& OutError);
+
+	TWeakPtr<FConcertSessionItem> Item;
+
+	/** Invoked when the user double click on the row */
+	FDoubleClickFunc DoubleClickFunc;
+	/** Invoked when the user commit the session rename. (This will send the request to server) */
+	FRenameFunc RenameFunc;
+	/** Given a session checks whether it is the default session the client is supposed to join */
+	FIsDefaultSession IsDefaultSession;
+	
+	TAttribute<FText> HighlightText;
+	TAttribute<bool> IsSelected;
+	TSharedPtr<SInlineEditableTextBlock> SessionNameText;
+};
diff --git a/Engine/Plugins/Developer/Concert/ConcertUI/ConcertSharedSlate/Source/ConcertSharedSlate/Public/SessionBrowser/ConcertBrowserUtils.h b/Engine/Plugins/Developer/Concert/ConcertUI/ConcertSharedSlate/Source/ConcertSharedSlate/Public/SessionBrowser/ConcertBrowserUtils.h
new file mode 100644
--- /dev/null
+++ b/Engine/Plugins/Developer/Concert/ConcertUI/ConcertSharedSlate/Source/ConcertSharedSlate/Public/SessionBrowser/ConcertBrowserUtils.h
@@ -0,0 +1,97 @@
+﻿// Copyright Epic Games, Inc. All Rights Reserved.
+
+#pragma once
+
+#include "CoreMinimal.h"
+
+#include "ConcertMessageData.h"
+
+#include "EditorFontGlyphs.h"
+#include "Styling/AppStyle.h"
+#include "Internationalization/Regex.h"
+#include "Layout/Visibility.h"
+#include "Misc/Attribute.h"
+#include "Misc/AsyncTaskNotification.h"
+#include "Styling/SlateColor.h"
+#include "Widgets/Images/SImage.h"
+#include "Widgets/Input/SButton.h"
+#include "Widgets/Layout/SBorder.h"
+#include "Widgets/Text/STextBlock.h"
+
+namespace ConcertBrowserUtils
+{
+	// Defines the sessions list view column tag names.
+	static const FName IconColName(TEXT("Icon"));
+	static const FName SessionColName(TEXT("Session"));
+	static const FName ServerColName(TEXT("Server"));
+
+	// Name of the filter box in the View option.
+	static const FName ActiveSessionsCheckBoxMenuName(TEXT("ActiveSessions"));
+	static const FName ArchivedSessionsCheckBoxMenuName(TEXT("ArchivedSessions"));
+	static const FName DefaultServerCheckBoxMenuName(TEXT("DefaultServer"));
+
+	// The awesome font used to pick the icon displayed in the session list view 'Icon' column.
+	static const FName IconColumnFontName(TEXT("FontAwesome.9"));
+
+	/** Utility function used to create buttons displaying only an icon (using FontAwesome). */
+	inline TSharedRef<SButton> MakeIconButton(const FName& ButtonStyle, const TAttribute<FText>& GlyphIcon, const TAttribute<FText>& Tooltip, const TAttribute<bool>& EnabledAttribute, const FOnClicked& OnClicked,
+		const FSlateColor& ForegroundColor, const TAttribute<EVisibility>& Visibility = EVisibility::Visible, const TAttribute<FMargin>& ContentPadding = FMargin(3.0, 2.0), const FName FontStyle = IconColumnFontName)
+	{
+		return SNew(SButton)
+			.ButtonStyle(FAppStyle::Get(), ButtonStyle)
+			.OnClicked(OnClicked)
+			.ToolTipText(Tooltip)
+			.ContentPadding(ContentPadding)
+			.Visibility(Visibility)
+			.IsEnabled(EnabledAttribute)
+			.VAlign(VAlign_Center)
+			.HAlign(HAlign_Center)
+			[
+				SNew(STextBlock)
+				.Font(FAppStyle::Get().GetFontStyle(FontStyle))
+				.Text(GlyphIcon)
+			];
+	}
+
+	/** Utility function used to create buttons displaying only an icon (using a brush) */
+	inline TSharedRef<SButton> MakeIconButton(const FName& ButtonStyle, const TAttribute<const FSlateBrush*>& Icon, const TAttribute<FText>& Tooltip, const TAttribute<bool>& EnabledAttribute, const FOnClicked& OnClicked, const TAttribute<EVisibility>& Visibility = EVisibility::Visible)
+	{
+		return SNew(SButton)
+			.ButtonStyle(FAppStyle::Get(), ButtonStyle)
+			.OnClicked(OnClicked)
+			.ToolTipText(Tooltip)
+			.ContentPadding(FMargin(0, 0))
+			.Visibility(Visibility)
+			.IsEnabled(EnabledAttribute)
+			.VAlign(VAlign_Center)
+			.HAlign(HAlign_Center)
+			[
+				SNew(SImage)
+				.Image(Icon)
+				.ColorAndOpacity(FSlateColor::UseForeground())
+			];
+	}
+
+	/** Returns the tooltip shown when hovering the triangle with an exclamation icon when a server doesn't validate the version requirements. */
+	inline FText GetServerVersionIgnoredTooltip()
+	{
+#define LOCTEXT_NAMESPACE "SConcertBrowser"
+		return LOCTEXT("ServerIgnoreSessionRequirementsTooltip", "Careful this server won't verify that you have the right requirements before you join a session");
+#undef LOCTEXT_NAMESPACE
+	}
+
+	/** Create a widget displaying the triangle with an exclamation icon in case the server flags include IgnoreSessionRequirement. */
+	inline TSharedRef<SWidget> MakeServerVersionIgnoredWidget(EConcertServerFlags InServerFlags)
+	{
+		return SNew(SBorder)
+			.BorderImage(FAppStyle::Get().GetBrush("NoBorder"))
+			.ColorAndOpacity(FAppStyle::Get().GetWidgetStyle<FButtonStyle>("FlatButton.Warning").Normal.TintColor.GetSpecifiedColor())
+			[
+				SNew(STextBlock)
+				.Font(FAppStyle::Get().GetFontStyle("FontAwesome.9"))
+				.Text(FEditorFontGlyphs::Exclamation_Triangle)
+				.ToolTipText(GetServerVersionIgnoredTooltip())
+				.Visibility((InServerFlags & EConcertServerFlags::IgnoreSessionRequirement) != EConcertServerFlags::None ? EVisibility::Visible : EVisibility::Collapsed)
+			];
+	}
+}
diff --git a/Engine/Plugins/Developer/Concert/ConcertUI/ConcertSharedSlate/Source/ConcertSharedSlate/Public/SessionBrowser/ConcertSessionBrowserSettings.h b/Engine/Plugins/Developer/Concert/ConcertUI/ConcertSharedSlate/Source/ConcertSharedSlate/Public/SessionBrowser/ConcertSessionBrowserSettings.h
new file mode 100644
--- /dev/null
+++ b/Engine/Plugins/Developer/Concert/ConcertUI/ConcertSharedSlate/Source/ConcertSharedSlate/Public/SessionBrowser/ConcertSessionBrowserSettings.h
@@ -0,0 +1,26 @@
+// Copyright Epic Games, Inc. All Rights Reserved.
+
+#pragma once
+
+#include "UObject/Object.h"
+#include "ConcertSessionBrowserSettings.generated.h"
+
+/** Serializes the multi-user session browser settings like the active filters. */
+UCLASS(config=Editor)
+class CONCERTSHAREDSLATE_API UConcertSessionBrowserSettings : public UObject
+{
+	GENERATED_BODY()
+
+public:
+	UConcertSessionBrowserSettings() {}
+
+	UPROPERTY(config, EditAnywhere, Category="Multi-User Session Browser")
+	bool bShowActiveSessions = true;
+
+	UPROPERTY(config, EditAnywhere, Category="Multi-User Session Browser")
+	bool bShowArchivedSessions = true;
+
+	UPROPERTY(config, EditAnywhere, Category="Multi-User Session Browser")
+	bool bShowDefaultServerSessionsOnly = false;
+};
+
diff --git a/Engine/Plugins/Developer/Concert/ConcertUI/ConcertSharedSlate/Source/ConcertSharedSlate/Public/SessionBrowser/ConcertSessionItem.h b/Engine/Plugins/Developer/Concert/ConcertUI/ConcertSharedSlate/Source/ConcertSharedSlate/Public/SessionBrowser/ConcertSessionItem.h
new file mode 100644
--- /dev/null
+++ b/Engine/Plugins/Developer/Concert/ConcertUI/ConcertSharedSlate/Source/ConcertSharedSlate/Public/SessionBrowser/ConcertSessionItem.h
@@ -0,0 +1,49 @@
+﻿// Copyright Epic Games, Inc. All Rights Reserved.
+
+#pragma once
+
+#include "CoreMinimal.h"
+#include "ConcertMessageData.h"
+
+/** Signal emitted when a session name text field should enter in edit mode. */
+DECLARE_MULTICAST_DELEGATE(FOnBeginEditConcertSessionNameRequest)
+
+/**
+ * Item displayed in the session list view.
+ */
+class CONCERTSHAREDSLATE_API FConcertSessionItem
+{
+public:
+	enum class EType
+	{
+		None,
+		NewSession,      // Editable item to enter a session name and a pick a server.
+		RestoreSession,  // Editable item to name the restored session.
+		SaveSession,     // Editable item to name the archive.
+		ActiveSession,   // Read-only item representing an active session.
+		ArchivedSession, // Read-only item representing an archived session.
+	};
+
+	FConcertSessionItem(EType Type, const FString& InSessionName, const FGuid& InSessionId, const FString& InServerName, const FGuid& InServerEndpoint, EConcertServerFlags InServerFlags)
+		: Type(Type)
+		, ServerAdminEndpointId(InServerEndpoint)
+		, SessionId(InSessionId)
+		, SessionName(InSessionName)
+		, ServerName(InServerName)
+		, ServerFlags(InServerFlags)
+	{
+	}
+
+	bool operator==(const FConcertSessionItem& Other) const
+	{
+		return Type == Other.Type && ServerAdminEndpointId == Other.ServerAdminEndpointId && SessionId == Other.SessionId;
+	}
+
+	EType Type = EType::None;
+	FGuid ServerAdminEndpointId;
+	FGuid SessionId;
+	FString SessionName;
+	FString ServerName;
+	FOnBeginEditConcertSessionNameRequest OnBeginEditSessionNameRequest; // Emitted when user press 'F2' or select 'Rename' from context menu.
+	EConcertServerFlags ServerFlags;
+};
\ No newline at end of file
diff --git a/Engine/Plugins/Developer/Concert/ConcertUI/ConcertSharedSlate/Source/ConcertSharedSlate/Public/SessionBrowser/IConcertSessionBrowserController.h b/Engine/Plugins/Developer/Concert/ConcertUI/ConcertSharedSlate/Source/ConcertSharedSlate/Public/SessionBrowser/IConcertSessionBrowserController.h
new file mode 100644
--- /dev/null
+++ b/Engine/Plugins/Developer/Concert/ConcertUI/ConcertSharedSlate/Source/ConcertSharedSlate/Public/SessionBrowser/IConcertSessionBrowserController.h
@@ -0,0 +1,61 @@
+﻿// Copyright Epic Games, Inc. All Rights Reserved.
+
+#pragma once
+
+#include "CoreMinimal.h"
+#include "ConcertMessageData.h"
+
+/**
+ * Runs and cache network queries for the UI. In the model-view-controller pattern, this class acts like the controller. Its purpose
+ * is to keep the UI code as decoupled as possible from the API used to query it. It encapsulate the asynchronous code and provide a
+ * simpler API to the UI.
+ */
+class CONCERTSHAREDSLATE_API IConcertSessionBrowserController
+{
+public:
+
+	struct FActiveSessionInfo
+	{
+		FConcertServerInfo ServerInfo;
+		FConcertSessionInfo SessionInfo;
+		TArray<FConcertSessionClientInfo> Clients;
+	};
+
+	struct FArchivedSessionInfo
+	{
+		FConcertServerInfo ServerInfo;
+		FConcertSessionInfo SessionInfo;
+	};
+
+public:
+
+	/** Returns the latest list of server known to this controller. */
+	virtual TArray<FConcertServerInfo> GetServers() const = 0;
+
+	/** Returns the latest list of active sessions known to this controller. */
+	virtual TArray<TSharedPtr<FActiveSessionInfo>> GetActiveSessions() const = 0;
+
+	/** Returns the latest list of archived sessions known to this controller. */
+	virtual TArray<TSharedPtr<FArchivedSessionInfo>> GetArchivedSessions() const = 0;
+
+	/** Returns the active sessions info corresponding to the specified parameters. Used to display the sessions details. */
+	virtual const FConcertSessionInfo* GetActiveSessionInfo(const FGuid& AdminEndpoint, const FGuid& SessionId) const = 0;
+
+	/** Returns the archived sessions info corresponding to the specified parameters. Used to display the sessions details. */
+	virtual const FConcertSessionInfo* GetArchivedSessionInfo(const FGuid& AdminEndpoint, const FGuid& SessionId) const = 0;
+
+	virtual void CreateSession(const FGuid& ServerAdminEndpointId, const FString& SessionName) = 0;
+	virtual void ArchiveSession(const FGuid& ServerAdminEndpointId, const FGuid& SessionId, const FString& ArchiveName, const FConcertSessionFilter& SessionFilter) = 0;
+	virtual void RestoreSession(const FGuid& ServerAdminEndpointId, const FGuid& SessionId, const FString& RestoredName, const FConcertSessionFilter& SessionFilter) = 0;
+	virtual void JoinSession(const FGuid& ServerAdminEndpointId, const FGuid& SessionId) = 0;
+	virtual void RenameActiveSession(const FGuid& ServerAdminEndpointId, const FGuid& SessionId, const FString& NewName) = 0;
+	virtual void RenameArchivedSession(const FGuid& ServerAdminEndpointId, const FGuid& SessionId, const FString& NewName) = 0;
+	virtual bool CanRenameActiveSession(const FGuid& ServerAdminEndpointId, const FGuid& SessionId) const = 0;
+	virtual bool CanRenameArchivedSession(const FGuid& ServerAdminEndpointId, const FGuid& SessionId) const = 0;
+	virtual void DeleteActiveSession(const FGuid& ServerAdminEndpointId, const FGuid& SessionId) = 0;
+	virtual void DeleteArchivedSession(const FGuid& ServerAdminEndpointId, const FGuid& SessionId) = 0;
+	virtual bool CanDeleteActiveSession(const FGuid& ServerAdminEndpointId, const FGuid& SessionId) const = 0;
+	virtual bool CanDeleteArchivedSession(const FGuid& ServerAdminEndpointId, const FGuid& SessionId) const = 0;
+	
+	virtual ~IConcertSessionBrowserController() = default;
+};
\ No newline at end of file
diff --git a/Engine/Plugins/Developer/Concert/ConcertUI/ConcertSharedSlate/Source/ConcertSharedSlate/Public/SessionBrowser/SConcertSessionBrowser.h b/Engine/Plugins/Developer/Concert/ConcertUI/ConcertSharedSlate/Source/ConcertSharedSlate/Public/SessionBrowser/SConcertSessionBrowser.h
new file mode 100644
--- /dev/null
+++ b/Engine/Plugins/Developer/Concert/ConcertUI/ConcertSharedSlate/Source/ConcertSharedSlate/Public/SessionBrowser/SConcertSessionBrowser.h
@@ -0,0 +1,217 @@
+﻿// Copyright Epic Games, Inc. All Rights Reserved.
+
+#pragma once
+
+#include "CoreMinimal.h"
+#include "Misc/TextFilter.h"
+#include "SlateFwd.h"
+#include "Widgets/DeclarativeSyntaxSupport.h"
+#include "Widgets/SCompoundWidget.h"
+#include "Widgets/Views/SHeaderRow.h"
+#include "Widgets/Views/SListView.h"
+
+class FConcertSessionItem;
+class IConcertSessionBrowserController;
+class ITableRow;
+class UConcertSessionBrowserSettings;
+class SGridPanel;
+class STableViewBase;
+
+struct FConcertSessionClientInfo;
+struct FConcertSessionInfo;
+
+DECLARE_DELEGATE_RetVal_OneParam(TSharedRef<SWidget>, FExtendSessionTable, TSharedRef<SWidget> /* TableView */);
+DECLARE_DELEGATE_OneParam(FExtenderDelegate, FExtender&);
+DECLARE_DELEGATE_TwoParams(FExtendSessionContextMenu, TSharedPtr<FConcertSessionItem>, FExtender&)
+DECLARE_DELEGATE_OneParam(FSessionDelegate, TSharedPtr<FConcertSessionItem> /*Session*/);
+
+/**
+ * Enables the user to browse/search/filter/sort active and archived sessions, create new session,
+ * archive active sessions, restore archived sessions, join a session and open the settings dialog.
+ */
+class CONCERTSHAREDSLATE_API SConcertSessionBrowser : public SCompoundWidget
+{
+public:
+	
+	struct CONCERTSHAREDSLATE_API ControlButtonExtensionHooks
+	{
+		/** Contains: Create Session */
+		static const FName BeforeSeparator;
+		/** Just separates the two */
+		static const FName Separator;
+		/** Contains: Restore, Archive, Delete */
+		static const FName AfterSeparator;
+	};
+
+	struct CONCERTSHAREDSLATE_API SessionContextMenuExtensionHooks
+	{
+		/** Contains: Archive (ActiveSession), Restore (ArchivedSession), Rename, Delete  */
+		static const FName ManageSession;
+	};
+	
+	SLATE_BEGIN_ARGS(SConcertSessionBrowser) { }
+
+	/** Optional name of the default session - relevant for highlighting */
+	SLATE_ATTRIBUTE(FString, DefaultSessionName)
+	/** Optional url of the default server - relevant for highlighting */
+	SLATE_ATTRIBUTE(FString, DefaultServerURL)
+
+	/** Used during construction to override how the session table view is created, e.g. to embed it into an overlay */
+	SLATE_EVENT(FExtendSessionTable, ExtendSessionTable)
+	/** Extends the buttons to the left of the search bar */
+	SLATE_EVENT(FExtenderDelegate, ExtendControllButtons)
+	/** Extends the menu when the user right-clicks a session */
+	SLATE_EVENT(FExtendSessionContextMenu, ExtendSessionContextMenu)
+	/** Custom slot placed to the right of the search bar */
+	SLATE_NAMED_SLOT(FArguments, RightOfSearchBar)
+	
+	/** Called when this session is clicked */
+	SLATE_EVENT(FSessionDelegate, OnSessionClicked)
+	/** Called when this session is double-clicked */
+	SLATE_EVENT(FSessionDelegate, OnSessionDoubleClicked)
+	/** Called after a user has requested to delete a session */
+	SLATE_EVENT(FSessionDelegate, OnRequestedDeleteSession)
+	
+	SLATE_END_ARGS();
+
+	/**
+	* Constructs the Browser.
+	* @param InArgs The Slate argument list.
+	* @param InController The controller used to send queries from the UI - represents controller in mode-view-controller pattern.
+	* @param[in,out] InSearchText The text to set in the search box and to remember (as output). Cannot be null.
+	*/
+	void Construct(const FArguments& InArgs, TSharedRef<IConcertSessionBrowserController> InController, TSharedPtr<FText> InSearchText);
+
+	//~ Begin SWidget Interface
+	virtual void Tick(const FGeometry& AllottedGeometry, const double InCurrentTime, const float InDeltaTime) override;
+	virtual FReply OnKeyDown(const FGeometry& MyGeometry, const FKeyEvent& InKeyEvent) override;
+	//~ End SWidget Interface
+	
+	void RefreshSessionList();
+	TArray<TSharedPtr<FConcertSessionItem>> GetSessions() const { return Sessions; }
+	TArray<TSharedPtr<FConcertSessionItem>> GetSelectedItems() const { return SessionsView->GetSelectedItems(); }
+
+	bool IsNewButtonEnabled() const { return IsNewButtonEnabledInternal(); }
+	bool IsRestoreButtonEnabled() const { return IsRestoreButtonEnabledInternal(); }
+	bool IsArchiveButtonEnabled() const { return IsArchiveButtonEnabledInternal(); }
+	bool IsRenameButtonEnabled() const { return IsRenameButtonEnabledInternal(); }
+	bool IsDeleteButtonEnabled() const { return IsDeleteButtonEnabledInternal(); }
+	
+	/** Adds row for creating new session. Exposed for other widgets, e.g. discovery overlay to create a new session. */
+	void InsertNewSessionEditableRow() { InsertNewSessionEditableRowInternal(); }
+	/** Creates row under the given (archived) session with which session can be restored. */
+	void InsertRestoreSessionAsEditableRow(const TSharedPtr<FConcertSessionItem>& ArchivedItem) { InsertRestoreSessionAsEditableRowInternal(ArchivedItem); };
+	
+private:
+	
+	// Layout the 'session|details' split view.
+	TSharedRef<SWidget> MakeBrowserContent(const FArguments& InArgs);
+
+	// Layout the sessions view and controls.
+	TSharedRef<SWidget> MakeControlBar(const FArguments& InArgs);
+	TSharedRef<SWidget> MakeButtonBar(const FArguments& InArgs);
+	TSharedRef<SWidget> MakeSessionTableView();
+	TSharedRef<SWidget> MakeSessionViewOptionsBar();
+	TSharedRef<ITableRow> OnGenerateSessionRowWidget(TSharedPtr<FConcertSessionItem> Item, const TSharedRef<STableViewBase>& OwnerTable);
+
+	// Creates row widgets for session list view, validates user inputs and forward user requests for processing to a delegate function implemented by this class.
+	TSharedRef<ITableRow> MakeActiveSessionRowWidget(const TSharedPtr<FConcertSessionItem>& ActiveItem, const TSharedRef<STableViewBase>& OwnerTable);
+	TSharedRef<ITableRow> MakeArchivedSessionRowWidget(const TSharedPtr<FConcertSessionItem>& ArchivedItem, const TSharedRef<STableViewBase>& OwnerTable);
+	TSharedRef<ITableRow> MakeNewSessionRowWidget(const TSharedPtr<FConcertSessionItem>& NewItem, const TSharedRef<STableViewBase>& OwnerTable);
+	TSharedRef<ITableRow> MakeRestoreSessionRowWidget(const TSharedPtr<FConcertSessionItem>& RestoreItem, const TSharedRef<STableViewBase>& OwnerTable);
+	TSharedRef<ITableRow> MakeSaveSessionRowWidget(const TSharedPtr<FConcertSessionItem>& ArchivedItem, const TSharedRef<STableViewBase>& OwnerTable);
+
+	// Creates the contextual menu when right clicking a session list view row.
+	TSharedPtr<SWidget> MakeContextualMenu();
+
+	// The buttons above the session view.
+	bool IsNewButtonEnabledInternal() const;
+	bool IsRestoreButtonEnabledInternal() const;
+	bool IsArchiveButtonEnabledInternal() const;
+	bool IsRenameButtonEnabledInternal() const;
+	bool IsDeleteButtonEnabledInternal() const;
+	FReply OnNewButtonClicked();
+	FReply OnRestoreButtonClicked();
+	FReply OnArchiveButtonClicked();
+	FReply OnDeleteButtonClicked();
+	void OnBeginEditingSessionName(TSharedPtr<FConcertSessionItem> Item);
+
+	// Manipulates the sessions view (the array and the UI).
+	void OnSessionSelectionChanged(TSharedPtr<FConcertSessionItem> SelectedSession, ESelectInfo::Type SelectInfo);
+	void InsertNewSessionEditableRowInternal();
+	void InsertRestoreSessionAsEditableRowInternal(const TSharedPtr<FConcertSessionItem>& ArchivedItem);
+	void InsertArchiveSessionAsEditableRow(const TSharedPtr<FConcertSessionItem>& ActiveItem);
+	void InsertEditableSessionRow(TSharedPtr<FConcertSessionItem> EditableItem, TSharedPtr<FConcertSessionItem> ParentItem);
+	void RemoveSessionRow(const TSharedPtr<FConcertSessionItem>& Item);
+
+	// Sessions sorting. (Sorts the session view)
+	EColumnSortMode::Type GetColumnSortMode(const FName ColumnId) const;
+	EColumnSortPriority::Type GetColumnSortPriority(const FName ColumnId) const;
+	void OnColumnSortModeChanged(const EColumnSortPriority::Type SortPriority, const FName& ColumnId, const EColumnSortMode::Type InSortMode);
+	void SortSessionList();
+	void EnsureEditableParentChildOrder();
+
+	// Sessions filtering. (Filters the session view)
+	void OnSearchTextChanged(const FText& InFilterText);
+	void OnSearchTextCommitted(const FText& InFilterText, ETextCommit::Type CommitType);
+	void OnFilterMenuChecked(const FName MenuName);
+	void PopulateSearchStrings(const FConcertSessionItem& Item, TArray<FString>& OutSearchStrings) const;
+	bool IsFilteredOut(const FConcertSessionItem& Item) const;
+	FText HighlightSearchText() const;
+
+	// Passes the user requests to FConcertBrowserController.
+	void RequestCreateSession(const TSharedPtr<FConcertSessionItem>& NewItem);
+	void RequestArchiveSession(const TSharedPtr<FConcertSessionItem>& ActiveItem, const FString& ArchiveName);
+	void RequestRestoreSession(const TSharedPtr<FConcertSessionItem>& RestoreItem, const FString& SessionName);
+	void RequestRenameSession(const TSharedPtr<FConcertSessionItem>& RenamedItem, const FString& NewName);
+	void RequestDeleteSession(const TSharedPtr<FConcertSessionItem>& DeletedItem);
+
+	TSharedPtr<IConcertSessionBrowserController> GetController() const;
+
+private:
+	
+	// Gives access to the concert data (servers, sessions, clients, etc).
+	TWeakPtr<IConcertSessionBrowserController> Controller;
+
+	// Keeps persistent user preferences, like the filters.
+	UConcertSessionBrowserSettings* PersistentSettings = nullptr;
+
+	/** Optional default session name - relevant for highlighting */
+	TAttribute<FString> DefaultSessionName;
+	/** Optional default server url - relevant for highlighting */
+	TAttribute<FString> DefaultServerUrl;
+	
+	FExtendSessionContextMenu ExtendSessionContextMenu;
+	FSessionDelegate OnSessionClicked;
+	FSessionDelegate OnSessionDoubleClicked;
+	FSessionDelegate OnRequestedDeleteSession;
+
+	// The items displayed in the session list view. It might be filtered and sorted compared to the full list hold by the controller.
+	TArray<TSharedPtr<FConcertSessionItem>> Sessions;
+
+	// The session list view.
+	TSharedPtr<SListView<TSharedPtr<FConcertSessionItem>>> SessionsView;
+
+	// The item corresponding to a row used to create/archive/restore a session. There is only one at the time
+	TSharedPtr<FConcertSessionItem> EditableSessionRow;
+	TSharedPtr<FConcertSessionItem> EditableSessionRowParent; // For archive/restore, indicate which element is archived or restored.
+
+	// Sorting.
+	EColumnSortMode::Type PrimarySortMode = EColumnSortMode::None;
+	EColumnSortMode::Type SecondarySortMode = EColumnSortMode::None;
+	FName PrimarySortedColumn;
+	FName SecondarySortedColumn;
+
+	// Filtering.
+	TSharedPtr<SSearchBox> SearchBox;
+	TSharedPtr<TTextFilter<const FConcertSessionItem&>> SearchTextFilter;
+	TSharedPtr<FText> SearchedText;
+	bool bRefreshSessionFilter = true;
+	FString LastDefaultServerUrl;
+
+	// Selected Session Details.
+	TSharedPtr<SBorder> SessionDetailsView;
+	TSharedPtr<SExpandableArea> DetailsArea;
+	TArray<TSharedPtr<FConcertSessionClientInfo>> Clients;
+	TSharedPtr<SExpandableArea> ClientsArea;
+};
diff --git a/Engine/Plugins/Developer/Concert/MultiUserClient/Source/MultiUserClient/Private/Widgets/ConcertSessionBrowserSettings.h b/Engine/Plugins/Developer/Concert/MultiUserClient/Source/MultiUserClient/Private/Widgets/ConcertSessionBrowserSettings.h
deleted file mode 100644
diff --git a/Engine/Plugins/Developer/Concert/MultiUserClient/Source/MultiUserClient/Private/Widgets/SConcertBrowser.cpp b/Engine/Plugins/Developer/Concert/MultiUserClient/Source/MultiUserClient/Private/Widgets/SConcertBrowser.cpp
--- a/Engine/Plugins/Developer/Concert/MultiUserClient/Source/MultiUserClient/Private/Widgets/SConcertBrowser.cpp
+++ b/Engine/Plugins/Developer/Concert/MultiUserClient/Source/MultiUserClient/Private/Widgets/SConcertBrowser.cpp
@@ -2,45 +2,37 @@
 
 #include "SConcertBrowser.h"
 
+#include "ConcertActivityStream.h"
+#include "ConcertFrontendStyle.h"
+#include "ConcertFrontendUtils.h"
+#include "ConcertLogGlobal.h"
+#include "ConcertSettings.h"
 #include "IMultiUserClientModule.h"
 #include "IConcertClient.h"
 #include "IConcertClientWorkspace.h"
 #include "IConcertSyncClient.h"
-#include "ConcertActivityStream.h"
-#include "ConcertFrontendStyle.h"
-#include "ConcertFrontendUtils.h"
+#include "MultiUserClientUtils.h"
 #include "SActiveSession.h"
 #include "SConcertSessionRecovery.h"
-#include "ConcertSessionBrowserSettings.h"
-#include "ConcertSettings.h"
-#include "ConcertLogGlobal.h"
-#include "MultiUserClientUtils.h"
+#include "SessionBrowser/ConcertBrowserUtils.h"
+#include "SessionBrowser/ConcertSessionItem.h"
+#include "SessionBrowser/IConcertSessionBrowserController.h"
+#include "SessionBrowser/SConcertSessionBrowser.h"
 
 #include "Algo/Transform.h"
 #include "Styling/AppStyle.h"
-#include "Framework/Commands/GenericCommands.h"
-#include "Framework/MultiBox/MultiBoxBuilder.h"
 #include "Internationalization/Regex.h"
 #include "Misc/AsyncTaskNotification.h"
-#include "Misc/Char.h"
-#include "Misc/MessageDialog.h"
 #include "Misc/TextFilter.h"
-#include "SlateOptMacros.h"
-#include "UObject/StrongObjectPtr.h"
 #include "Widgets/Images/SImage.h"
 #include "Widgets/Images/SThrobber.h"
 #include "Widgets/Input/SButton.h"
-#include "Widgets/Input/SComboBox.h"
-#include "Widgets/Input/SEditableTextBox.h"
 #include "Widgets/Input/SSearchBox.h"
 #include "Widgets/Layout/SBorder.h"
 #include "Widgets/Layout/SBox.h"
 #include "Widgets/Layout/SExpandableArea.h"
 #include "Widgets/Layout/SGridPanel.h"
-#include "Widgets/Layout/SHeader.h"
 #include "Widgets/Layout/SScrollBox.h"
-#include "Widgets/Layout/SSpacer.h"
-#include "Widgets/Layout/SSeparator.h"
 #include "Widgets/Layout/SUniformGridPanel.h"
 #include "Widgets/Text/SInlineEditableTextBlock.h"
 #include "Widgets/Text/STextBlock.h"
@@ -53,136 +45,15 @@
 
 #define LOCTEXT_NAMESPACE "SConcertBrowser"
 
-namespace ConcertBrowserUtils
-{
-
-// Defines the sessions list view column tag names.
-const FName IconColName(TEXT("Icon"));
-const FName SessionColName(TEXT("Session"));
-const FName ServerColName(TEXT("Server"));
-
-// Name of the filter box in the View option.
-const FName ActiveSessionsCheckBoxMenuName(TEXT("ActiveSessions"));
-const FName ArchivedSessionsCheckBoxMenuName(TEXT("ArchivedSessions"));
-const FName DefaultServerCheckBoxMenuName(TEXT("DefaultServer"));
-
-// The awesome font used to pick the icon displayed in the session list view 'Icon' column.
-const FName IconColumnFontName(TEXT("FontAwesome.9"));
-
-/** Utility function used to create buttons displaying only an icon (using FontAwesome). */
-TSharedRef<SButton> MakeIconButton(const FName& ButtonStyle, const TAttribute<FText>& GlyphIcon, const TAttribute<FText>& Tooltip, const TAttribute<bool>& EnabledAttribute, const FOnClicked& OnClicked,
-	const FSlateColor& ForegroundColor, const TAttribute<EVisibility>& Visibility = EVisibility::Visible, const TAttribute<FMargin>& ContentPadding = FMargin(3.0, 2.0), const FName FontStyle = IconColumnFontName)
-{
-	return SNew(SButton)
-		.ButtonStyle(FAppStyle::Get(), ButtonStyle)
-		.OnClicked(OnClicked)
-		.ToolTipText(Tooltip)
-		.ContentPadding(ContentPadding)
-		.Visibility(Visibility)
-		.IsEnabled(EnabledAttribute)
-		.VAlign(VAlign_Center)
-		.HAlign(HAlign_Center)
-		[
-			SNew(STextBlock)
-			.Font(FAppStyle::Get().GetFontStyle(FontStyle))
-			.Text(GlyphIcon)
-		];
-}
-
-/** Utility function used to create buttons displaying only an icon (using a brush) */
-TSharedRef<SButton> MakeIconButton(const FName& ButtonStyle, const TAttribute<const FSlateBrush*>& Icon, const TAttribute<FText>& Tooltip, const TAttribute<bool>& EnabledAttribute, const FOnClicked& OnClicked, const TAttribute<EVisibility>& Visibility = EVisibility::Visible)
-{
-	return SNew(SButton)
-		.ButtonStyle(FAppStyle::Get(), ButtonStyle)
-		.OnClicked(OnClicked)
-		.ToolTipText(Tooltip)
-		.ContentPadding(FMargin(0, 0))
-		.Visibility(Visibility)
-		.IsEnabled(EnabledAttribute)
-		.VAlign(VAlign_Center)
-		.HAlign(HAlign_Center)
-		[
-			SNew(SImage)
-			.Image(Icon)
-			.ColorAndOpacity(FSlateColor::UseForeground())
-		];
-}
-
-/** Returns the tooltip shown when hovering the triangle with an exclamation icon when a server doesn't validate the version requirements. */
-FText GetServerVersionIgnoredTooltip()
-{
-	return LOCTEXT("ServerIgnoreSessionRequirementsTooltip", "Careful this server won't verify that you have the right requirements before you join a session");
-}
-
-/** Create a widget displaying the triangle with an exclamation icon in case the server flags include IgnoreSessionRequirement. */
-TSharedRef<SWidget> MakeServerVersionIgnoredWidget(EConcertServerFlags InServerFlags)
-{
-	return SNew(SBorder)
-		.BorderImage(FAppStyle::Get().GetBrush("NoBorder"))
-		.ColorAndOpacity(FAppStyle::Get().GetWidgetStyle<FButtonStyle>("FlatButton.Warning").Normal.TintColor.GetSpecifiedColor())
-		[
-			SNew(STextBlock)
-			.Font(FAppStyle::Get().GetFontStyle("FontAwesome.9"))
-			.Text(FEditorFontGlyphs::Exclamation_Triangle)
-			.ToolTipText(GetServerVersionIgnoredTooltip())
-			.Visibility((InServerFlags & EConcertServerFlags::IgnoreSessionRequirement) != EConcertServerFlags::None ? EVisibility::Visible : EVisibility::Collapsed)
-		];
-}
-
-} // namespace ConcertSessionBrowserUtils
-
-/** Signal emitted when a session name text field should enter in edit mode. */
-DECLARE_MULTICAST_DELEGATE(FOnBeginEditConcertSessionNameRequest)
-
-/**
- * Item displayed in the session list view.
- */
-class FConcertSessionItem
-{
-public:
-	enum class EType
-	{
-		None,
-		NewSession,      // Editable item to enter a session name and a pick a server.
-		RestoreSession,  // Editable item to name the restored session.
-		SaveSession,     // Editable item to name the archive.
-		ActiveSession,   // Read-only item representing an active session.
-		ArchivedSession, // Read-only item representing an archived session.
-	};
-
-	FConcertSessionItem(EType Type, const FString& InSessionName, const FGuid& InSessionId, const FString& InServerName, const FGuid& InServerEndpoint, EConcertServerFlags InServerFlags)
-		: Type(Type)
-		, ServerAdminEndpointId(InServerEndpoint)
-		, SessionId(InSessionId)
-		, SessionName(InSessionName)
-		, ServerName(InServerName)
-		, ServerFlags(InServerFlags)
-	{
-	}
-
-	bool operator==(const FConcertSessionItem& Other) const
-	{
-		return Type == Other.Type && ServerAdminEndpointId == Other.ServerAdminEndpointId && SessionId == Other.SessionId;
-	}
-
-	EType Type = EType::None;
-	FGuid ServerAdminEndpointId;
-	FGuid SessionId;
-	FString SessionName;
-	FString ServerName;
-	FOnBeginEditConcertSessionNameRequest OnBeginEditSessionNameRequest; // Emitted when user press 'F2' or select 'Rename' from context menu.
-	EConcertServerFlags ServerFlags;
-};
-
-
 /**
  * Runs and cache network queries for the UI. In the model-view-controller pattern, this class acts like the controller. Its purpose
  * is to keep the UI code as decoupled as possible from the API used to query it. It encapsulate the asynchronous code and provide a
  * simpler API to the UI.
  */
-class FConcertBrowserController
+class FConcertClientSessionBrowserController : public IConcertSessionBrowserController
 {
 public:
+	
 	/** Keeps the state of an active async request and provides a tool to cancel its future continuation execution. */
 	struct FAsyncRequest
 	{
@@ -201,28 +72,9 @@ public:
 		/** Determines whether or not the async request continuation code should execute. Reset to disarm execution of an async future continuation. */
 		TSharedPtr<uint8, ESPMode::ThreadSafe> FutureExecutionToken;
 	};
-
-	struct FActiveSessionInfo
-	{
-		FConcertServerInfo ServerInfo;
-		FConcertSessionInfo SessionInfo;
-		FAsyncRequest ListClientRequest;
-		TArray<FConcertSessionClientInfo> Clients;
-		bool bSessionNameDirty = false; // Raised when the UI and the cache values may be out of sync if a rename failed (UI assumed it succeeded)
-	};
-
-	struct FArchivedSessionInfo
-	{
-		FConcertServerInfo ServerInfo;
-		FConcertSessionInfo SessionInfo;
-		bool bSessionNameDirty = false; // Raised when the UI and the cache values may be out of sync if a rename failed (UI assumed it succeeded)
-	};
-
-public:
-	FConcertBrowserController(IConcertClientPtr InConcertClient);
-	~FConcertBrowserController();
-
-	IConcertClientPtr GetConcertClient() { return ConcertClient; }
+	
+	FConcertClientSessionBrowserController(IConcertClientPtr InConcertClient);
+	virtual ~FConcertClientSessionBrowserController() override;
 
 	/**
 	 * Fires new requests to retrieve all known server and for each server, their active and archived sessions. The responses are
@@ -244,43 +96,53 @@ public:
 
 	/** Returns true if the controller received async responses and updated its cache since the last time the function was called, then clear the flag. */
 	bool GetAndClearDiscoveryUpdateFlag();
-
-	/** Returns the latest list of server known to this controller. Ensure to call TickServersAndSessionsDiscovery() periodically. */
-	const TArray<FConcertServerInfo>& GetServers() const { return Servers; }
-
-	/** Returns the latest list of active sessions known to this controller. Ensure to call TickServersAndSessionsDiscovery() periodically.*/
-	const TArray<TSharedPtr<FActiveSessionInfo>>& GetActiveSessions() const { return ActiveSessions; }
-
-	/** Returns the latest list of archived sessions known to this controller. Ensure to call TickServersAndSessionsDiscovery() periodically.*/
-	const TArray<TSharedPtr<FArchivedSessionInfo>>& GetArchivedSessions() const { return ArchivedSessions; }
-
+	
 	/** Returns the latest list of clients corresponding to the session known to this controller. Ensure to call TickClientsDiscovery() periodically. */
 	const TArray<FConcertSessionClientInfo>& GetClients(const FGuid& ServerAdminEndpointId, const FGuid& SessionId) const;
-
-	/** Returns the active sessions info corresponding to the specified parameters. Used to display the sessions details. */
-	const FConcertSessionInfo* GetActiveSessionInfo(const FGuid& AdminEndpoint, const FGuid& SessionId) const;
-
-	/** Returns the archived sessions info corresponding to the specified parameters. Used to display the sessions details. */
-	const FConcertSessionInfo* GetArchivedSessionInfo(const FGuid& AdminEndpoint, const FGuid& SessionId) const;
-
-	void CreateSession(const FGuid& ServerAdminEndpointId, const FString& SessionName);
-	void ArchiveSession(const FGuid& ServerAdminEndpointId, const FGuid& SessionId, const FString& ArchiveName, const FConcertSessionFilter& SessionFilter);
-	void RestoreSession(const FGuid& ServerAdminEndpointId, const FGuid& SessionId, const FString& RestoredName, const FConcertSessionFilter& SessionFilter);
+	
 	void JoinSession(const FGuid& ServerAdminEndpointId, const FGuid& SessionId);
-	void RenameActiveSession(const FGuid& ServerAdminEndpointId, const FGuid& SessionId, const FString& NewName);
-	void RenameArchivedSession(const FGuid& ServerAdminEndpointId, const FGuid& SessionId, const FString& NewName);
-	bool CanRenameActiveSession(const FGuid& ServerAdminEndpointId, const FGuid& SessionId) const;
-	bool CanRenameArchivedSession(const FGuid& ServerAdminEndpointId, const FGuid& SessionId) const;
-	void DeleteActiveSession(const FGuid& ServerAdminEndpointId, const FGuid& SessionId);
-	void DeleteArchivedSession(const FGuid& ServerAdminEndpointId, const FGuid& SessionId);
-	bool CanDeleteActiveSession(const FGuid& ServerAdminEndpointId, const FGuid& SessionId) const;
-	bool CanDeleteArchivedSession(const FGuid& ServerAdminEndpointId, const FGuid& SessionId) const;
-
+	
 	// The 2 functions below are used to prevent fast UI transition with the 'no session' panel flashing when the session list is empty.
 	bool HasReceivedInitialSessionList() const { return bInitialActiveSessionQueryResponded && bInitialArchivedSessionQueryResponded; }
 	bool IsCreatingSession() const { return CreateSessionRequests.Num() > 0 || ExpectedSessionsToDiscover.Num() > 0; }
+	
+	IConcertClientPtr GetConcertClient() const { return ConcertClient; }
+	
+	//~ Begin IConcertSessionBrowserController Interface
+	virtual TArray<FConcertServerInfo> GetServers() const override { return Servers; }
+	
+	virtual TArray<TSharedPtr<FActiveSessionInfo>> GetActiveSessions() const override
+	{
+		TArray<TSharedPtr<FActiveSessionInfo>> Result;
+		Algo::Transform(ActiveSessions, Result, [](TSharedPtr<FClientActiveSessionInfo> Info) { return Info; });
+		return Result;
+	}
+	
+	virtual TArray<TSharedPtr<FArchivedSessionInfo>> GetArchivedSessions() const override
+	{
+		TArray<TSharedPtr<FArchivedSessionInfo>> Result;
+		Algo::Transform(ArchivedSessions, Result, [](TSharedPtr<FClientArchivedSessionInfo> Info) { return Info; });
+		return Result;
+	}
+	
+	virtual const FConcertSessionInfo* GetActiveSessionInfo(const FGuid& AdminEndpoint, const FGuid& SessionId) const override;
+	virtual const FConcertSessionInfo* GetArchivedSessionInfo(const FGuid& AdminEndpoint, const FGuid& SessionId) const override;
+
+	virtual void CreateSession(const FGuid& ServerAdminEndpointId, const FString& SessionName) override;
+	virtual void ArchiveSession(const FGuid& ServerAdminEndpointId, const FGuid& SessionId, const FString& ArchiveName, const FConcertSessionFilter& SessionFilter) override;
+	virtual void RestoreSession(const FGuid& ServerAdminEndpointId, const FGuid& SessionId, const FString& RestoredName, const FConcertSessionFilter& SessionFilter) override;
+	virtual void RenameActiveSession(const FGuid& ServerAdminEndpointId, const FGuid& SessionId, const FString& NewName) override;
+	virtual void RenameArchivedSession(const FGuid& ServerAdminEndpointId, const FGuid& SessionId, const FString& NewName) override;
+	virtual bool CanRenameActiveSession(const FGuid& ServerAdminEndpointId, const FGuid& SessionId) const override;
+	virtual bool CanRenameArchivedSession(const FGuid& ServerAdminEndpointId, const FGuid& SessionId) const override;
+	virtual void DeleteActiveSession(const FGuid& ServerAdminEndpointId, const FGuid& SessionId) override;
+	virtual void DeleteArchivedSession(const FGuid& ServerAdminEndpointId, const FGuid& SessionId) override;
+	virtual bool CanDeleteActiveSession(const FGuid& ServerAdminEndpointId, const FGuid& SessionId) const override;
+	virtual bool CanDeleteArchivedSession(const FGuid& ServerAdminEndpointId, const FGuid& SessionId) const override;
+	//~ End IConcertSessionBrowserController Interface
 
 private:
+	
 	/** Hold information about a session created by this client, not yet 'discovered' by a 'list session' query, but expected to be soon. */
 	struct FPendingSessionDiscovery
 	{
@@ -289,6 +151,28 @@ private:
 		FString SessionName;
 	};
 
+	struct FMonitoredSessionInfo
+	{
+		TSharedRef<FActiveSessionInfo> Session;
+		FAsyncRequest ListClientRequest;
+
+		FMonitoredSessionInfo(TSharedRef<FActiveSessionInfo> Session)
+			: Session(Session)
+		{}
+	};
+
+	struct FClientActiveSessionInfo : public FActiveSessionInfo
+	{
+		/** Raised when the UI and the cache values may be out of sync if a rename failed (UI assumed it succeeded) */
+		bool bSessionNameDirty = false;
+	};
+
+	struct FClientArchivedSessionInfo : public FArchivedSessionInfo
+	{
+		/** Raised when the UI and the cache values may be out of sync if a rename failed (UI assumed it succeeded) */
+		bool bSessionNameDirty = false; 
+	};
+
 	void UpdateSessionsAsync();
 	void UpdateActiveSessionsAsync(const FConcertServerInfo& ServerInfo);
 	void UpdateArchivedSessionsAsync(const FConcertServerInfo& ServerInfo);
@@ -304,7 +188,6 @@ private:
 	void OnActiveSessionListDirty() { ++SessionListVersion; bCacheUpdated = true; } // This will force the UI to refresh its list.
 	void OnArchivedSessionListDirty() { ++SessionListVersion; bCacheUpdated = true; } // This will force the UI to refresh its list.
 
-private:
 	// Holds a concert client instance.
 	IConcertClientPtr ConcertClient;
 
@@ -314,11 +197,11 @@ private:
 
 	// The cached lists.
 	TArray<FConcertServerInfo> Servers;
-	TArray<TSharedPtr<FActiveSessionInfo>> ActiveSessions;
-	TArray<TSharedPtr<FArchivedSessionInfo>> ArchivedSessions;
+	TArray<TSharedPtr<FClientActiveSessionInfo>> ActiveSessions;
+	TArray<TSharedPtr<FClientArchivedSessionInfo>> ArchivedSessions;
 
 	// The session for which the clients are monitored. UI only monitor client of 1 session at the time.
-	TSharedPtr<FActiveSessionInfo> ClientMonitoredSession;
+	TOptional<FMonitoredSessionInfo> ClientMonitoredSession;
 
 	// Holds the version of data cached by the controller. The version is updated when an async response is received and implies a change in the cached values.
 	uint32 ServerListVersion = 0;
@@ -334,7 +217,7 @@ private:
 };
 
 
-FConcertBrowserController::FConcertBrowserController(IConcertClientPtr InConcertClient)
+FConcertClientSessionBrowserController::FConcertClientSessionBrowserController(IConcertClientPtr InConcertClient)
 {
 	check(InConcertClient.IsValid()); // Don't expect this class to be instantiated if the concert client is not available.
 	check(InConcertClient->IsConfigured()); // Expected to be done by higher level code.
@@ -350,7 +233,7 @@ FConcertBrowserController::FConcertBrowserController(IConcertClientPtr InConcert
 	TickServersAndSessionsDiscovery();
 }
 
-FConcertBrowserController::~FConcertBrowserController()
+FConcertClientSessionBrowserController::~FConcertClientSessionBrowserController()
 {
 	if (ConcertClient.IsValid())
 	{
@@ -358,12 +241,12 @@ FConcertBrowserController::~FConcertBrowserController()
 	}
 }
 
-const TArray<FConcertSessionClientInfo>& FConcertBrowserController::GetClients(const FGuid& AdminEndpoint, const FGuid& SessionId) const
+const TArray<FConcertSessionClientInfo>& FConcertClientSessionBrowserController::GetClients(const FGuid& AdminEndpoint, const FGuid& SessionId) const
 {
 	// If a session clients are monitored and the list is cached
-	if (ClientMonitoredSession.IsValid() && ClientMonitoredSession->ServerInfo.AdminEndpointId == AdminEndpoint && ClientMonitoredSession->SessionInfo.SessionId == SessionId)
+	if (ClientMonitoredSession.IsSet() && ClientMonitoredSession->Session->ServerInfo.AdminEndpointId == AdminEndpoint && ClientMonitoredSession->Session->SessionInfo.SessionId == SessionId)
 	{
-		return ClientMonitoredSession->Clients; // Returns the list retrieved by the last TickClientsDiscovery().
+		return ClientMonitoredSession->Session->Clients; // Returns the list retrieved by the last TickClientsDiscovery().
 	}
 
 	// Returns an empty list for now. We expect the caller to call TickClientsDiscovery() periodically to maintain the session client list.
@@ -371,9 +254,16 @@ const TArray<FConcertSessionClientInfo>& FConcertBrowserController::GetClients(c
 	return EmptyClientList;
 }
 
-const FConcertSessionInfo* FConcertBrowserController::GetActiveSessionInfo(const FGuid& AdminEndpoint, const FGuid& SessionId) const
+void FConcertClientSessionBrowserController::JoinSession(const FGuid& ServerAdminEndpointId, const FGuid& SessionId)
+{
+	// On success: The client joins the session and SConcertBrowser::HandleSessionConnectionChanged() will transit the UI to the SActiveSession.
+	// On failure: An async notification banner will be displayer to the user.
+	ConcertClient->JoinSession(ServerAdminEndpointId, SessionId);
+}
+
+const FConcertSessionInfo* FConcertClientSessionBrowserController::GetActiveSessionInfo(const FGuid& AdminEndpoint, const FGuid& SessionId) const
 {
-	const TSharedPtr<FActiveSessionInfo>* SessionInfo = ActiveSessions.FindByPredicate([&AdminEndpoint, &SessionId](const TSharedPtr<FActiveSessionInfo>& MatchCandidate)
+	const TSharedPtr<FClientActiveSessionInfo>* SessionInfo = ActiveSessions.FindByPredicate([&AdminEndpoint, &SessionId](const TSharedPtr<FActiveSessionInfo>& MatchCandidate)
 	{
 		return MatchCandidate->ServerInfo.AdminEndpointId == AdminEndpoint && MatchCandidate->SessionInfo.SessionId == SessionId;
 	});
@@ -381,9 +271,9 @@ const FConcertSessionInfo* FConcertBrowserController::GetActiveSessionInfo(const
 	return SessionInfo != nullptr ? &(*SessionInfo)->SessionInfo : nullptr;
 }
 
-const FConcertSessionInfo* FConcertBrowserController::GetArchivedSessionInfo(const FGuid& AdminEndpoint, const FGuid& SessionId) const
+const FConcertSessionInfo* FConcertClientSessionBrowserController::GetArchivedSessionInfo(const FGuid& AdminEndpoint, const FGuid& SessionId) const
 {
-	const TSharedPtr<FArchivedSessionInfo>* SessionInfo = ArchivedSessions.FindByPredicate([&AdminEndpoint, &SessionId](const TSharedPtr<FArchivedSessionInfo>& MatchCandidate)
+	const TSharedPtr<FClientArchivedSessionInfo>* SessionInfo = ArchivedSessions.FindByPredicate([&AdminEndpoint, &SessionId](const TSharedPtr<FArchivedSessionInfo>& MatchCandidate)
 	{
 		return MatchCandidate->ServerInfo.AdminEndpointId == AdminEndpoint && MatchCandidate->SessionInfo.SessionId == SessionId;
 	});
@@ -391,7 +281,7 @@ const FConcertSessionInfo* FConcertBrowserController::GetArchivedSessionInfo(con
 	return SessionInfo != nullptr ? &(*SessionInfo)->SessionInfo : nullptr;
 }
 
-void FConcertBrowserController::CreateSession(const FGuid& ServerAdminEndpointId, const FString& SessionName)
+void FConcertClientSessionBrowserController::CreateSession(const FGuid& ServerAdminEndpointId, const FString& SessionName)
 {
 	FAsyncRequest& CreateRequest = CreateSessionRequests.AddDefaulted_GetRef();
 	TWeakPtr<uint8, ESPMode::ThreadSafe> CreateRequestExecutionToken = CreateRequest.ResetExecutionToken();
@@ -416,7 +306,7 @@ void FConcertBrowserController::CreateSession(const FGuid& ServerAdminEndpointId
 	});
 }
 
-void FConcertBrowserController::ArchiveSession(const FGuid& ServerAdminEndpointId, const FGuid& SessionId, const FString& ArchiveName, const FConcertSessionFilter& SessionFilter)
+void FConcertClientSessionBrowserController::ArchiveSession(const FGuid& ServerAdminEndpointId, const FGuid& SessionId, const FString& ArchiveName, const FConcertSessionFilter& SessionFilter)
 {
 	// On success, an archived is created and TickServersAndSessionsDiscovery() will eventually discover it.
 	// On failure: An async notification banner will be displayer to the user.
@@ -427,7 +317,7 @@ void FConcertBrowserController::ArchiveSession(const FGuid& ServerAdminEndpointI
 	ConcertClient->ArchiveSession(ServerAdminEndpointId, ArchiveSessionArgs);
 }
 
-void FConcertBrowserController::RestoreSession(const FGuid& ServerAdminEndpointId, const FGuid& SessionId, const FString& RestoredName, const FConcertSessionFilter& SessionFilter)
+void FConcertClientSessionBrowserController::RestoreSession(const FGuid& ServerAdminEndpointId, const FGuid& SessionId, const FString& RestoredName, const FConcertSessionFilter& SessionFilter)
 {
 	FString ArchivedSessionName;
 	if (const FConcertSessionInfo* SessionInfo = GetArchivedSessionInfo(ServerAdminEndpointId, SessionId))
@@ -528,17 +418,10 @@ void FConcertBrowserController::RestoreSession(const FGuid& ServerAdminEndpointI
 	FSlateApplication::Get().AddWindow(NewWindow, true);
 }
 
-void FConcertBrowserController::JoinSession(const FGuid& ServerAdminEndpointId, const FGuid& SessionId)
-{
-	// On success: The client joins the session and SConcertBrowser::HandleSessionConnectionChanged() will transit the UI to the SActiveSession.
-	// On failure: An async notification banner will be displayer to the user.
-	ConcertClient->JoinSession(ServerAdminEndpointId, SessionId);
-}
-
-void FConcertBrowserController::RenameActiveSession(const FGuid& ServerAdminEndpointId, const FGuid& SessionId, const FString& NewName)
+void FConcertClientSessionBrowserController::RenameActiveSession(const FGuid& ServerAdminEndpointId, const FGuid& SessionId, const FString& NewName)
 {
 	// Find the currently cached session info.
-	const TSharedPtr<FActiveSessionInfo>* SessionInfo = ActiveSessions.FindByPredicate([&ServerAdminEndpointId, &SessionId](const TSharedPtr<FActiveSessionInfo>& MatchCandidate)
+	const TSharedPtr<FClientActiveSessionInfo>* SessionInfo = ActiveSessions.FindByPredicate([&ServerAdminEndpointId, &SessionId](const TSharedPtr<FActiveSessionInfo>& MatchCandidate)
 	{
 		return MatchCandidate->ServerInfo.AdminEndpointId == ServerAdminEndpointId && MatchCandidate->SessionInfo.SessionId == SessionId;
 	});
@@ -554,10 +437,10 @@ void FConcertBrowserController::RenameActiveSession(const FGuid& ServerAdminEndp
 	});
 }
 
-void FConcertBrowserController::RenameArchivedSession(const FGuid& ServerAdminEndpointId, const FGuid& SessionId, const FString& NewName)
+void FConcertClientSessionBrowserController::RenameArchivedSession(const FGuid& ServerAdminEndpointId, const FGuid& SessionId, const FString& NewName)
 {
 	// Find the currently cached session info.
-	const TSharedPtr<FArchivedSessionInfo>* SessionInfo = ArchivedSessions.FindByPredicate([&ServerAdminEndpointId, &SessionId](const TSharedPtr<FArchivedSessionInfo>& MatchCandidate)
+	const TSharedPtr<FClientArchivedSessionInfo>* SessionInfo = ArchivedSessions.FindByPredicate([&ServerAdminEndpointId, &SessionId](const TSharedPtr<FArchivedSessionInfo>& MatchCandidate)
 	{
 		return MatchCandidate->ServerInfo.AdminEndpointId == ServerAdminEndpointId && MatchCandidate->SessionInfo.SessionId == SessionId;
 	});
@@ -573,33 +456,33 @@ void FConcertBrowserController::RenameArchivedSession(const FGuid& ServerAdminEn
 	});
 }
 
-bool FConcertBrowserController::CanRenameActiveSession(const FGuid& ServerAdminEndpointId, const FGuid& SessionId) const
+bool FConcertClientSessionBrowserController::CanRenameActiveSession(const FGuid& ServerAdminEndpointId, const FGuid& SessionId) const
 {
 	return CanDeleteActiveSession(ServerAdminEndpointId, SessionId); // Rename requires the same permission than delete.
 }
 
-bool FConcertBrowserController::CanRenameArchivedSession(const FGuid& ServerAdminEndpointId, const FGuid& SessionId) const
+bool FConcertClientSessionBrowserController::CanRenameArchivedSession(const FGuid& ServerAdminEndpointId, const FGuid& SessionId) const
 {
 	return CanDeleteArchivedSession(ServerAdminEndpointId, SessionId); // Rename requires the same permission than delete.
 }
 
-void FConcertBrowserController::DeleteActiveSession(const FGuid& ServerAdminEndpointId, const FGuid& SessionId)
+void FConcertClientSessionBrowserController::DeleteActiveSession(const FGuid& ServerAdminEndpointId, const FGuid& SessionId)
 {
 	// On success, an session is deleted and TickServersAndSessionsDiscovery() will eventually notice it.
 	// On failure: An async notification banner will be displayer to the user.
 	ConcertClient->DeleteSession(ServerAdminEndpointId, SessionId);
 }
 
-void FConcertBrowserController::DeleteArchivedSession(const FGuid& ServerAdminEndpointId, const FGuid& SessionId)
+void FConcertClientSessionBrowserController::DeleteArchivedSession(const FGuid& ServerAdminEndpointId, const FGuid& SessionId)
 {
 	// On success, an archive is deleted and TickServersAndSessionsDiscovery() will eventually notice it.
 	// On failure: An async notification banner will be displayer to the user.
 	ConcertClient->DeleteSession(ServerAdminEndpointId, SessionId);
 }
 
-bool FConcertBrowserController::CanDeleteActiveSession(const FGuid& ServerAdminEndpointId, const FGuid& SessionId) const
+bool FConcertClientSessionBrowserController::CanDeleteActiveSession(const FGuid& ServerAdminEndpointId, const FGuid& SessionId) const
 {
-	const TSharedPtr<FActiveSessionInfo>* SessionInfo = ActiveSessions.FindByPredicate([&ServerAdminEndpointId, &SessionId](const TSharedPtr<FActiveSessionInfo>& MatchCandidate)
+	const TSharedPtr<FClientActiveSessionInfo>* SessionInfo = ActiveSessions.FindByPredicate([&ServerAdminEndpointId, &SessionId](const TSharedPtr<FActiveSessionInfo>& MatchCandidate)
 	{
 		return MatchCandidate->ServerInfo.AdminEndpointId == ServerAdminEndpointId && MatchCandidate->SessionInfo.SessionId == SessionId;
 	});
@@ -608,9 +491,9 @@ bool FConcertBrowserController::CanDeleteActiveSession(const FGuid& ServerAdminE
 	return SessionInfo == nullptr ? false : ConcertClient->IsOwnerOf((*SessionInfo)->SessionInfo);
 }
 
-bool FConcertBrowserController::CanDeleteArchivedSession(const FGuid& ServerAdminEndpointId, const FGuid& SessionId) const
+bool FConcertClientSessionBrowserController::CanDeleteArchivedSession(const FGuid& ServerAdminEndpointId, const FGuid& SessionId) const
 {
-	const TSharedPtr<FArchivedSessionInfo>* SessionInfo = ArchivedSessions.FindByPredicate([&ServerAdminEndpointId, &SessionId](const TSharedPtr<FArchivedSessionInfo>& MatchCandidate)
+	const TSharedPtr<FClientArchivedSessionInfo>* SessionInfo = ArchivedSessions.FindByPredicate([&ServerAdminEndpointId, &SessionId](const TSharedPtr<FArchivedSessionInfo>& MatchCandidate)
 	{
 		return MatchCandidate->ServerInfo.AdminEndpointId == ServerAdminEndpointId && MatchCandidate->SessionInfo.SessionId == SessionId;
 	});
@@ -619,7 +502,7 @@ bool FConcertBrowserController::CanDeleteArchivedSession(const FGuid& ServerAdmi
 	return SessionInfo == nullptr ? false : ConcertClient->IsOwnerOf((*SessionInfo)->SessionInfo);
 }
 
-TPair<uint32, uint32> FConcertBrowserController::TickServersAndSessionsDiscovery()
+TPair<uint32, uint32> FConcertClientSessionBrowserController::TickServersAndSessionsDiscovery()
 {
 	// Fire new async queries to poll the sessions on all known servers.
 	UpdateSessionsAsync();
@@ -628,7 +511,7 @@ TPair<uint32, uint32> FConcertBrowserController::TickServersAndSessionsDiscovery
 	return MakeTuple(ServerListVersion, SessionListVersion);
 }
 
-uint32 FConcertBrowserController::TickClientsDiscovery(const FGuid& ServerAdminEndpointId, const FGuid& SessionId)
+uint32 FConcertClientSessionBrowserController::TickClientsDiscovery(const FGuid& ServerAdminEndpointId, const FGuid& SessionId)
 {
 	// Fire new async queries to poll clients for the selected server/session pair.
 	UpdateClientsAsync(ServerAdminEndpointId, SessionId);
@@ -637,14 +520,14 @@ uint32 FConcertBrowserController::TickClientsDiscovery(const FGuid& ServerAdminE
 	return ClientListVersion;
 }
 
-bool FConcertBrowserController::GetAndClearDiscoveryUpdateFlag()
+bool FConcertClientSessionBrowserController::GetAndClearDiscoveryUpdateFlag()
 {
 	bool bOldCacheChanged = bCacheUpdated; // This flag is raised every time an async response updates the cached data.
 	bCacheUpdated = false;
 	return bOldCacheChanged;
 }
 
-void FConcertBrowserController::UpdateSessionsAsync()
+void FConcertClientSessionBrowserController::UpdateSessionsAsync()
 {
 	// Get the list of known servers.
 	TArray<FConcertServerInfo> OnlineServers = ConcertClient->GetKnownServers();
@@ -712,7 +595,7 @@ void FConcertBrowserController::UpdateSessionsAsync()
 	Servers = MoveTemp(OnlineServers);
 }
 
-void FConcertBrowserController::UpdateActiveSessionsAsync(const FConcertServerInfo& ServerInfo)
+void FConcertClientSessionBrowserController::UpdateActiveSessionsAsync(const FConcertServerInfo& ServerInfo)
 {
 	// Check if a request is already pulling active session from this server.
 	FAsyncRequest& ListActiveSessionAsyncRequest = ActiveSessionRequests.FindOrAdd(ServerInfo.InstanceInfo.InstanceId);
@@ -745,7 +628,7 @@ void FConcertBrowserController::UpdateActiveSessionsAsync(const FConcertServerIn
 				bInitialActiveSessionQueryResponded = true;
 
 				// Remove from the cache any session that were deleted from 'Server' since the last update. Find which one were renamed.
-				ActiveSessions.RemoveAll([this, &ServerInfo, &Response](const TSharedPtr<FActiveSessionInfo>& DiscardCandidate)
+				ActiveSessions.RemoveAll([this, &ServerInfo, &Response](const TSharedPtr<FClientActiveSessionInfo>& DiscardCandidate)
 				{
 					// If the candidate is owned by another server.
 					if (DiscardCandidate->ServerInfo.InstanceInfo.InstanceId != ServerInfo.InstanceInfo.InstanceId)
@@ -778,11 +661,11 @@ void FConcertBrowserController::UpdateActiveSessionsAsync(const FConcertServerIn
 				for (const FConcertSessionInfo& SessionInfo : Response.Sessions)
 				{
 					// Try to find the session in the list of active session.
-					if (!ActiveSessions.ContainsByPredicate([&SessionInfo](const TSharedPtr<FActiveSessionInfo>& MatchCandidate)
+					if (!ActiveSessions.ContainsByPredicate([&SessionInfo](const TSharedPtr<FClientActiveSessionInfo>& MatchCandidate)
 						{ return SessionInfo.ServerInstanceId == MatchCandidate->ServerInfo.InstanceInfo.InstanceId && SessionInfo.SessionId == MatchCandidate->SessionInfo.SessionId; }))
 					{
 						// This is a newly discovered session, add it to the list.
-						ActiveSessions.Add(MakeShared<FActiveSessionInfo>(FActiveSessionInfo{ServerInfo, SessionInfo}));
+						ActiveSessions.Add(MakeShared<FClientActiveSessionInfo>(FClientActiveSessionInfo{ServerInfo, SessionInfo}));
 						OnActiveSessionDiscovered(*ActiveSessions.Last());
 					}
 
@@ -799,7 +682,7 @@ void FConcertBrowserController::UpdateActiveSessionsAsync(const FConcertServerIn
 		});
 }
 
-void FConcertBrowserController::UpdateArchivedSessionsAsync(const FConcertServerInfo& ServerInfo)
+void FConcertClientSessionBrowserController::UpdateArchivedSessionsAsync(const FConcertServerInfo& ServerInfo)
 {
 	// Check if a request is already pulling archived sessions from this server.
 	FAsyncRequest& ListArchivedSessionAsyncRequest = ArchivedSessionRequests.FindOrAdd(ServerInfo.InstanceInfo.InstanceId);
@@ -829,7 +712,7 @@ void FConcertBrowserController::UpdateArchivedSessionsAsync(const FConcertServer
 				bInitialArchivedSessionQueryResponded = true;
 
 				// Remove from the cache archives that were deleted from 'Server' since the last update. Find which one were renamed.
-				ArchivedSessions.RemoveAll([this, &ServerInfo, &Response](const TSharedPtr<FArchivedSessionInfo>& DiscardCandidate)
+				ArchivedSessions.RemoveAll([this, &ServerInfo, &Response](const TSharedPtr<FClientArchivedSessionInfo>& DiscardCandidate)
 				{
 					// If the discard candidate is stored on another server.
 					if (DiscardCandidate->ServerInfo.InstanceInfo.InstanceId != ServerInfo.InstanceInfo.InstanceId)
@@ -866,7 +749,7 @@ void FConcertBrowserController::UpdateArchivedSessionsAsync(const FConcertServer
 						{ return ServerInfo.InstanceInfo.InstanceId == MatchCandidate->ServerInfo.InstanceInfo.InstanceId && SessionInfo.SessionId == MatchCandidate->SessionInfo.SessionId; }))
 					{
 						// This is a newly discovered archive, add it to the list.
-						ArchivedSessions.Add(MakeShared<FArchivedSessionInfo>(FArchivedSessionInfo{ServerInfo, SessionInfo}));
+						ArchivedSessions.Add(MakeShared<FClientArchivedSessionInfo>(FClientArchivedSessionInfo{ServerInfo, SessionInfo}));
 						OnArchivedSessionDiscovered(*ArchivedSessions.Last());
 					}
 				}
@@ -875,31 +758,37 @@ void FConcertBrowserController::UpdateArchivedSessionsAsync(const FConcertServer
 		});
 }
 
-void FConcertBrowserController::UpdateClientsAsync(const FGuid& ServerAdminEndpointId, const FGuid& SessionId)
+void FConcertClientSessionBrowserController::UpdateClientsAsync(const FGuid& ServerAdminEndpointId, const FGuid& SessionId)
 {
 	// The clients request is for a different session than the one cached by the last client update (The UI only shows 1 at the time).
-	if (ClientMonitoredSession.IsValid() && (ClientMonitoredSession->ServerInfo.AdminEndpointId != ServerAdminEndpointId || ClientMonitoredSession->SessionInfo.SessionId != SessionId))
+	if (ClientMonitoredSession.IsSet() && (ClientMonitoredSession->Session->ServerInfo.AdminEndpointId != ServerAdminEndpointId || ClientMonitoredSession->Session->SessionInfo.SessionId != SessionId))
 	{
 		// Reset the cache.
-		ClientMonitoredSession->Clients.Reset();
+		ClientMonitoredSession->Session->Clients.Reset();
 		ClientMonitoredSession->ListClientRequest.Cancel();
 		ClientMonitoredSession.Reset();
 		ClientListVersion = 0;
 	}
 
 	// Find the session requested by the user.
-	if (!ClientMonitoredSession.IsValid())
+	if (!ClientMonitoredSession.IsSet())
 	{
-		TSharedPtr<FActiveSessionInfo>* MatchEntry = ActiveSessions.FindByPredicate([&ServerAdminEndpointId, &SessionId](const TSharedPtr<FActiveSessionInfo>& MatchCandidate)
+		TSharedPtr<FClientActiveSessionInfo>* MatchEntry = ActiveSessions.FindByPredicate([&ServerAdminEndpointId, &SessionId](const TSharedPtr<FClientActiveSessionInfo>& MatchCandidate)
 		{
 			return MatchCandidate->ServerInfo.AdminEndpointId == ServerAdminEndpointId && MatchCandidate->SessionInfo.SessionId == SessionId;
 		});
 
-		// Update the active session.
-		ClientMonitoredSession = MatchEntry != nullptr ? *MatchEntry : nullptr;
+		if (MatchEntry != nullptr)
+		{
+			ClientMonitoredSession = FMonitoredSessionInfo { MatchEntry->ToSharedRef() };
+		}
+		else
+		{
+			ClientMonitoredSession.Reset();
+		}
 	}
 
-	if (!ClientMonitoredSession.IsValid() || ClientMonitoredSession->ListClientRequest.IsOngoing())
+	if (!ClientMonitoredSession.IsSet() || ClientMonitoredSession->ListClientRequest.IsOngoing())
 	{
 		// Don't stack another async request on top of the on-going one, TickClientsDiscovery() will eventually catch up any possibly missed info here.
 		return;
@@ -924,17 +813,17 @@ void FConcertBrowserController::UpdateClientsAsync(const FGuid& ServerAdminEndpo
 			// If the request succeeded.
 			if (Response.ResponseCode == EConcertResponseCode::Success)
 			{
-				if (ClientMonitoredSession->Clients.Num() != Response.SessionClients.Num())
+				if (ClientMonitoredSession->Session->Clients.Num() != Response.SessionClients.Num())
 				{
-					ClientMonitoredSession->Clients = Response.SessionClients;
-					ClientMonitoredSession->Clients.Sort(SortClientPredicate);
-					OnActiveSessionClientsUpdated(*ClientMonitoredSession);
+					ClientMonitoredSession->Session->Clients = Response.SessionClients;
+					ClientMonitoredSession->Session->Clients.Sort(SortClientPredicate);
+					OnActiveSessionClientsUpdated(*ClientMonitoredSession->Session);
 					return;
 				}
 				else if (Response.SessionClients.Num() == 0) // All existing clients disconnected.
 				{
-					ClientMonitoredSession->Clients.Reset();
-					OnActiveSessionClientsUpdated(*ClientMonitoredSession);
+					ClientMonitoredSession->Session->Clients.Reset();
+					OnActiveSessionClientsUpdated(*ClientMonitoredSession->Session);
 				}
 				else // Compare the old and the new list, both sorted by client endpoint id.
 				{
@@ -944,12 +833,12 @@ void FConcertBrowserController::UpdateClientsAsync(const FGuid& ServerAdminEndpo
 					int Index = 0;
 					for (const FConcertSessionClientInfo& Client : SortedClients)
 					{
-						if (ClientMonitoredSession->Clients[Index].ClientEndpointId != Client.ClientEndpointId || // Not the same client?
-						    ClientMonitoredSession->Clients[Index].ClientInfo != Client.ClientInfo)               // Client info was updated?
+						if (ClientMonitoredSession->Session->Clients[Index].ClientEndpointId != Client.ClientEndpointId || // Not the same client?
+						    ClientMonitoredSession->Session->Clients[Index].ClientInfo != Client.ClientInfo)               // Client info was updated?
 						{
 							// The two lists are not identical, don't bother finding the 'delta', refresh all clients.
-							ClientMonitoredSession->Clients = SortedClients;
-							OnActiveSessionClientsUpdated(*ClientMonitoredSession);
+							ClientMonitoredSession->Session->Clients = SortedClients;
+							OnActiveSessionClientsUpdated(*ClientMonitoredSession->Session);
 							break;
 						}
 						++Index;
@@ -959,2197 +848,588 @@ void FConcertBrowserController::UpdateClientsAsync(const FGuid& ServerAdminEndpo
 		});
 }
 
-
 /**
- * The type of row used in the session list view to display active or archived session.
+ * Widget displayed when discovering multi-user server(s) or session(s).
  */
-class SSessionRow : public SMultiColumnTableRow<TSharedPtr<FConcertSessionItem>>
+class SConcertDiscovery : public SCompoundWidget
 {
 public:
-	typedef TFunction<void(TSharedPtr<FConcertSessionItem>)> FDoubleClickFunc;
-	typedef TFunction<void(TSharedPtr<FConcertSessionItem>, const FString&)> FRenameFunc;
-
-	SLATE_BEGIN_ARGS(SSessionRow)
-		: _OnDoubleClickFunc()
-		, _OnRenameFunc()
-		, _HighlightText()
-		, _IsSelected(false)
+	SLATE_BEGIN_ARGS(SConcertDiscovery)
+		: _Text()
+		, _ThrobberVisibility(EVisibility::Visible)
+		, _ButtonVisibility(EVisibility::Visible)
+		, _IsButtonEnabled(true)
+		, _ButtonStyle(&FAppStyle::Get().GetWidgetStyle<FButtonStyle>("FlatButton"))
+		, _ButtonIcon()
+		, _ButtonText()
+		, _ButtonToolTip()
+		, _OnButtonClicked()
 	{
 	}
 
-	SLATE_ARGUMENT(FDoubleClickFunc, OnDoubleClickFunc)
-	SLATE_ARGUMENT(FRenameFunc, OnRenameFunc)
-	SLATE_ATTRIBUTE(FText, HighlightText)
-	SLATE_ATTRIBUTE(bool, IsSelected)
+	SLATE_ATTRIBUTE(FText, Text)
+	SLATE_ATTRIBUTE(EVisibility, ThrobberVisibility)
+	SLATE_ATTRIBUTE(EVisibility, ButtonVisibility)
+	SLATE_ATTRIBUTE(bool, IsButtonEnabled)
+	SLATE_STYLE_ARGUMENT(FButtonStyle, ButtonStyle)
+	SLATE_ATTRIBUTE(const FSlateBrush*, ButtonIcon)
+	SLATE_ATTRIBUTE(FText, ButtonText)
+	SLATE_ATTRIBUTE(FText, ButtonToolTip)
+	SLATE_EVENT( FOnClicked, OnButtonClicked)
 
-	SLATE_END_ARGS()
+	SLATE_END_ARGS();
 
 public:
-	void Construct(const FArguments& InArgs, TSharedPtr<FConcertSessionItem> InItem, const TSharedRef<STableViewBase>& InOwnerTableView, IConcertClientPtr InConcertClient);
-	virtual TSharedRef<SWidget> GenerateWidgetForColumn(const FName& ColumnName) override;
-	virtual FReply OnMouseButtonDoubleClick( const FGeometry& InMyGeometry, const FPointerEvent& InMouseEvent) override;
-
-	void OnSessionNameCommitted(const FText& NewSessionName, ETextCommit::Type CommitType);
+	void Construct(const FArguments& InArgs)
+	{
+		ChildSlot
+		[
+			SNew(SBox)
+			.HAlign(HAlign_Center)
+			.VAlign(VAlign_Center)
+			.Clipping(EWidgetClipping::ClipToBounds)
+			[
+				SNew(SVerticalBox)
+				+SVerticalBox::Slot()
+				.AutoHeight()
+				.HAlign(HAlign_Center)
+				[
+					SNew(SCircularThrobber)
+					.Visibility(InArgs._ThrobberVisibility)
+				]
 
-private:
-	void OnBeginEditingSessionName() { SessionNameText->EnterEditingMode(); }
-	bool OnValidatingSessionName(const FText& NewSessionName, FText& OutError);
+				+SVerticalBox::Slot()
+				.AutoHeight()
+				.HAlign(HAlign_Center)
+				[
+					SNew(STextBlock).Text(InArgs._Text).Justification(ETextJustify::Center)
+				]
 
-private:
-	TWeakPtr<FConcertSessionItem> Item;
-	FDoubleClickFunc DoubleClickFunc; // Invoked when the user double click on the row.
-	FRenameFunc RenameFunc; // Invoked when the user commit the session rename. (This will send the request to server)
-	TAttribute<FText> HighlightText;
-	TAttribute<bool> IsSelected;
-	TSharedPtr<SInlineEditableTextBlock> SessionNameText;
-	IConcertClientPtr ConcertClient;
+				+SVerticalBox::Slot()
+				.Padding(0, 4, 0, 0)
+				.AutoHeight()
+				.HAlign(HAlign_Center)
+				[
+					SNew(SButton)
+					.ForegroundColor(FSlateColor::UseForeground())
+					.ButtonStyle(InArgs._ButtonStyle)
+					.Visibility(InArgs._ButtonVisibility)
+					.HAlign(HAlign_Center)
+					.VAlign(VAlign_Center)
+					.IsEnabled(InArgs._IsButtonEnabled)
+					.OnClicked(InArgs._OnButtonClicked)
+					.ToolTipText(InArgs._ButtonToolTip)
+					.ContentPadding(FMargin(8, 4))
+					[
+						SNew(SHorizontalBox)
+						+SHorizontalBox::Slot()
+						.AutoWidth()
+						.VAlign(VAlign_Center)
+						.Padding(0, 0, 3, 0)
+						[
+							SNew(SImage).Image(InArgs._ButtonIcon)
+						]
+						+SHorizontalBox::Slot()
+						.AutoWidth()
+						.VAlign(VAlign_Bottom)
+						[
+							SNew(STextBlock).Text(InArgs._ButtonText)
+						]
+					]
+				]
+			]
+		];
+	}
 };
 
 
-void SSessionRow::Construct(const FArguments& InArgs, TSharedPtr<FConcertSessionItem> InItem, const TSharedRef<STableViewBase>& InOwnerTableView, IConcertClientPtr InConcertClient)
-{
-	Item = MoveTemp(InItem);
-	DoubleClickFunc = InArgs._OnDoubleClickFunc; // This function should join a session or add a row to restore an archive.
-	RenameFunc = InArgs._OnRenameFunc; // Function invoked to send a rename request to the server.
-	HighlightText = InArgs._HighlightText;
-	IsSelected = InArgs._IsSelected;
-	ConcertClient = InConcertClient;
-
-	// Construct base class
-	SMultiColumnTableRow<TSharedPtr<FConcertSessionItem>>::Construct(FSuperRowType::FArguments(), InOwnerTableView);
-
-	// Listen and handle rename request.
-	InItem->OnBeginEditSessionNameRequest.AddSP(this, &SSessionRow::OnBeginEditingSessionName);
-}
-
-TSharedRef<SWidget> SSessionRow::GenerateWidgetForColumn(const FName& ColumnName)
+/**
+ * Displayed when something is not available.
+ */
+class SConcertNoAvailability : public SCompoundWidget
 {
-	TSharedPtr<FConcertSessionItem> ItemPin = Item.Pin();
+public:
+	SLATE_BEGIN_ARGS(SConcertNoAvailability) : _Text() {}
+	SLATE_ATTRIBUTE(FText, Text)
+	SLATE_END_ARGS()
 
-	if (ColumnName == ConcertBrowserUtils::IconColName)
+public:
+	void Construct(const FArguments& InArgs)
 	{
-		return SNew(SBox)
-		.VAlign(VAlign_Center)
-		.HAlign(HAlign_Center)
-		.Padding(2)
-		.ToolTipText(ItemPin->Type == FConcertSessionItem::EType::ActiveSession ? LOCTEXT("ActiveIconTooltip", "Active session") : LOCTEXT("ArchivedIconTooltip", "Archived Session"))
+		ChildSlot
 		[
-			SNew(STextBlock)
-			.Font(FAppStyle::Get().GetFontStyle(ConcertBrowserUtils::IconColumnFontName))
-			.Text(ItemPin->Type == FConcertSessionItem::EType::ActiveSession ? FEditorFontGlyphs::Circle : FEditorFontGlyphs::Archive)
-			.ColorAndOpacity(ItemPin->Type == FConcertSessionItem::EType::ActiveSession ? FAppStyle::Get().GetWidgetStyle<FButtonStyle>("FlatButton.Success").Normal.TintColor : FSlateColor::UseSubduedForeground())
+			SNew(SConcertDiscovery) // Reuse this panel, but only show the message.
+				.Text(InArgs._Text)
+				.ThrobberVisibility(EVisibility::Collapsed)
+				.ButtonVisibility(EVisibility::Collapsed)
 		];
 	}
+};
 
-	bool bIsDefaultConfig = (ItemPin->Type == FConcertSessionItem::EType::ActiveSession && ItemPin->SessionName == ConcertClient->GetConfiguration()->DefaultSessionName && ItemPin->ServerName == ConcertClient->GetConfiguration()->DefaultServerURL);
-	FSlateFontInfo FontInfo;
-	FSlateColor FontColor;
-	if (ItemPin->Type == FConcertSessionItem::EType::ActiveSession)
-	{
-		FontColor = bIsDefaultConfig ? FSlateColor(FLinearColor::White) : FSlateColor(FLinearColor::White * 0.8f);
-		FontInfo = FAppStyle::Get().GetFontStyle("NormalFont");
-	}
-	else
-	{
-		FontColor = FSlateColor::UseSubduedForeground();
-		FontInfo = FCoreStyle::GetDefaultFontStyle("Italic", 9);
-	}
 
-	if (ColumnName == ConcertBrowserUtils::SessionColName)
-	{
-		return SNew(SBox)
-			.VAlign(VAlign_Center)
-			[
-				SAssignNew(SessionNameText, SInlineEditableTextBlock)
-				.Text_Lambda([this]() { return FText::AsCultureInvariant(Item.Pin()->SessionName); })
-				.HighlightText(HighlightText)
-				.OnTextCommitted(this, &SSessionRow::OnSessionNameCommitted)
-				.IsReadOnly(false)
-				.IsSelected(FIsSelected::CreateLambda([this]() { return IsSelected.Get(); }))
-				.OnVerifyTextChanged(this, &SSessionRow::OnValidatingSessionName)
-				.Font(FontInfo)
-				.ColorAndOpacity(FontColor)
-			];
-	}
+/**
+ * Enables the user to browse/search/filter/sort active and archived sessions, create new session,
+ * archive active sessions, restore archived sessions, join a session and open the settings dialog.
+ */
+class SConcertClientSessionBrowser : public SCompoundWidget
+{
+public:
+	
+	SLATE_BEGIN_ARGS(SConcertClientSessionBrowser) { }
+	SLATE_END_ARGS();
 
-	check(ColumnName == ConcertBrowserUtils::ServerColName);
+	/**
+	* Constructs the Browser.
+	* @param InArgs The Slate argument list.
+	* @param InConcertClient The concert client used to list, join, delete the sessions.
+	* @param[in,out] InSearchText The text to set in the search box and to remember (as output). Cannot be null.
+	*/
+	void Construct(const FArguments& InArgs, IConcertClientPtr InConcertClient, TSharedPtr<FText> InSearchText);
 
-	if (bIsDefaultConfig)
-	{
-		return SNew(SHorizontalBox)
-			+SHorizontalBox::Slot()
-			.AutoWidth()
-			.VAlign(VAlign_Center)
-			[
-				SNew(STextBlock)
-				.Text(FText::Format(INVTEXT("{0} * "), FText::AsCultureInvariant(ItemPin->ServerName)))
-				.HighlightText(HighlightText)
-				.Font(FontInfo)
-				.ColorAndOpacity(FontColor)
-			]
-			+SHorizontalBox::Slot()
-			.AutoWidth()
-			.VAlign(VAlign_Center)
-			[
-				SNew(STextBlock)
-				.Text(LOCTEXT("DefaultServerSession", "(Default Session/Server)"))
-				.HighlightText(HighlightText)
-				.Font(FCoreStyle::GetDefaultFontStyle("Regular", 9))
-				.ColorAndOpacity(FontColor)
-			]
-			+SHorizontalBox::Slot()
-			[
-				SNew(SSpacer)
-			]
-			+SHorizontalBox::Slot()
-			.AutoWidth()
-			.VAlign(VAlign_Center)
-			.HAlign(HAlign_Right)
-			[
-				ConcertBrowserUtils::MakeServerVersionIgnoredWidget(ItemPin->ServerFlags)
-			];
-	}
-	else
-	{
-		return SNew(SHorizontalBox)
-			+SHorizontalBox::Slot()
-			.AutoWidth()
-			.VAlign(VAlign_Center)
-			[
-				SNew(STextBlock)
-				.Text(FText::AsCultureInvariant(ItemPin->ServerName))
-				.HighlightText(HighlightText)
-				.Font(FontInfo)
-				.ColorAndOpacity(FontColor)
-			]
-			+SHorizontalBox::Slot()
-			[
-				SNew(SSpacer)
-			]
-			+SHorizontalBox::Slot()
-			.AutoWidth()
-			.VAlign(VAlign_Center)
-			.HAlign(HAlign_Right)
-			[
-				ConcertBrowserUtils::MakeServerVersionIgnoredWidget(ItemPin->ServerFlags)
-			];
-	}
-}
-
-FReply SSessionRow::OnMouseButtonDoubleClick(const FGeometry& InMyGeometry, const FPointerEvent& InMouseEvent)
-{
-	if (TSharedPtr<FConcertSessionItem> ItemPin = Item.Pin())
-	{
-		DoubleClickFunc(ItemPin);
-	}
-	return FReply::Handled();
-}
-
-bool SSessionRow::OnValidatingSessionName(const FText& NewSessionName, FText& OutError)
-{
-	OutError = ConcertSettingsUtils::ValidateSessionName(NewSessionName.ToString());
-	return OutError.IsEmpty();
-}
-
-void SSessionRow::OnSessionNameCommitted(const FText& NewSessionName, ETextCommit::Type CommitType)
-{
-	if (TSharedPtr<FConcertSessionItem> ItemPin = Item.Pin())
-	{
-		FString NewName = NewSessionName.ToString();
-		if (NewName != ItemPin->SessionName) // Was renamed?
-		{
-			if (ConcertSettingsUtils::ValidateSessionName(NewName).IsEmpty()) // Name is valid?
-			{
-				RenameFunc(ItemPin, NewName); // Send the rename request to the server. (Server may still refuse at this point)
-			}
-			else
-			{
-				// NOTE: Error are interactively detected and raised by OnValidatingSessionName()
-				FSlateApplication::Get().SetKeyboardFocus(SessionNameText);
-			}
-		}
-	}
-}
-
-
-/**
- * The type of row used in the session list view to edit a new session (the session name + server).
- */
-class SNewSessionRow : public SMultiColumnTableRow<TSharedPtr<FConcertSessionItem>>
-{
-public:
-	typedef TFunction<TPair<uint32, const TArray<FConcertServerInfo>&>()> FGetServersFunc;
-	typedef TFunction<void(const TSharedPtr<FConcertSessionItem>&)> FAcceptFunc; // Should remove the editable 'new' row and creates the sessions.
-	typedef TFunction<void(const TSharedPtr<FConcertSessionItem>&)> FDeclineFunc; // Should just remove the editable 'new' row.
-
-	SLATE_BEGIN_ARGS(SNewSessionRow) // Needed to use Args because Construct() is limited to 5 arguments and 6 were required.
-		: _GetServerFunc()
-		, _OnAcceptFunc()
-		, _OnDeclineFunc()
-		, _HighlightText()
-	{
-	}
-
-	SLATE_ARGUMENT(FGetServersFunc, GetServerFunc)
-	SLATE_ARGUMENT(FAcceptFunc, OnAcceptFunc)
-	SLATE_ARGUMENT(FDeclineFunc, OnDeclineFunc)
-	SLATE_ATTRIBUTE(FText, HighlightText)
-
-	SLATE_END_ARGS()
-
-public:
-	void Construct(const FArguments& InArgs, TSharedPtr<FConcertSessionItem> InItem, const TSharedRef<STableViewBase>& InOwnerTableView, IConcertClientPtr InConcertClient);
-	virtual TSharedRef<SWidget> GenerateWidgetForColumn(const FName& ColumnName) override;
 	virtual void Tick(const FGeometry& AllottedGeometry, const double InCurrentTime, const float InDeltaTime) override;
 
 private:
-	TSharedRef<SWidget> OnGenerateServersComboOptionWidget(TSharedPtr<FConcertServerInfo> Item);
-	TSharedRef<SWidget> MakeSelectedServerWidget();
-	FText GetSelectedServerText() const;
-	FText GetSelectedServerIgnoreVersionText() const;
-	FText GetSelectedServerIgnoreVersionTooltip() const;
-	FText GetServerDisplayName(const FString& ServerName) const;
-
-	void OnSessionNameChanged(const FText& NewName);
-	void OnSessionNameCommitted(const FText& NewText, ETextCommit::Type CommitType);
-	FReply OnAccept();
-	FReply OnDecline();
-	FReply OnKeyDownHandler(const FGeometry&, const FKeyEvent&); // Registered as handler to the editable text (vs. OnKeyDown() virtual method).
-
-	void UpdateServerList();
-
-private:
-	IConcertClientPtr ConcertClient;
-	TWeakPtr<FConcertSessionItem> Item; // Holds the new item to fill with session name and server.
-	TArray<TSharedPtr<FConcertServerInfo>> Servers; // Servers displayed in the server combo box.
-	TSharedPtr<SComboBox<TSharedPtr<FConcertServerInfo>>> ServersComboBox;
-	TSharedPtr<SEditableTextBox> EditableSessionName;
-	FGetServersFunc GetServersFunc;
-	FAcceptFunc AcceptFunc;
-	FDeclineFunc DeclineFunc;
-	TAttribute<FText> HighlightText;
-	uint32 ServerListVersion = 0;
-	bool bInitialFocusTaken = false;
-};
-
-void SNewSessionRow::Construct(const FArguments& InArgs, TSharedPtr<FConcertSessionItem> InItem, const TSharedRef<STableViewBase>& InOwnerTableView, IConcertClientPtr InConcertClient)
-{
-	check(InConcertClient.IsValid());
-
-	Item = MoveTemp(InItem);
-	ConcertClient = InConcertClient;
-	GetServersFunc = InArgs._GetServerFunc;
-	AcceptFunc = InArgs._OnAcceptFunc;
-	DeclineFunc = InArgs._OnDeclineFunc;
-	HighlightText = InArgs._HighlightText;
+	
+	// Layout the 'session|details' split view.
+	TSharedRef<SWidget> MakeBrowserContent(TSharedPtr<FText> InSearchText);
+	void ExtendControlButtons(FExtender& Extender);
+	void ExtendSessionContextMenu(TSharedPtr<FConcertSessionItem> Item, FExtender& Extender);
+	TSharedRef<SWidget> MakeUserAndSettings();
+	TSharedRef<SWidget> MakeOverlayedTableView(TSharedRef<SWidget> TableView);
+	
+	// Layouts the session detail panel.
+	TSharedRef<SWidget> MakeSessionDetails(TSharedPtr<FConcertSessionItem> Item);
+	TSharedRef<SWidget> MakeActiveSessionDetails(TSharedPtr<FConcertSessionItem> Item);
+	TSharedRef<SWidget> MakeArchivedSessionDetails(TSharedPtr<FConcertSessionItem> Item);
+	TSharedRef<ITableRow> OnGenerateClientRowWidget(TSharedPtr<FConcertSessionClientInfo> Item, const TSharedRef<STableViewBase>& OwnerTable);
+	void PopulateSessionInfoGrid(SGridPanel& Grid, const FConcertSessionInfo& SessionInfo);
 
-	// Construct base class
-	SMultiColumnTableRow<TSharedPtr<FConcertSessionItem>>::Construct(FSuperRowType::FArguments(), InOwnerTableView);
+	// The buttons in some overlays
+	bool IsLaunchServerButtonEnabled() const;
+	bool IsJoinButtonEnabled() const;
+	bool IsAutoJoinButtonEnabled() const;
+	bool IsCancelAutoJoinButtonEnabled() const;
+	FReply OnNewButtonClicked();
+	FReply OnLaunchServerButtonClicked();
+	FReply OnShutdownServerButtonClicked();
+	FReply OnAutoJoinButtonClicked();
+	FReply OnCancelAutoJoinButtonClicked();
 
-	// Fill the server combo.
-	UpdateServerList();
-}
+	// SConcertSessionBrowser extensions
+	FReply OnJoinButtonClicked();
+	void RequestJoinSession(const TSharedPtr<FConcertSessionItem>& ActiveItem);
 
-void SNewSessionRow::Tick(const FGeometry& AllottedGeometry, const double InCurrentTime, const float InDeltaTime)
-{
-	// Check if the server list displayed by the combo box should be updated.
-	if (GetServersFunc().Key != ServerListVersion)
-	{
-		UpdateServerList();
-	}
+	// Manipulates the sessions view (the array and the UI).
+	void OnSessionSelectionChanged(TSharedPtr<FConcertSessionItem> SelectedSession);
+	void OnSessionDoubleClicked(TSharedPtr<FConcertSessionItem> SelectedSession);
 
-	// Should give the focus to an editable text.
-	if (!bInitialFocusTaken)
-	{
-		bInitialFocusTaken = FSlateApplication::Get().SetKeyboardFocus(EditableSessionName.ToSharedRef());
-	}
-}
+	// Update server/session/clients lists.
+	EActiveTimerReturnType TickDiscovery(double InCurrentTime, float InDeltaTime);
+	void UpdateDiscovery();
+	void RefreshClientList(const TArray<FConcertSessionClientInfo>& LastestClientList);
 
-TSharedRef<SWidget> SNewSessionRow::GenerateWidgetForColumn(const FName& ColumnName)
-{
-	TSharedPtr<FConcertSessionItem> ItemPin = Item.Pin();
+private:
+	
+	// Gives access to the concert data (servers, sessions, clients, etc).
+	TSharedPtr<FConcertClientSessionBrowserController> Controller;
 
-	if (ColumnName == ConcertBrowserUtils::IconColName)
-	{
-		// 'New' icon
-		return SNew(SBox)
-			.VAlign(VAlign_Center)
-			.HAlign(HAlign_Center)
-			[
-				SNew(STextBlock)
-					.Font(FAppStyle::Get().GetFontStyle(ConcertBrowserUtils::IconColumnFontName))
-					.Text(FEditorFontGlyphs::Plus_Circle)
-			];
-	}
-	else if (ColumnName == ConcertBrowserUtils::SessionColName)
-	{
-		return SNew(SBox)
-			.VAlign(VAlign_Center)
-			.Padding(FMargin(0, 0, 2, 0))
-			[
-				SAssignNew(EditableSessionName, SEditableTextBox)
-				.HintText(LOCTEXT("EnterSessionNameHint", "Enter a session name"))
-				.OnTextCommitted(this, &SNewSessionRow::OnSessionNameCommitted)
-				.OnKeyDownHandler(this, &SNewSessionRow::OnKeyDownHandler)
-				.OnTextChanged(this, &SNewSessionRow::OnSessionNameChanged)
-			];
-	}
-	else
-	{
-		return SNew(SHorizontalBox)
+	TSharedPtr<SConcertSessionBrowser> SessionBrowser;
 
-			// 'Server' combo
-			+SHorizontalBox::Slot()
-			.AutoWidth()
-			.Padding(0, 1)
-			[
-				SAssignNew(ServersComboBox, SComboBox<TSharedPtr<FConcertServerInfo>>)
-				.OptionsSource(&Servers)
-				.OnGenerateWidget(this, &SNewSessionRow::OnGenerateServersComboOptionWidget)
-				.ToolTipText(LOCTEXT("SelectServerTooltip", "Select the server on which the session should be created"))
-				[
-					MakeSelectedServerWidget()
-				]
-			]
+	// Filtering.
+	bool bRefreshSessionFilter = true;
+	FString DefaultServerURL;
 
-			+SHorizontalBox::Slot()
-			.AutoWidth()
-			.Padding(2.0f)
-			.HAlign(HAlign_Left)
-			[
-				SNew(SUniformGridPanel)
-				.SlotPadding(FMargin(1.0f, 0.0f))
+	// Selected Session Details.
+	TSharedPtr<SBorder> SessionDetailsView;
+	TSharedPtr<SExpandableArea> DetailsArea;
+	TArray<TSharedPtr<FConcertSessionClientInfo>> Clients;
+	TSharedPtr<SExpandableArea> ClientsArea;
+	TSharedPtr<SListView<TSharedPtr<FConcertSessionClientInfo>>> ClientsView;
 
-				// 'Accept' button
-				+SUniformGridPanel::Slot(0, 0)
-				[
-					ConcertBrowserUtils::MakeIconButton(
-						TEXT("FlatButton.Success"),
-						FEditorFontGlyphs::Check,
-						LOCTEXT("CreateCheckIconTooltip", "Create the session"),
-						TAttribute<bool>::Create([this]() { return !EditableSessionName->GetText().IsEmpty(); }),
-						FOnClicked::CreateRaw(this, &SNewSessionRow::OnAccept),
-						FSlateColor(FLinearColor::White))
-				]
+	// Used to compare the version used by UI versus the version cached in the controller.
+	uint32 DisplayedSessionListVersion = 0;
+	uint32 DisplayedClientListVersion = 0;
+	uint32 ServerListVersion = 0;
+	bool bLocalServerRunning = false;
 
-				// 'Decline' button
-				+SUniformGridPanel::Slot(1, 0)
-				[
-					ConcertBrowserUtils::MakeIconButton(
-						TEXT("FlatButton.Danger"),
-						FEditorFontGlyphs::Times,
-						LOCTEXT("CancelIconTooltip", "Cancel"),
-						true, // Always enabled.
-						FOnClicked::CreateRaw(this, &SNewSessionRow::OnDecline),
-						FSlateColor(FLinearColor::White))
-				]
-			];
-	}
-}
+	TSharedPtr<SWidget> ServerDiscoveryPanel; // Displayed until a server is found.
+	TSharedPtr<SWidget> SessionDiscoveryPanel; // Displayed until a session is found.
+	TSharedPtr<SWidget> NoSessionSelectedPanel; // Displays 'select a session to view details' message in the details section.
+	TSharedPtr<SWidget> NoSessionDetailsPanel; // Displays 'no details available' message in the details section.
+	TSharedPtr<SWidget> NoClientPanel; // Displays the 'no client connected' message in Clients expendable area.
+};
 
-TSharedRef<SWidget> SNewSessionRow::OnGenerateServersComboOptionWidget(TSharedPtr<FConcertServerInfo> ServerItem)
+void SConcertClientSessionBrowser::Construct(const FArguments& InArgs, IConcertClientPtr InConcertClient, TSharedPtr<FText> InSearchText)
 {
-	bool bIsDefaultServer = ServerItem->ServerName == ConcertClient->GetConfiguration()->DefaultServerURL;
-
-	FText Tooltip;
-	if (bIsDefaultServer)
-	{
-		Tooltip = LOCTEXT("DefaultServerTooltip", "Default Configured Server");
-	}
-	else if (ServerItem->ServerName == FPlatformProcess::ComputerName())
-	{
-		Tooltip = LOCTEXT("LocalServerTooltip", "Local Server Running on This Computer");
-	}
-	else
+	if (!InConcertClient.IsValid())
 	{
-		Tooltip = LOCTEXT("OnlineServerTooltip", "Online Server");
+		return; // Don't build the UI if ConcertClient is not available.
 	}
 
-	return SNew(SHorizontalBox)
-		.ToolTipText(Tooltip)
-
-		+SHorizontalBox::Slot()
-		.AutoWidth()
-		[
-			SNew(STextBlock)
-			.Font(bIsDefaultServer ? FAppStyle::Get().GetFontStyle("BoldFont") : FAppStyle::Get().GetFontStyle("NormalFont"))
-			.Text(GetServerDisplayName(ServerItem->ServerName))
-		]
-		+SHorizontalBox::Slot()
-		.AutoWidth()
-		[
-			ConcertBrowserUtils::MakeServerVersionIgnoredWidget(ServerItem->ServerFlags)
-		];
-}
-
-void SNewSessionRow::UpdateServerList()
-{
-	// Remember the currently selected item (if any).
-	TSharedPtr<FConcertServerInfo> SelectedItem = ServersComboBox->GetSelectedItem(); // Instance in current list.
-
-	// Clear the current list. The list is rebuilt from scratch.
-	Servers.Reset();
+	Controller = MakeShared<FConcertClientSessionBrowserController>(InConcertClient);
 
-	TSharedPtr<FConcertServerInfo> LocalServerInfo;
-	TSharedPtr<FConcertServerInfo> DefaultServerInfo;
-	TSharedPtr<FConcertServerInfo> SelectedServerInfo; // Instance in the new list.
-
-	const UConcertClientConfig* ConcertClientConfig = ConcertClient->GetConfiguration();
-
-	// Update the list version.
-	ServerListVersion = GetServersFunc().Key;
+	// Displayed if no server is available.
+	ServerDiscoveryPanel = SNew(SConcertDiscovery)
+		.Text(LOCTEXT("LookingForServer", "Looking for Multi-User Servers..."))
+		.Visibility_Lambda([this]() { return Controller->GetServers().Num() == 0 ? EVisibility::Visible : EVisibility::Hidden; })
+		.IsButtonEnabled(this, &SConcertClientSessionBrowser::IsLaunchServerButtonEnabled)
+		.ButtonStyle(FAppStyle::Get(), "FlatButton.Default")
+		.ButtonIcon(FConcertFrontendStyle::Get()->GetBrush("Concert.NewServer.Small"))
+		.ButtonText(LOCTEXT("LaunchLocalServer", "Launch a Server"))
+		.ButtonToolTip(LOCTEXT("LaunchServerTooltip", "Launch a Multi-User server on your computer unless one is already running.\nThe editor UDP messaging settings will be passed to the launching server.\nThe server will use the server port found in Multi-User client settings if non 0 or use the editor udp messaging port number + 1, if non 0."))
+		.OnButtonClicked(this, &SConcertClientSessionBrowser::OnLaunchServerButtonClicked);
 
-	// Convert to shared ptr (slate needs that) and find if the latest list contains a default/local server.
-	for (const FConcertServerInfo& ServerInfo : GetServersFunc().Value)
+	// Controls the text displayed in the 'No sessions' panel.
+	auto GetNoSessionText = [this]()
 	{
-		TSharedPtr<FConcertServerInfo> ComboItem = MakeShared<FConcertServerInfo>(ServerInfo);
-
-		if (ComboItem->ServerName == ConcertClientConfig->DefaultServerURL) // Default server is deemed more important than local server to display the icon aside the server.
-		{
-			DefaultServerInfo = ComboItem;
-		}
-		else if (ComboItem->ServerName == FPlatformProcess::ComputerName())
-		{
-			LocalServerInfo = ComboItem;
-		}
-
-		if (SelectedItem.IsValid() && SelectedItem->ServerName == ComboItem->ServerName)
+		if (!Controller->HasReceivedInitialSessionList())
 		{
-			SelectedServerInfo = ComboItem; // Preserve user selection using the new instance.
+			return LOCTEXT("LookingForSession", "Looking for Multi-User Sessions...");
 		}
 
-		Servers.Emplace(MoveTemp(ComboItem));
-	}
+		return Controller->GetActiveSessions().Num() == 0 && Controller->GetArchivedSessions().Num() == 0 ?
+			LOCTEXT("NoSessionAvailable", "No Sessions Available") :
+			LOCTEXT("AllSessionsFilteredOut", "No Sessions Match the Filters\nChange Your Filter to View Sessions");
+	};
 
-	// Sort the server list alphabetically.
-	Servers.Sort([](const TSharedPtr<FConcertServerInfo>& Lhs, const TSharedPtr<FConcertServerInfo>& Rhs) { return Lhs->ServerName < Rhs->ServerName; });
+	// Displayed when discovering session or if no session is available.
+	SessionDiscoveryPanel = SNew(SConcertDiscovery)
+		.Text_Lambda(GetNoSessionText)
+		.Visibility_Lambda([this]() { return Controller->GetServers().Num() > 0 && SessionBrowser->GetSessions().Num() == 0 && !Controller->IsCreatingSession() ? EVisibility::Visible : EVisibility::Hidden; })
+		.ThrobberVisibility_Lambda([this]() { return !Controller->HasReceivedInitialSessionList() ? EVisibility::Visible : EVisibility::Collapsed; })
+		.ButtonVisibility_Lambda([this]() { return Controller->HasReceivedInitialSessionList() && Controller->GetActiveSessions().Num() == 0 && Controller->GetArchivedSessions().Num() == 0 ? EVisibility::Visible : EVisibility::Collapsed; })
+		.ButtonStyle(FAppStyle::Get(), "FlatButton.Default")
+		.ButtonIcon(FConcertFrontendStyle::Get()->GetBrush("Concert.NewSession.Small"))
+		.ButtonText(LOCTEXT("CreateSession", "Create Session"))
+		.ButtonToolTip(LOCTEXT("CreateSessionTooltip", "Create a new session"))
+		.OnButtonClicked(FOnClicked::CreateSP(this, &SConcertClientSessionBrowser::OnNewButtonClicked));
 
-	// If a server is running on this machine, put it first in the list.
-	if (LocalServerInfo.IsValid() && Servers[0] != LocalServerInfo)
-	{
-		Servers.Remove(LocalServerInfo); // Keep sort order.
-		Servers.Insert(LocalServerInfo, 0);
-	}
+	// Displayed when the selected session client view is empty (no client to display).
+	NoClientPanel = SNew(SConcertNoAvailability)
+		.Text(LOCTEXT("NoClientAvailable", "No Connected Clients"))
+		.Visibility_Lambda([this]() { return Clients.Num() == 0 ? EVisibility::Visible : EVisibility::Hidden; });
 
-	// If a 'default server' is configured and available, put it first in the list. (Possibly overruling the local one)
-	if (DefaultServerInfo.IsValid() && Servers[0] != DefaultServerInfo)
-	{
-		Servers.Remove(DefaultServerInfo); // Keep sort order.
-		Servers.Insert(DefaultServerInfo, 0);
-	}
+	// Displayed as details when no session is selected. (No session selected or the selected session doesn't have any)
+	NoSessionSelectedPanel = SNew(SConcertNoAvailability)
+		.Text(LOCTEXT("NoSessionSelected", "Select a Session to View Details"));
 
-	// If a server was selected and is still in the updated list.
-	if (SelectedServerInfo.IsValid())
-	{
-		// Preserve user selection.
-		ServersComboBox->SetSelectedItem(SelectedServerInfo);
-	}
-	else if (Servers.Num())
-	{
-		// Select the very first item in the list which is most likely be the default or the local server as they were put first above.
-		ServersComboBox->SetSelectedItem(Servers[0]);
-	}
-	else // Server list is empty.
-	{
-		ServersComboBox->ClearSelection();
-		Servers.Reset();
-	}
+	// Displayed as details when the selected session has not specific details to display.
+	NoSessionDetailsPanel = SNew(SConcertNoAvailability)
+		.Text(LOCTEXT("NoSessionDetails", "The Selected Session Has No Details"));
 
-	ServersComboBox->RefreshOptions();
-}
+	// List used in details panel to display clients connected to an active session.
+	ClientsView = SNew(SListView<TSharedPtr<FConcertSessionClientInfo>>)
+		.ListItemsSource(&Clients)
+		.OnGenerateRow(this, &SConcertClientSessionBrowser::OnGenerateClientRowWidget)
+		.SelectionMode(ESelectionMode::Single)
+		.AllowOverscroll(EAllowOverscroll::No);
 
-TSharedRef<SWidget> SNewSessionRow::MakeSelectedServerWidget()
-{
-	return SNew(SHorizontalBox)
-		+SHorizontalBox::Slot()
-		.AutoWidth()
-		[
-			SNew(STextBlock)
-			.Text(this, &SNewSessionRow::GetSelectedServerText)
-			.HighlightText(HighlightText)
-		]
-		+SHorizontalBox::Slot()
-		.AutoWidth()
-		.Padding(2, 0, 0, 0)
-		[
-			SNew(STextBlock)
-			.Font(FAppStyle::Get().GetFontStyle("FontAwesome.9"))
-			.Text(this, &SNewSessionRow::GetSelectedServerIgnoreVersionText)
-			.ToolTipText(this, &SNewSessionRow::GetSelectedServerIgnoreVersionTooltip)
-		];
-}
+	ChildSlot
+	[
+		MakeBrowserContent(InSearchText)
+	];
 
-FText SNewSessionRow::GetSelectedServerText() const
-{
-	TSharedPtr<FConcertServerInfo> SelectedServer = ServersComboBox->GetSelectedItem();
-	if (SelectedServer.IsValid())
-	{
-		return GetServerDisplayName(SelectedServer->ServerName);
-	}
-	return LOCTEXT("SelectAServer", "Select a Server");
-}
+	// Create a timer to periodically poll the server for sessions and session clients at a lower frequency than the normal tick.
+	RegisterActiveTimer(1.0f, FWidgetActiveTimerDelegate::CreateSP(this, &SConcertClientSessionBrowser::TickDiscovery));
 
-FText SNewSessionRow::GetServerDisplayName(const FString& ServerName) const
-{
-	bool bIsDefaultServer = ServerName == ConcertClient->GetConfiguration()->DefaultServerURL;
-	if (bIsDefaultServer)
-	{
-		return FText::Format(LOCTEXT("DefaultServer", "{0} (Default)"), FText::FromString(FPlatformProcess::ComputerName()));
-	}
-	else if (ServerName == FPlatformProcess::ComputerName())
-	{
-		return FText::Format(LOCTEXT("MyComputer", "{0} (My Computer)"), FText::FromString(FPlatformProcess::ComputerName()));
-	}
-	return FText::FromString(ServerName);
+	bLocalServerRunning = IMultiUserClientModule::Get().IsConcertServerRunning();
 }
 
-FText SNewSessionRow::GetSelectedServerIgnoreVersionText() const
+TSharedRef<SWidget> SConcertClientSessionBrowser::MakeBrowserContent(TSharedPtr<FText> InSearchText)
 {
-	if (ServersComboBox->GetSelectedItem() && (ServersComboBox->GetSelectedItem()->ServerFlags & EConcertServerFlags::IgnoreSessionRequirement) != EConcertServerFlags::None)
-	{
-		return FEditorFontGlyphs::Exclamation_Triangle;
-	}
-	return FText();
-}
+	return SNew(SBox)
+		.HAlign(HAlign_Fill)
+		.VAlign(VAlign_Fill)
+		[
+			// Splitter upper part displaying the available sessions/(server).
+			SNew(SSplitter)
+			.Orientation(Orient_Vertical)
+			.MinimumSlotHeight(80.0f) // Prevent widgets from overlapping.
+			+SSplitter::Slot()
+			.Value(0.6)
+			[
+				SAssignNew(SessionBrowser, SConcertSessionBrowser, Controller.ToSharedRef(), InSearchText)
+					.ExtendSessionTable(this, &SConcertClientSessionBrowser::MakeOverlayedTableView)
+					.ExtendControllButtons(this, &SConcertClientSessionBrowser::ExtendControlButtons)
+					.ExtendSessionContextMenu(this, &SConcertClientSessionBrowser::ExtendSessionContextMenu)
+					.RightOfSearchBar()
+					[
+						MakeUserAndSettings()
+					]
+					.OnSessionClicked_Raw(this, &SConcertClientSessionBrowser::OnSessionSelectionChanged)
+					.OnSessionDoubleClicked_Raw(this, &SConcertClientSessionBrowser::OnSessionDoubleClicked)
+					.OnRequestedDeleteSession_Lambda([this](auto) { UpdateDiscovery(); /* Don't wait up to 1s, kick discovery right now */ })
+			]
 
-FText SNewSessionRow::GetSelectedServerIgnoreVersionTooltip() const
-{
-	if (ServersComboBox->GetSelectedItem() && (ServersComboBox->GetSelectedItem()->ServerFlags & EConcertServerFlags::IgnoreSessionRequirement) != EConcertServerFlags::None)
-	{
-		return ConcertBrowserUtils::GetServerVersionIgnoredTooltip();
-	}
-	return FText();
-}
-
-FReply SNewSessionRow::OnAccept()
-{
-	// Read the session name given by the user.
-	TSharedPtr<FConcertSessionItem> ItemPin = Item.Pin();
-	FString NewSessionName = EditableSessionName->GetText().ToString();
-
-	FText InvalidNameErrorMsg = ConcertSettingsUtils::ValidateSessionName(NewSessionName);
-	if (InvalidNameErrorMsg.IsEmpty()) // Name is valid?
-	{
-		ItemPin->SessionName = EditableSessionName->GetText().ToString();
-		ItemPin->ServerName = ServersComboBox->GetSelectedItem()->ServerName;
-		ItemPin->ServerAdminEndpointId = ServersComboBox->GetSelectedItem()->AdminEndpointId;
-		AcceptFunc(ItemPin); // Delegate to create the session.
-	}
-	else
-	{
-		EditableSessionName->SetError(InvalidNameErrorMsg);
-		FSlateApplication::Get().SetKeyboardFocus(EditableSessionName);
-	}
-
-	return FReply::Handled();
-}
-
-FReply SNewSessionRow::OnDecline()
-{
-	DeclineFunc(Item.Pin()); // Decline the creation and remove the row from the table.
-	return FReply::Handled();
-}
-
-void SNewSessionRow::OnSessionNameChanged(const FText& NewName)
-{
-	EditableSessionName->SetError(ConcertSettingsUtils::ValidateSessionName(NewName.ToString()));
-}
-
-void SNewSessionRow::OnSessionNameCommitted(const FText& NewText, ETextCommit::Type CommitType)
-{
-	if (CommitType == ETextCommit::OnEnter)
-	{
-		OnAccept(); // Create the session.
-	}
-}
-
-FReply SNewSessionRow::OnKeyDownHandler(const FGeometry&, const FKeyEvent& KeyEvent)
-{
-	// NOTE: This is invoked when the editable text field has the focus.
-	return KeyEvent.GetKey() == EKeys::Escape ? OnDecline() : FReply::Unhandled();
-}
-
-
-/**
- * The type of row used in the session list view to create and archive or restore a session (edit the information required for the operation).
- */
-class SSaveRestoreSessionRow : public SMultiColumnTableRow<TSharedPtr<FConcertSessionItem>>
-{
-public:
-	typedef TFunction<void(TSharedPtr<FConcertSessionItem>, const FString&)> FAcceptFunc; // Should remove the editable row and save or restore the session.
-	typedef TFunction<void(TSharedPtr<FConcertSessionItem>)> FDeclineFunc; // Should only remove the editable row from the table.
-
-	SLATE_BEGIN_ARGS(SSaveRestoreSessionRow)
-		: _OnAcceptFunc()
-		, _OnDeclineFunc()
-		, _HighlightText()
-	{
-	}
-
-	SLATE_ARGUMENT(FAcceptFunc, OnAcceptFunc)
-	SLATE_ARGUMENT(FDeclineFunc, OnDeclineFunc)
-	SLATE_ATTRIBUTE(FText, HighlightText)
-
-	SLATE_END_ARGS()
-
-public:
-	void Construct(const FArguments& InArgs, TSharedPtr<FConcertSessionItem> Node, const TSharedRef<STableViewBase>& InOwnerTableView);
-	virtual TSharedRef<SWidget> GenerateWidgetForColumn(const FName& ColumnName) override;
-	virtual void Tick(const FGeometry& AllottedGeometry, const double InCurrentTime, const float InDeltaTime) override;
-
-	void OnSessionNameChanged(const FText& NewName);
-	void OnSessionNameCommitted(const FText& NewText, ETextCommit::Type CommitType);
-	FReply OnAccept();
-	FReply OnDecline();
-	FReply OnKeyDownHandler(const FGeometry&, const FKeyEvent&); // Registered as handler to the editable text (vs. OnKeyDown() virtual method for this widget).
-
-	// Override SMultiColumnTableRow functions to ensure a wire is drawn between the item to restore and the editable row to link them together.
-	virtual TBitArray<> GetWiresNeededByDepth() const override;
-	virtual bool IsLastChild() const override { return true; }
-	virtual int32 DoesItemHaveChildren() const override { return 0; }
-	virtual bool IsItemExpanded() const override { return false; }
-
-	FText GetDefaultName(const FConcertSessionItem& Item) const; // Generates a default name for an archive or a restored session.
-
-private:
-	TWeakPtr<FConcertSessionItem> Item;
-	TSharedPtr<SEditableTextBox> EditableSessionName;
-	FAcceptFunc AcceptFunc;
-	FDeclineFunc DeclineFunc;
-	TAttribute<FText> HighlightText;
-	bool bInitialFocusTaken = false;
-};
-
-
-void SSaveRestoreSessionRow::Construct(const FArguments& InArgs, TSharedPtr<FConcertSessionItem> InNode, const TSharedRef<STableViewBase>& InOwnerTableView)
-{
-	Item = MoveTemp(InNode);
-	AcceptFunc = InArgs._OnAcceptFunc;
-	DeclineFunc = InArgs._OnDeclineFunc;
-	HighlightText = InArgs._HighlightText;
-
-	// Construct base class
-	SMultiColumnTableRow<TSharedPtr<FConcertSessionItem>>::Construct(FSuperRowType::FArguments(), InOwnerTableView);
-}
-
-TBitArray<> SSaveRestoreSessionRow::GetWiresNeededByDepth() const
-{
-	TBitArray<> Bits;
-	Bits.Add(false);
-	return Bits;
-}
-
-FText SSaveRestoreSessionRow::GetDefaultName(const FConcertSessionItem& InItem) const
-{
-	if (InItem.Type == FConcertSessionItem::EType::SaveSession)
-	{
-		return FText::Format(LOCTEXT("DefaultName", "{0}.{1}"), FText::FromString(InItem.SessionName), FText::FromString(FDateTime::UtcNow().ToString()));
-	}
-
-	// Supposing the name of the archive has the dates as suffix, like SessionXYZ.2019.03.13-19.39.12, then extracts SessionXYZ
-	static FRegexPattern Pattern(TEXT(R"((.*)\.\d+\.\d+\.\d+\-\d+\.\d+\.\d+$)"));
-	FRegexMatcher Matcher(Pattern, InItem.SessionName);
-	if (Matcher.FindNext())
-	{
-		return FText::FromString(Matcher.GetCaptureGroup(1));
-	}
-
-	return FText::FromString(InItem.SessionName);
-}
-
-TSharedRef<SWidget> SSaveRestoreSessionRow::GenerateWidgetForColumn(const FName& ColumnName)
-{
-	TSharedPtr<FConcertSessionItem> ItemPin = Item.Pin();
-
-	if (ColumnName == ConcertBrowserUtils::IconColName)
-	{
-		return SNew(SHorizontalBox)
-			+SHorizontalBox::Slot()
-			.AutoWidth()
-			.Padding(8, 0)
-			[
-				SNew(SExpanderArrow, SharedThis(this))
-				.StyleSet(&FAppStyle::Get())
-				.ShouldDrawWires(true)
-			];
-	}
-	else if (ColumnName == ConcertBrowserUtils::SessionColName)
-	{
-		return SNew(SHorizontalBox)
-
-			// 'Restore as/Save as' text
-			+SHorizontalBox::Slot()
-			.VAlign(VAlign_Center)
-			.AutoWidth()
-			.Padding(2.0, 0.0)
-			[
-				SNew(STextBlock).Text(ItemPin->Type == FConcertSessionItem::EType::RestoreSession ? LOCTEXT("RestoreAs", "Restore as:") : LOCTEXT("ArchiveAs", "Archive as:"))
-			]
-
-			// Editable text.
-			+SHorizontalBox::Slot()
-			.VAlign(VAlign_Center)
-			.Padding(FMargin(0, 0, 2, 0))
-			[
-				SAssignNew(EditableSessionName, SEditableTextBox)
-				.HintText(ItemPin->Type == FConcertSessionItem::EType::RestoreSession ? LOCTEXT("RestoreSessionHint", "Enter a session name") : LOCTEXT("ArchivSessionHint", "Enter an archive name"))
-				.OnTextCommitted(this, &SSaveRestoreSessionRow::OnSessionNameCommitted)
-				.OnKeyDownHandler(this, &SSaveRestoreSessionRow::OnKeyDownHandler)
-				.OnTextChanged(this, &SSaveRestoreSessionRow::OnSessionNameChanged)
-				.Text(GetDefaultName(*ItemPin))
-				.SelectAllTextWhenFocused(true)
-			];
-	}
-	else
-	{
-		check(ColumnName == ConcertBrowserUtils::ServerColName);
-
-		// 'Server' text block.
-		return SNew(SHorizontalBox)
-		+SHorizontalBox::Slot()
-		.AutoWidth()
-		[
-			SNew(SBox)
-			.VAlign(VAlign_Center)
-			[
-				// Server.
-				SNew(STextBlock)
-				.Text(FText::FromString(ItemPin->ServerName))
-				.HighlightText(HighlightText)
-			]
-		]
-
-		+SHorizontalBox::Slot()
-		.AutoWidth()
-		.Padding(2.0f)
-		.HAlign(HAlign_Left)
-		[
-			SNew(SUniformGridPanel)
-			.SlotPadding(FMargin(1.0f, 0.0f))
-
-			// 'Accept' button
-			+SUniformGridPanel::Slot(0, 0)
-			[
-				ConcertBrowserUtils::MakeIconButton(
-					TEXT("FlatButton.Success"),
-					FEditorFontGlyphs::Check,
-					ItemPin->Type == FConcertSessionItem::EType::RestoreSession ? LOCTEXT("RestoreCheckIconTooltip", "Restore the session") : LOCTEXT("ArchiveCheckIconTooltip", "Archive the session"),
-					TAttribute<bool>::Create([this]() { return !EditableSessionName->GetText().IsEmpty(); }), // Enabled?
-					FOnClicked::CreateRaw(this, &SSaveRestoreSessionRow::OnAccept),
-					FSlateColor(FLinearColor::White))
-			]
-
-			// 'Cancel' button
-			+SUniformGridPanel::Slot(1, 0)
-			[
-				ConcertBrowserUtils::MakeIconButton(
-					TEXT("FlatButton.Danger"),
-					FEditorFontGlyphs::Times,
-					LOCTEXT("CancelTooltip", "Cancel"),
-					true, // Enabled?
-					FOnClicked::CreateRaw(this, &SSaveRestoreSessionRow::OnDecline),
-					FSlateColor(FLinearColor::White))
-			]
-		];
-	}
-}
-
-void SSaveRestoreSessionRow::Tick(const FGeometry& AllottedGeometry, const double InCurrentTime, const float InDeltaTime)
-{
-	// Should give the focus to an editable text.
-	if (!bInitialFocusTaken)
-	{
-		bInitialFocusTaken = FSlateApplication::Get().SetKeyboardFocus(EditableSessionName.ToSharedRef());
-	}
-}
-
-void SSaveRestoreSessionRow::OnSessionNameChanged(const FText& NewName)
-{
-	EditableSessionName->SetError(ConcertSettingsUtils::ValidateSessionName(NewName.ToString()));
-}
-
-void SSaveRestoreSessionRow::OnSessionNameCommitted(const FText& NewText, ETextCommit::Type CommitType)
-{
-	TSharedPtr<FConcertSessionItem> ItemPin = Item.Pin();
-	if (CommitType == ETextCommit::Type::OnEnter)
-	{
-		OnAccept();
-	}
-}
-
-FReply SSaveRestoreSessionRow::OnAccept()
-{
-	// Read the session name given by the user.
-	TSharedPtr<FConcertSessionItem> ItemPin = Item.Pin();
-	FString Name = EditableSessionName->GetText().ToString(); // Archive name or restored session name.
-
-	// Ensure the user provided a name.
-	FText InvalidNameErrorMsg = ConcertSettingsUtils::ValidateSessionName(Name);
-	if (InvalidNameErrorMsg.IsEmpty()) // Name is valid?
-	{
-		AcceptFunc(ItemPin, Name); // Delegate archiving/restoring operation.
-	}
-	else
-	{
-		EditableSessionName->SetError(InvalidNameErrorMsg);
-		FSlateApplication::Get().SetKeyboardFocus(EditableSessionName);
-	}
-
-	return FReply::Handled();
-}
-
-FReply SSaveRestoreSessionRow::OnDecline()
-{
-	DeclineFunc(Item.Pin()); // Remove the save/restore editable row.
-	return FReply::Handled();
-}
-
-FReply SSaveRestoreSessionRow::OnKeyDownHandler(const FGeometry&, const FKeyEvent& KeyEvent)
-{
-	// NOTE: This handler is to capture the 'Escape' while the text field has focus.
-	return KeyEvent.GetKey() == EKeys::Escape ? OnDecline() : FReply::Unhandled();
-}
-
-
-/**
- * Widget displayed when discovering multi-user server(s) or session(s).
- */
-class SConcertDiscovery : public SCompoundWidget
-{
-public:
-	SLATE_BEGIN_ARGS(SConcertDiscovery)
-		: _Text()
-		, _ThrobberVisibility(EVisibility::Visible)
-		, _ButtonVisibility(EVisibility::Visible)
-		, _IsButtonEnabled(true)
-		, _ButtonStyle(&FAppStyle::Get().GetWidgetStyle<FButtonStyle>("FlatButton"))
-		, _ButtonIcon()
-		, _ButtonText()
-		, _ButtonToolTip()
-		, _OnButtonClicked()
-	{
-	}
-
-	SLATE_ATTRIBUTE(FText, Text)
-	SLATE_ATTRIBUTE(EVisibility, ThrobberVisibility)
-	SLATE_ATTRIBUTE(EVisibility, ButtonVisibility)
-	SLATE_ATTRIBUTE(bool, IsButtonEnabled)
-	SLATE_STYLE_ARGUMENT(FButtonStyle, ButtonStyle)
-	SLATE_ATTRIBUTE(const FSlateBrush*, ButtonIcon)
-	SLATE_ATTRIBUTE(FText, ButtonText)
-	SLATE_ATTRIBUTE(FText, ButtonToolTip)
-	SLATE_EVENT( FOnClicked, OnButtonClicked)
-
-	SLATE_END_ARGS();
-
-public:
-	void Construct(const FArguments& InArgs)
-	{
-		ChildSlot
-		[
-			SNew(SBox)
-			.HAlign(HAlign_Center)
-			.VAlign(VAlign_Center)
-			.Clipping(EWidgetClipping::ClipToBounds)
+			// Session details.
+			+SSplitter::Slot()
+			.Value(0.4)
 			[
-				SNew(SVerticalBox)
-				+SVerticalBox::Slot()
-				.AutoHeight()
-				.HAlign(HAlign_Center)
-				[
-					SNew(SCircularThrobber)
-					.Visibility(InArgs._ThrobberVisibility)
-				]
-
-				+SVerticalBox::Slot()
-				.AutoHeight()
-				.HAlign(HAlign_Center)
-				[
-					SNew(STextBlock).Text(InArgs._Text).Justification(ETextJustify::Center)
-				]
-
-				+SVerticalBox::Slot()
-				.Padding(0, 4, 0, 0)
-				.AutoHeight()
-				.HAlign(HAlign_Center)
+				SAssignNew(SessionDetailsView, SBorder)
+				.BorderImage(FAppStyle::Get().GetBrush("ToolPanel.GroupBorder"))
 				[
-					SNew(SButton)
-					.ForegroundColor(FSlateColor::UseForeground())
-					.ButtonStyle(InArgs._ButtonStyle)
-					.Visibility(InArgs._ButtonVisibility)
-					.HAlign(HAlign_Center)
-					.VAlign(VAlign_Center)
-					.IsEnabled(InArgs._IsButtonEnabled)
-					.OnClicked(InArgs._OnButtonClicked)
-					.ToolTipText(InArgs._ButtonToolTip)
-					.ContentPadding(FMargin(8, 4))
-					[
-						SNew(SHorizontalBox)
-						+SHorizontalBox::Slot()
-						.AutoWidth()
-						.VAlign(VAlign_Center)
-						.Padding(0, 0, 3, 0)
-						[
-							SNew(SImage).Image(InArgs._ButtonIcon)
-						]
-						+SHorizontalBox::Slot()
-						.AutoWidth()
-						.VAlign(VAlign_Bottom)
-						[
-							SNew(STextBlock).Text(InArgs._ButtonText)
-						]
-					]
+					NoSessionSelectedPanel.ToSharedRef()
 				]
 			]
 		];
-	}
-};
-
-
-/**
- * Displayed when something is not available.
- */
-class SConcertNoAvailability : public SCompoundWidget
-{
-public:
-	SLATE_BEGIN_ARGS(SConcertNoAvailability) : _Text() {}
-	SLATE_ATTRIBUTE(FText, Text)
-	SLATE_END_ARGS()
-
-public:
-	void Construct(const FArguments& InArgs)
-	{
-		ChildSlot
-		[
-			SNew(SConcertDiscovery) // Reuse this panel, but only show the message.
-				.Text(InArgs._Text)
-				.ThrobberVisibility(EVisibility::Collapsed)
-				.ButtonVisibility(EVisibility::Collapsed)
-		];
-	}
-};
-
-
-/**
- * Enables the user to browse/search/filter/sort active and archived sessions, create new session,
- * archive active sessions, restore archived sessions, join a session and open the settings dialog.
- */
-class SConcertSessionBrowser : public SCompoundWidget
-{
-public:
-	SLATE_BEGIN_ARGS(SConcertSessionBrowser) { }
-	SLATE_END_ARGS();
-
-	/**
-	* Constructs the Browser.
-	* @param InArgs The Slate argument list.
-	* @param InConcertClient The concert client used to list, join, delete the sessions.
-	* @param[in,out] InSearchText The text to set in the search box and to remember (as output). Cannot be null.
-	*/
-	void Construct(const FArguments& InArgs, IConcertClientPtr InConcertClient, TSharedPtr<FText> InSearchText);
-
-	virtual void Tick(const FGeometry& AllottedGeometry, const double InCurrentTime, const float InDeltaTime) override;
-	virtual FReply OnKeyDown(const FGeometry& MyGeometry, const FKeyEvent& InKeyEvent) override;
-
-private:
-	// Layout the 'session|details' split view.
-	TSharedRef<SWidget> MakeBrowserContent();
-
-	// Layout the sessions view and controls.
-	TSharedRef<SWidget> MakeControlBar();
-	TSharedRef<SWidget> MakeButtonBar();
-	TSharedRef<SWidget> MakeSessionTableView();
-	TSharedRef<SWidget> MakeSessionViewOptionsBar();
-	TSharedRef<ITableRow> OnGenerateSessionRowWidget(TSharedPtr<FConcertSessionItem> Item, const TSharedRef<STableViewBase>& OwnerTable);
-
-	// Layouts the session detail panel.
-	TSharedRef<SWidget> MakeSessionDetails(TSharedPtr<FConcertSessionItem> Item);
-	TSharedRef<SWidget> MakeActiveSessionDetails(TSharedPtr<FConcertSessionItem> Item);
-	TSharedRef<SWidget> MakeArchivedSessionDetails(TSharedPtr<FConcertSessionItem> Item);
-	TSharedRef<ITableRow> OnGenerateClientRowWidget(TSharedPtr<FConcertSessionClientInfo> Item, const TSharedRef<STableViewBase>& OwnerTable);
-	void PopulateSessionInfoGrid(SGridPanel& Grid, const FConcertSessionInfo& SessionInfo);
-
-	// Creates row widgets for session list view, validates user inputs and forward user requests for processing to a delegate function implemented by this class.
-	TSharedRef<ITableRow> MakeActiveSessionRowWidget(const TSharedPtr<FConcertSessionItem>& ActiveItem, const TSharedRef<STableViewBase>& OwnerTable);
-	TSharedRef<ITableRow> MakeArchivedSessionRowWidget(const TSharedPtr<FConcertSessionItem>& ArchivedItem, const TSharedRef<STableViewBase>& OwnerTable);
-	TSharedRef<ITableRow> MakeNewSessionRowWidget(const TSharedPtr<FConcertSessionItem>& NewItem, const TSharedRef<STableViewBase>& OwnerTable);
-	TSharedRef<ITableRow> MakeRestoreSessionRowWidget(const TSharedPtr<FConcertSessionItem>& RestoreItem, const TSharedRef<STableViewBase>& OwnerTable);
-	TSharedRef<ITableRow> MakeSaveSessionRowWidget(const TSharedPtr<FConcertSessionItem>& ArchivedItem, const TSharedRef<STableViewBase>& OwnerTable);
-
-	// Creates the contextual menu when right clicking a session list view row.
-	TSharedPtr<SWidget> MakeContextualMenu();
-
-	// The buttons above the session view.
-	bool IsNewButtonEnabled() const;
-	bool IsJoinButtonEnabled() const;
-	bool IsRestoreButtonEnabled() const;
-	bool IsArchiveButtonEnabled() const;
-	bool IsRenameButtonEnabled() const;
-	bool IsDeleteButtonEnabled() const;
-	bool IsLaunchServerButtonEnabled() const;
-	bool IsAutoJoinButtonEnabled() const;
-	bool IsCancelAutoJoinButtonEnabled() const;
-	FReply OnNewButtonClicked();
-	FReply OnJoinButtonClicked();
-	FReply OnRestoreButtonClicked();
-	FReply OnArchiveButtonClicked();
-	FReply OnDeleteButtonClicked();
-	FReply OnLaunchServerButtonClicked();
-	FReply OnShutdownServerButtonClicked();
-	FReply OnAutoJoinButtonClicked();
-	FReply OnCancelAutoJoinButtonClicked();
-	void OnBeginEditingSessionName(TSharedPtr<FConcertSessionItem> Item);
-
-	// Manipulates the sessions view (the array and the UI).
-	void OnSessionSelectionChanged(TSharedPtr<FConcertSessionItem> SelectedSession, ESelectInfo::Type SelectInfo);
-	void InsertNewSessionEditableRow();
-	void InsertRestoreSessionAsEditableRow(const TSharedPtr<FConcertSessionItem>& ArchivedItem);
-	void InsertArchiveSessionAsEditableRow(const TSharedPtr<FConcertSessionItem>& ActiveItem);
-	void InsertEditableSessionRow(TSharedPtr<FConcertSessionItem> EditableItem, TSharedPtr<FConcertSessionItem> ParentItem);
-	void RemoveSessionRow(const TSharedPtr<FConcertSessionItem>& Item);
-
-	// Sessions sorting. (Sorts the session view)
-	EColumnSortMode::Type GetColumnSortMode(const FName ColumnId) const;
-	EColumnSortPriority::Type GetColumnSortPriority(const FName ColumnId) const;
-	void OnColumnSortModeChanged(const EColumnSortPriority::Type SortPriority, const FName& ColumnId, const EColumnSortMode::Type InSortMode);
-	void SortSessionList();
-	void EnsureEditableParentChildOrder();
-
-	// Sessions filtering. (Filters the session view)
-	void OnSearchTextChanged(const FText& InFilterText);
-	void OnSearchTextCommitted(const FText& InFilterText, ETextCommit::Type CommitType);
-	void OnFilterMenuChecked(const FName MenuName);
-	void PopulateSearchStrings(const FConcertSessionItem& Item, TArray<FString>& OutSearchStrings) const;
-	bool IsFilteredOut(const FConcertSessionItem& Item) const;
-	FText HighlightSearchText() const;
-
-	// Passes the user requests to FConcertBrowserController.
-	void RequestCreateSession(const TSharedPtr<FConcertSessionItem>& NewItem);
-	void RequestJoinSession(const TSharedPtr<FConcertSessionItem>& ActiveItem);
-	void RequestArchiveSession(const TSharedPtr<FConcertSessionItem>& ActiveItem, const FString& ArchiveName);
-	void RequestRestoreSession(const TSharedPtr<FConcertSessionItem>& RestoreItem, const FString& SessionName);
-	void RequestRenameSession(const TSharedPtr<FConcertSessionItem>& RenamedItem, const FString& NewName);
-	void RequestDeleteSession(const TSharedPtr<FConcertSessionItem>& DeletedItem);
-
-	// Update server/session/clients lists.
-	EActiveTimerReturnType TickDiscovery(double InCurrentTime, float InDeltaTime);
-	void UpdateDiscovery();
-	void RefreshSessionList();
-	void RefreshClientList(const TArray<FConcertSessionClientInfo>& LastestClientList);
-
-private:
-	// Gives access to the concert data (servers, sessions, clients, etc).
-	TUniquePtr<FConcertBrowserController> Controller;
-
-	// Keeps persistent user preferences, like the filters.
-	TStrongObjectPtr<UConcertSessionBrowserSettings> PersitentSettings;
-
-	// The items displayed in the session list view. It might be filtered and sorted compared to the full list hold by the controller.
-	TArray<TSharedPtr<FConcertSessionItem>> Sessions;
-
-	// The session list view.
-	TSharedPtr<SListView<TSharedPtr<FConcertSessionItem>>> SessionsView;
-
-	// The item corresponding to a row used to create/archive/restore a session. There is only one at the time
-	TSharedPtr<FConcertSessionItem> EditableSessionRow;
-	TSharedPtr<FConcertSessionItem> EditableSessionRowParent; // For archive/restore, indicate which element is archived or restored.
-
-	// Sorting.
-	EColumnSortMode::Type PrimarySortMode;
-	EColumnSortMode::Type SecondarySortMode;
-	FName PrimarySortedColumn;
-	FName SecondarySortedColumn;
-
-	// Filtering.
-	TSharedPtr<SSearchBox> SearchBox;
-	TSharedPtr<TTextFilter<const FConcertSessionItem&>> SearchTextFilter;
-	TSharedPtr<FText> SearchedText;
-	bool bRefreshSessionFilter = true;
-	FString DefaultServerURL;
-
-	// Selected Session Details.
-	TSharedPtr<SBorder> SessionDetailsView;
-	TSharedPtr<SExpandableArea> DetailsArea;
-	TArray<TSharedPtr<FConcertSessionClientInfo>> Clients;
-	TSharedPtr<SExpandableArea> ClientsArea;
-	TSharedPtr<SListView<TSharedPtr<FConcertSessionClientInfo>>> ClientsView;
-
-	// Used to compare the version used by UI versus the version cached in the controller.
-	uint32 DisplayedSessionListVersion = 0;
-	uint32 DisplayedClientListVersion = 0;
-	uint32 ServerListVersion = 0;
-	bool bLocalServerRunning = false;
-
-	TSharedPtr<SWidget> ServerDiscoveryPanel; // Displayed until a server is found.
-	TSharedPtr<SWidget> SessionDiscoveryPanel; // Displayed until a session is found.
-	TSharedPtr<SWidget> NoSessionSelectedPanel; // Displays 'select a session to view details' message in the details section.
-	TSharedPtr<SWidget> NoSessionDetailsPanel; // Displays 'no details available' message in the details section.
-	TSharedPtr<SWidget> NoClientPanel; // Displays the 'no client connected' message in Clients expendable area.
-};
-
-
-void SConcertSessionBrowser::Construct(const FArguments& InArgs, IConcertClientPtr InConcertClient, TSharedPtr<FText> InSearchText)
-{
-	if (!InConcertClient.IsValid())
-	{
-		return; // Don't build the UI if ConcertClient is not available.
-	}
-
-	Controller = MakeUnique<FConcertBrowserController>(InConcertClient);
-
-	// Reload the persistent settings, such as the filters.
-	PersitentSettings = TStrongObjectPtr<UConcertSessionBrowserSettings>(GetMutableDefault<UConcertSessionBrowserSettings>());
-
-	// Setup search filter.
-	SearchedText = InSearchText; // Reload a previous search text (in any). Useful to remember searched text between join/leave sessions, but not persistent if the tab is closed.
-	SearchTextFilter = MakeShared<TTextFilter<const FConcertSessionItem&>>(TTextFilter<const FConcertSessionItem&>::FItemToStringArray::CreateSP(this, &SConcertSessionBrowser::PopulateSearchStrings));
-	SearchTextFilter->OnChanged().AddSP(this, &SConcertSessionBrowser::RefreshSessionList);
-
-	// Displayed if no server is available.
-	ServerDiscoveryPanel = SNew(SConcertDiscovery)
-		.Text(LOCTEXT("LookingForServer", "Looking for Multi-User Servers..."))
-		.Visibility_Lambda([this]() { return Controller->GetServers().Num() == 0 ? EVisibility::Visible : EVisibility::Hidden; })
-		.IsButtonEnabled(this, &SConcertSessionBrowser::IsLaunchServerButtonEnabled)
-		.ButtonStyle(FAppStyle::Get(), "FlatButton.Default")
-		.ButtonIcon(FConcertFrontendStyle::Get()->GetBrush("Concert.NewServer.Small"))
-		.ButtonText(LOCTEXT("LaunchLocalServer", "Launch a Server"))
-		.ButtonToolTip(LOCTEXT("LaunchServerTooltip", "Launch a Multi-User server on your computer unless one is already running.\nThe editor UDP messaging settings will be passed to the launching server.\nThe server will use the server port found in Multi-User client settings if non 0 or use the editor udp messaging port number + 1, if non 0."))
-		.OnButtonClicked(this, &SConcertSessionBrowser::OnLaunchServerButtonClicked);
-
-	// Controls the text displayed in the 'No sessions' panel.
-	auto GetNoSessionText = [this]()
-	{
-		if (!Controller->HasReceivedInitialSessionList())
-		{
-			return LOCTEXT("LookingForSession", "Looking for Multi-User Sessions...");
-		}
-
-		return Controller->GetActiveSessions().Num() == 0 && Controller->GetArchivedSessions().Num() == 0 ?
-			LOCTEXT("NoSessionAvailable", "No Sessions Available") :
-			LOCTEXT("AllSessionsFilteredOut", "No Sessions Match the Filters\nChange Your Filter to View Sessions");
-	};
-
-	// Displayed when discovering session or if no session is available.
-	SessionDiscoveryPanel = SNew(SConcertDiscovery)
-		.Text_Lambda(GetNoSessionText)
-		.Visibility_Lambda([this]() { return Controller->GetServers().Num() > 0 && Sessions.Num() == 0 && !Controller->IsCreatingSession() ? EVisibility::Visible : EVisibility::Hidden; })
-		.ThrobberVisibility_Lambda([this]() { return !Controller->HasReceivedInitialSessionList() ? EVisibility::Visible : EVisibility::Collapsed; })
-		.ButtonVisibility_Lambda([this]() { return Controller->HasReceivedInitialSessionList() && Controller->GetActiveSessions().Num() == 0 && Controller->GetArchivedSessions().Num() == 0 ? EVisibility::Visible : EVisibility::Collapsed; })
-		.ButtonStyle(FAppStyle::Get(), "FlatButton.Default")
-		.ButtonIcon(FConcertFrontendStyle::Get()->GetBrush("Concert.NewSession.Small"))
-		.ButtonText(LOCTEXT("CreateSession", "Create Session"))
-		.ButtonToolTip(LOCTEXT("CreateSessionTooltip", "Create a new session"))
-		.OnButtonClicked(FOnClicked::CreateSP(this, &SConcertSessionBrowser::OnNewButtonClicked));
-
-	// Displayed when the selected session client view is empty (no client to display).
-	NoClientPanel = SNew(SConcertNoAvailability)
-		.Text(LOCTEXT("NoClientAvailable", "No Connected Clients"))
-		.Visibility_Lambda([this]() { return Clients.Num() == 0 ? EVisibility::Visible : EVisibility::Hidden; });
-
-	// Displayed as details when no session is selected. (No session selected or the selected session doesn't have any)
-	NoSessionSelectedPanel = SNew(SConcertNoAvailability)
-		.Text(LOCTEXT("NoSessionSelected", "Select a Session to View Details"));
-
-	// Displayed as details when the selected session has not specific details to display.
-	NoSessionDetailsPanel = SNew(SConcertNoAvailability)
-		.Text(LOCTEXT("NoSessionDetails", "The Selected Session Has No Details"));
-
-	// List used in details panel to display clients connected to an active session.
-	ClientsView = SNew(SListView<TSharedPtr<FConcertSessionClientInfo>>)
-		.ListItemsSource(&Clients)
-		.OnGenerateRow(this, &SConcertSessionBrowser::OnGenerateClientRowWidget)
-		.SelectionMode(ESelectionMode::Single)
-		.AllowOverscroll(EAllowOverscroll::No);
-
-	ChildSlot
-	[
-		MakeBrowserContent()
-	];
-
-	// Create a timer to periodically poll the server for sessions and session clients at a lower frequency than the normal tick.
-	RegisterActiveTimer(1.0f, FWidgetActiveTimerDelegate::CreateSP(this, &SConcertSessionBrowser::TickDiscovery));
-
-	if (!SearchedText->IsEmpty())
-	{
-		SearchBox->SetText(*SearchedText); // This trigger the chain of actions to apply the search filter.
-	}
-
-	bLocalServerRunning = IMultiUserClientModule::Get().IsConcertServerRunning();
-}
-
-TSharedRef<SWidget> SConcertSessionBrowser::MakeBrowserContent()
-{
-	return SNew(SBox)
-		.HAlign(HAlign_Fill)
-		.VAlign(VAlign_Fill)
-		[
-			// Splitter upper part displaying the available sessions/(server).
-			SNew(SSplitter)
-			.Orientation(Orient_Vertical)
-			.MinimumSlotHeight(80.0f) // Prevent widgets from overlapping.
-			+SSplitter::Slot()
-			.Value(0.6)
-			[
-				SNew(SBorder)
-				.BorderImage(FAppStyle::Get().GetBrush("ToolPanel.GroupBorder"))
-				.Padding(FMargin(1.0f, 2.0f))
-				[
-					SNew(SVerticalBox)
-
-					+SVerticalBox::Slot()
-					.AutoHeight()
-					[
-						MakeControlBar()
-					]
-
-					// Session list.
-					+SVerticalBox::Slot()
-					.FillHeight(1.0f)
-					.Padding(1.0f, 2.0f)
-					[
-						SNew(SOverlay)
-						+SOverlay::Slot()
-						[
-							MakeSessionTableView()
-						]
-						+SOverlay::Slot()
-						.Padding(0, 20, 0, 0) // To ensure the panel is below the row header.
-						[
-							SessionDiscoveryPanel.ToSharedRef()
-						]
-						+SOverlay::Slot()
-						.Padding(0, 20, 0, 0) // To ensure the panel is below the row header.
-						[
-							ServerDiscoveryPanel.ToSharedRef()
-						]
-					]
-
-					+SVerticalBox::Slot()
-					.AutoHeight()
-					.Padding(FMargin(2.0f, 0.0f))
-					[
-						SNew(SSeparator)
-					]
-
-					// Session Count/View options filter.
-					+SVerticalBox::Slot()
-					.AutoHeight()
-					.Padding(FMargin(2.0f, 0.0f))
-					[
-						MakeSessionViewOptionsBar()
-					]
-				]
-			]
-
-			// Session details.
-			+SSplitter::Slot()
-			.Value(0.4)
-			[
-				SAssignNew(SessionDetailsView, SBorder)
-				.BorderImage(FAppStyle::Get().GetBrush("ToolPanel.GroupBorder"))
-				[
-					NoSessionSelectedPanel.ToSharedRef()
-				]
-			]
-		];
-}
-
-EActiveTimerReturnType SConcertSessionBrowser::TickDiscovery(double InCurrentTime, float InDeltaTime)
-{
-	// Cache the result of this function because it is very expensive. It kills the framerate if polled every frame.
-	bLocalServerRunning = IMultiUserClientModule::Get().IsConcertServerRunning();
-
-	UpdateDiscovery();
-	return EActiveTimerReturnType::Continue;
-}
-
-void SConcertSessionBrowser::UpdateDiscovery()
-{
-	// Check if the controller has updated data since the last 'tick' and fire new asynchronous requests for future 'tick'.
-	TPair<uint32, uint32> ServerSessionListVersions = Controller->TickServersAndSessionsDiscovery();
-	ServerListVersion = ServerSessionListVersions.Key;
-
-	if (ServerSessionListVersions.Value != DisplayedSessionListVersion) // Need to refresh the list?
-	{
-		RefreshSessionList();
-		DisplayedSessionListVersion = ServerSessionListVersions.Value;
-	}
-
-	// If an active session is selected.
-	TArray<TSharedPtr<FConcertSessionItem>> SelectedSession = SessionsView->GetSelectedItems();
-	if (SelectedSession.Num() && SelectedSession[0]->Type == FConcertSessionItem::EType::ActiveSession)
-	{
-		// Ensure to poll its clients.
-		uint32 CachedClientListVersion = Controller->TickClientsDiscovery(SelectedSession[0]->ServerAdminEndpointId, SelectedSession[0]->SessionId);
-		if (CachedClientListVersion != DisplayedClientListVersion) // Need to refresh the list?
-		{
-			RefreshClientList(Controller->GetClients(SelectedSession[0]->ServerAdminEndpointId, SelectedSession[0]->SessionId));
-			DisplayedClientListVersion = CachedClientListVersion;
-		}
-	}
-}
-
-void SConcertSessionBrowser::RefreshSessionList()
-{
-	// Remember the selected instances (if any).
-	TArray<TSharedPtr<FConcertSessionItem>> SelectedItems = SessionsView->GetSelectedItems();
-	TArray<TSharedPtr<FConcertSessionItem>> ReselectedItems;
-	TSharedPtr<FConcertSessionItem> NewEditableRowParent;
-
-	// Predicate returning true if the specified item should be re-selected.
-	auto IsSelected = [&SelectedItems](const FConcertSessionItem& Item)
-	{
-		return SelectedItems.ContainsByPredicate([Item](const TSharedPtr<FConcertSessionItem>& Visited) { return *Visited == Item; });
-	};
-
-	// Matches the object instances before the update to the new instance after the update.
-	auto ReconcileObjectInstances = [&](TSharedPtr<FConcertSessionItem> NewItem)
-	{
-		if (IsSelected(*NewItem))
-		{
-			ReselectedItems.Add(NewItem);
-		}
-		else if (EditableSessionRowParent.IsValid() && !NewEditableRowParent.IsValid() && *EditableSessionRowParent == *NewItem)
-		{
-			NewEditableRowParent = NewItem;
-		}
-	};
-
-	// Clear sessions.
-	Sessions.Reset();
-
-	// Populate the live sessions.
-	for (const TSharedPtr<FConcertBrowserController::FActiveSessionInfo>& ActiveSession : Controller->GetActiveSessions())
-	{
-		FConcertSessionItem NewItem(FConcertSessionItem::EType::ActiveSession, ActiveSession->SessionInfo.SessionName, ActiveSession->SessionInfo.SessionId, ActiveSession->ServerInfo.ServerName, ActiveSession->ServerInfo.AdminEndpointId, ActiveSession->ServerInfo.ServerFlags);
-		if (!IsFilteredOut(NewItem))
-		{
-			Sessions.Emplace(MakeShared<FConcertSessionItem>(MoveTemp(NewItem)));
-			ReconcileObjectInstances(Sessions.Last());
-		}
-	}
-
-	// Populate the archived.
-	for (const TSharedPtr<FConcertBrowserController::FArchivedSessionInfo>& ArchivedSession : Controller->GetArchivedSessions())
-	{
-		FConcertSessionItem NewItem(FConcertSessionItem::EType::ArchivedSession, ArchivedSession->SessionInfo.SessionName, ArchivedSession->SessionInfo.SessionId, ArchivedSession->ServerInfo.ServerName, ArchivedSession->ServerInfo.AdminEndpointId, ArchivedSession->ServerInfo.ServerFlags);
-		if (!IsFilteredOut(NewItem))
-		{
-			Sessions.Emplace(MakeShared<FConcertSessionItem>(MoveTemp(NewItem)));
-			ReconcileObjectInstances(Sessions.Last());
-		}
-	}
-
-	// Restore the editable row state. (SortSessionList() below will ensure the parent/child relationship)
-	EditableSessionRowParent = NewEditableRowParent;
-	if (EditableSessionRow.IsValid())
-	{
-		if (EditableSessionRow->Type == FConcertSessionItem::EType::NewSession)
-		{
-			Sessions.Insert(EditableSessionRow, 0); // Always put 'new session' row at the top.
-		}
-		else if (EditableSessionRowParent.IsValid())
-		{
-			Sessions.Add(EditableSessionRow); // SortSessionList() called below will ensure the correct parent/child order.
-		}
-	}
-
-	// Restore previous selection.
-	if (ReselectedItems.Num())
-	{
-		for (const TSharedPtr<FConcertSessionItem>& Item : ReselectedItems)
-		{
-			SessionsView->SetItemSelection(Item, true);
-		}
-	}
-
-	SortSessionList();
-	SessionsView->RequestListRefresh();
-}
-
-void SConcertSessionBrowser::RefreshClientList(const TArray<FConcertSessionClientInfo>& LastestClientList)
-{
-	// Remember which client is selected.
-	TArray<TSharedPtr<FConcertSessionClientInfo>> SelectedItems = ClientsView->GetSelectedItems();
-
-	// Copy the list of clients.
-	TArray<TSharedPtr<FConcertSessionClientInfo>> LatestClientPtrs;
-	Algo::Transform(LastestClientList, LatestClientPtrs, [](const FConcertSessionClientInfo& Client) { return MakeShared<FConcertSessionClientInfo>(Client); });
-
-	// Merge the current list with the new list, removing client that disappeared and adding client that appeared.
-	ConcertFrontendUtils::SyncArraysByPredicate(Clients, MoveTemp(LatestClientPtrs), [](const TSharedPtr<FConcertSessionClientInfo>& ClientToFind)
-	{
-		return [ClientToFind](const TSharedPtr<FConcertSessionClientInfo>& PotentialClientMatch)
-		{
-			return PotentialClientMatch->ClientEndpointId == ClientToFind->ClientEndpointId && PotentialClientMatch->ClientInfo == ClientToFind->ClientInfo;
-		};
-	});
-
-	// Sort the list item alphabetically.
-	Clients.StableSort([](const TSharedPtr<FConcertSessionClientInfo>& Lhs, const TSharedPtr<FConcertSessionClientInfo>& Rhs)
-	{
-		return Lhs->ClientInfo.DisplayName < Rhs->ClientInfo.DisplayName;
-	});
-
-	// Preserve previously selected item (if any).
-	if (SelectedItems.Num())
-	{
-		ClientsView->SetSelection(SelectedItems[0]);
-	}
-
-	ClientsView->RequestListRefresh();
-}
-
-void SConcertSessionBrowser::OnSearchTextChanged(const FText& InFilterText)
-{
-	SearchTextFilter->SetRawFilterText(InFilterText);
-	SearchBox->SetError( SearchTextFilter->GetFilterErrorText() );
-	*SearchedText = InFilterText;
-
-	bRefreshSessionFilter = true;
-}
-
-void SConcertSessionBrowser::OnSearchTextCommitted(const FText& InFilterText, ETextCommit::Type CommitType)
-{
-	if (!InFilterText.EqualTo(*SearchedText))
-	{
-		OnSearchTextChanged(InFilterText);
-	}
-}
-
-void SConcertSessionBrowser::PopulateSearchStrings(const FConcertSessionItem& Item, TArray<FString>& OutSearchStrings) const
-{
-	OutSearchStrings.Add(Item.ServerName);
-	OutSearchStrings.Add(Item.SessionName);
-}
-
-bool SConcertSessionBrowser::IsFilteredOut(const FConcertSessionItem& Item) const
-{
-	bool bIsDefaultServer = DefaultServerURL.IsEmpty() || Item.ServerName == DefaultServerURL;
-
-	return (!PersitentSettings->bShowActiveSessions && (Item.Type == FConcertSessionItem::EType::ActiveSession || Item.Type == FConcertSessionItem::EType::SaveSession)) ||
-	       (!PersitentSettings->bShowArchivedSessions && (Item.Type == FConcertSessionItem::EType::ArchivedSession || Item.Type == FConcertSessionItem::EType::RestoreSession)) ||
-	       (PersitentSettings->bShowDefaultServerSessionsOnly && !bIsDefaultServer) ||
-	       !SearchTextFilter->PassesFilter(Item);
-}
-
-FText SConcertSessionBrowser::HighlightSearchText() const
-{
-	return *SearchedText;
-}
-
-TSharedRef<SWidget> SConcertSessionBrowser::MakeControlBar()
-{
-	return SNew(SHorizontalBox)
-
-		// The New/Join/Restore/Delete/Archive buttons
-		+SHorizontalBox::Slot()
-		.AutoWidth()
-		[
-			MakeButtonBar()
-		]
-
-		// The search text.
-		+SHorizontalBox::Slot()
-		.FillWidth(1.0)
-		.Padding(4.0f, 3.0f, 8.0f, 3.0f)
-		[
-			SAssignNew(SearchBox, SSearchBox)
-			.HintText(LOCTEXT("SearchHint", "Search Session"))
-			.OnTextChanged(this, &SConcertSessionBrowser::OnSearchTextChanged)
-			.OnTextCommitted(this, &SConcertSessionBrowser::OnSearchTextCommitted)
-			.DelayChangeNotificationsWhileTyping(true)
-		]
-
-		// The user "Avatar color" displayed as a small square colored by the user avatar color.
-		+SHorizontalBox::Slot()
-		.VAlign(VAlign_Center)
-		.AutoWidth()
-		[
-			SNew(STextBlock)
-			.Font(FAppStyle::Get().GetFontStyle("FontAwesome.8"))
-			.ColorAndOpacity_Lambda([this]() { return Controller->GetConcertClient()->GetClientInfo().AvatarColor; })
-			.Text(FEditorFontGlyphs::Square)
-		]
-
-		// The user "Display Name".
-		+SHorizontalBox::Slot()
-		.VAlign(VAlign_Center)
-		.AutoWidth()
-		.Padding(3, 0, 2, 0)
-		[
-			SNew(STextBlock)
-			.Text_Lambda([this]() { return FText::FromString(Controller->GetConcertClient()->GetClientInfo().DisplayName);} )
-		]
-
-		// The "Settings" icons.
-		+SHorizontalBox::Slot()
-		.VAlign(VAlign_Fill)
-		.AutoWidth()
-		[
-			SNew(SButton)
-			.ButtonStyle(FAppStyle::Get(), TEXT("SimpleButton"))
-			.OnClicked_Lambda([](){FModuleManager::GetModulePtr<ISettingsModule>("Settings")->ShowViewer("Project", "Plugins", "Concert"); return FReply::Handled(); })
-			.VAlign(VAlign_Center)
-			.HAlign(HAlign_Fill)
-			[
-				SNew(STextBlock)
-				.Font(FAppStyle::Get().GetFontStyle("FontAwesome.16"))
-				.Text(FEditorFontGlyphs::Cogs)
-			]
-		];
-}
-
-TSharedRef<SWidget> SConcertSessionBrowser::MakeButtonBar()
-{
-	TAttribute<FText> AutoJoinTooltip = TAttribute<FText>::Create([this]()
-	{
-		if (Controller->GetConcertClient()->CanAutoConnect()) // Default session and server are configured?
-		{
-			return FText::Format(LOCTEXT("JoinDefaultSessionTooltip", "Join the default session '{0}' on '{1}'"),
-				FText::FromString(Controller->GetConcertClient()->GetConfiguration()->DefaultSessionName),
-				FText::FromString(Controller->GetConcertClient()->GetConfiguration()->DefaultServerURL));
-		}
-		else
-		{
-			return LOCTEXT("JoinDefaultSessionConfiguredTooltip", "Join the default session configured in the Multi-Users settings");
-		}
-	});
-
-	TAttribute<FText> CancelAutoJoinTooltip = TAttribute<FText>::Create([this]()
-	{
-		return FText::Format(LOCTEXT("CancelJoinDefaultSessionTooltip", "Cancel joining the default session '{0}' on '{1}'"),
-			FText::FromString(Controller->GetConcertClient()->GetConfiguration()->DefaultSessionName),
-			FText::FromString(Controller->GetConcertClient()->GetConfiguration()->DefaultServerURL));
-	});
-
-	constexpr float PaddingBetweenButtons = 3.0f;
-
-	return SNew(SHorizontalBox)
-
-		// Launch server.
-		+SHorizontalBox::Slot()
-		.AutoWidth()
-		.Padding(0, 0, PaddingBetweenButtons, 0)
-		[
-			ConcertBrowserUtils::MakeIconButton(TEXT("SimpleButton"), FConcertFrontendStyle::Get()->GetBrush("Concert.NewServer"), LOCTEXT("LaunchServerTooltip", "Launch a Multi-User server on your computer unless one is already running.\nThe editor UDP messaging settings will be passed to the launching server.\nThe server will use the server port found in Multi-User client settings if non 0 or use the editor udp messaging port number + 1, if non 0."),
-				TAttribute<bool>(this, &SConcertSessionBrowser::IsLaunchServerButtonEnabled),
-				FOnClicked::CreateSP(this, &SConcertSessionBrowser::OnLaunchServerButtonClicked),
-				TAttribute<EVisibility>::Create([this]() { return IsLaunchServerButtonEnabled() ? EVisibility::Visible : EVisibility::Collapsed; }))
-		]
-
-		// Stop server.
-		+SHorizontalBox::Slot()
-		[
-			ConcertBrowserUtils::MakeIconButton(TEXT("SimpleButton"), FConcertFrontendStyle::Get()->GetBrush("Concert.CloseServer"), LOCTEXT("ShutdownServerTooltip", "Shutdown the Multi-User server running on this computer."),
-				true, // Always enabled, but might be collapsed.
-				FOnClicked::CreateSP(this, &SConcertSessionBrowser::OnShutdownServerButtonClicked),
-				TAttribute<EVisibility>::Create([this]() { return IsLaunchServerButtonEnabled() ? EVisibility::Collapsed : EVisibility::Visible; }))
-		]
-
-		// New Session
-		+SHorizontalBox::Slot()
-		.AutoWidth()
-		.Padding(0, 0, PaddingBetweenButtons, 0)
-		[
-			ConcertBrowserUtils::MakeIconButton(TEXT("SimpleButton"), FConcertFrontendStyle::Get()->GetBrush("Concert.NewSession"), LOCTEXT("NewButtonTooltip", "Create a new session"),
-				TAttribute<bool>(this, &SConcertSessionBrowser::IsNewButtonEnabled),
-				FOnClicked::CreateSP(this, &SConcertSessionBrowser::OnNewButtonClicked))
-		]
-
-		// Separator
-		+SHorizontalBox::Slot()
-		.AutoWidth()
-		.Padding(0, 1, PaddingBetweenButtons, 1)
-		[
-			SNew(SSeparator).Orientation(Orient_Vertical)
-		]
-
-		// Auto-join
-		+SHorizontalBox::Slot()
-		.AutoWidth()
-		.Padding(0, 0, PaddingBetweenButtons, 0)
-		[
-			ConcertBrowserUtils::MakeIconButton(TEXT("SimpleButton"), FConcertFrontendStyle::Get()->GetBrush("Concert.JoinDefaultSession"), AutoJoinTooltip,
-				TAttribute<bool>(this, &SConcertSessionBrowser::IsAutoJoinButtonEnabled),
-				FOnClicked::CreateSP(this, &SConcertSessionBrowser::OnAutoJoinButtonClicked),
-				TAttribute<EVisibility>::Create([this]() { return !IsCancelAutoJoinButtonEnabled() ? EVisibility::Visible : EVisibility::Collapsed; })) // Default button shown if both auto-join/cancel are disabled.
-		]
-
-		// Cancel auto join. (Shares the same slot as Auto-join)
-		+SHorizontalBox::Slot()
-		.AutoWidth()
-		.Padding(0, 0, PaddingBetweenButtons, 0)
-		[
-			ConcertBrowserUtils::MakeIconButton(TEXT("SimpleButton"), FConcertFrontendStyle::Get()->GetBrush("Concert.CancelAutoJoin"), CancelAutoJoinTooltip,
-				TAttribute<bool>(this, &SConcertSessionBrowser::IsCancelAutoJoinButtonEnabled),
-				FOnClicked::CreateSP(this, &SConcertSessionBrowser::OnCancelAutoJoinButtonClicked),
-				TAttribute<EVisibility>::Create([this]() { return IsCancelAutoJoinButtonEnabled() ? EVisibility::Visible : EVisibility::Collapsed; }))
-		]
-
-		// Join
-		+SHorizontalBox::Slot()
-		.AutoWidth()
-		.Padding(0, 0, PaddingBetweenButtons, 0)
-		[
-			ConcertBrowserUtils::MakeIconButton(TEXT("SimpleButton"), FConcertFrontendStyle::Get()->GetBrush("Concert.JoinSession"), LOCTEXT("JoinButtonTooltip", "Join the selected session"),
-				TAttribute<bool>(this, &SConcertSessionBrowser::IsJoinButtonEnabled),
-				FOnClicked::CreateSP(this, &SConcertSessionBrowser::OnJoinButtonClicked),
-				TAttribute<EVisibility>::Create([this]() { return !IsRestoreButtonEnabled() ? EVisibility::Visible : EVisibility::Collapsed; })) // Default button shown if both join/restore are disabled.
-		]
-
-		// Restore (Share the same slot as Join)
-		+SHorizontalBox::Slot()
-		.AutoWidth()
-		.Padding(0, 0, PaddingBetweenButtons, 0)
-		[
-			ConcertBrowserUtils::MakeIconButton(TEXT("SimpleButton"), FConcertFrontendStyle::Get()->GetBrush("Concert.RestoreSession"), LOCTEXT("RestoreButtonTooltip", "Restore the selected session"),
-				TAttribute<bool>(this, &SConcertSessionBrowser::IsRestoreButtonEnabled),
-				FOnClicked::CreateSP(this, &SConcertSessionBrowser::OnRestoreButtonClicked),
-				TAttribute<EVisibility>::Create([this]() { return IsRestoreButtonEnabled() ? EVisibility::Visible : EVisibility::Collapsed; }))
-		]
-
-		// Archive.
-		+SHorizontalBox::Slot()
-		.AutoWidth()
-		.Padding(0, 0, PaddingBetweenButtons, 0)
-		[
-			ConcertBrowserUtils::MakeIconButton(TEXT("SimpleButton"), FConcertFrontendStyle::Get()->GetBrush("Concert.ArchiveSession"), LOCTEXT("ArchiveButtonTooltip", "Archive the selected session"),
-				TAttribute<bool>(this, &SConcertSessionBrowser::IsArchiveButtonEnabled),
-				FOnClicked::CreateSP(this, &SConcertSessionBrowser::OnArchiveButtonClicked))
-		]
-
-		// Delete.
-		+SHorizontalBox::Slot()
-		.AutoWidth()
-		[
-			ConcertBrowserUtils::MakeIconButton(TEXT("SimpleButton"), FConcertFrontendStyle::Get()->GetBrush("Concert.DeleteSession"), LOCTEXT("DeleteButtonTooltip", "Delete the selected session if permitted"),
-				TAttribute<bool>(this, &SConcertSessionBrowser::IsDeleteButtonEnabled),
-				FOnClicked::CreateSP(this, &SConcertSessionBrowser::OnDeleteButtonClicked))
-		];
-}
-
-TSharedRef<SWidget> SConcertSessionBrowser::MakeSessionTableView()
-{
-	PrimarySortedColumn = ConcertBrowserUtils::IconColName;
-	PrimarySortMode = EColumnSortMode::Ascending;
-	SecondarySortedColumn = ConcertBrowserUtils::SessionColName;
-	SecondarySortMode = EColumnSortMode::Ascending;
-
-	return SAssignNew(SessionsView, SListView<TSharedPtr<FConcertSessionItem>>)
-		.SelectionMode(ESelectionMode::Single)
-		.ListItemsSource(&Sessions)
-		.OnGenerateRow(this, &SConcertSessionBrowser::OnGenerateSessionRowWidget)
-		.SelectionMode(ESelectionMode::Single)
-		.OnSelectionChanged(this, &SConcertSessionBrowser::OnSessionSelectionChanged)
-		.OnContextMenuOpening(this, &SConcertSessionBrowser::MakeContextualMenu)
-		.HeaderRow(
-			SNew(SHeaderRow)
-
-			+SHeaderRow::Column(ConcertBrowserUtils::IconColName)
-			.DefaultLabel(FText::GetEmpty())
-			.SortPriority(this, &SConcertSessionBrowser::GetColumnSortPriority, ConcertBrowserUtils::IconColName)
-			.SortMode(this, &SConcertSessionBrowser::GetColumnSortMode, ConcertBrowserUtils::IconColName)
-			.OnSort(this, &SConcertSessionBrowser::OnColumnSortModeChanged)
-			.FixedWidth(20)
-
-			+SHeaderRow::Column(ConcertBrowserUtils::SessionColName)
-			.DefaultLabel(LOCTEXT("SessioName", "Session"))
-			.SortPriority(this, &SConcertSessionBrowser::GetColumnSortPriority, ConcertBrowserUtils::SessionColName)
-			.SortMode(this, &SConcertSessionBrowser::GetColumnSortMode, ConcertBrowserUtils::SessionColName)
-			.OnSort(this, &SConcertSessionBrowser::OnColumnSortModeChanged)
-
-			+SHeaderRow::Column(ConcertBrowserUtils::ServerColName)
-			.DefaultLabel(LOCTEXT("Server", "Server"))
-			.SortPriority(this, &SConcertSessionBrowser::GetColumnSortPriority, ConcertBrowserUtils::ServerColName)
-			.SortMode(this, &SConcertSessionBrowser::GetColumnSortMode, ConcertBrowserUtils::ServerColName)
-			.OnSort(this, &SConcertSessionBrowser::OnColumnSortModeChanged));
-}
-
-EColumnSortMode::Type SConcertSessionBrowser::GetColumnSortMode(const FName ColumnId) const
-{
-	if (ColumnId == PrimarySortedColumn)
-	{
-		return PrimarySortMode;
-	}
-	else if (ColumnId == SecondarySortedColumn)
-	{
-		return SecondarySortMode;
-	}
-
-	return EColumnSortMode::None;
-}
-
-EColumnSortPriority::Type SConcertSessionBrowser::GetColumnSortPriority(const FName ColumnId) const
-{
-	if (ColumnId == PrimarySortedColumn)
-	{
-		return EColumnSortPriority::Primary;
-	}
-	else if (ColumnId == SecondarySortedColumn)
-	{
-		return EColumnSortPriority::Secondary;
-	}
-
-	return EColumnSortPriority::Max; // No specific priority.
-}
-
-void SConcertSessionBrowser::OnColumnSortModeChanged(const EColumnSortPriority::Type SortPriority, const FName& ColumnId, const EColumnSortMode::Type InSortMode)
-{
-	if (SortPriority == EColumnSortPriority::Primary)
-	{
-		PrimarySortedColumn = ColumnId;
-		PrimarySortMode = InSortMode;
-
-		if (ColumnId == SecondarySortedColumn) // Cannot be primary and secondary at the same time.
-		{
-			SecondarySortedColumn = FName();
-			SecondarySortMode = EColumnSortMode::None;
-		}
-	}
-	else if (SortPriority == EColumnSortPriority::Secondary)
-	{
-		SecondarySortedColumn = ColumnId;
-		SecondarySortMode = InSortMode;
-	}
-
-	SortSessionList();
-	SessionsView->RequestListRefresh();
-}
-
-void SConcertSessionBrowser::SortSessionList()
-{
-	check(!PrimarySortedColumn.IsNone()); // Should always have a primary column. User cannot clear this one.
-
-	auto Compare = [this](const TSharedPtr<FConcertSessionItem>& Lhs, const TSharedPtr<FConcertSessionItem>& Rhs, const FName& ColName, EColumnSortMode::Type SortMode)
-	{
-		if (Lhs->Type == FConcertSessionItem::EType::NewSession) // Always keep editable 'new session' row at the top.
-		{
-			return true;
-		}
-		else if (Rhs->Type == FConcertSessionItem::EType::NewSession)
-		{
-			return false;
-		}
-
-		if (ColName == ConcertBrowserUtils::IconColName)
-		{
-			return SortMode == EColumnSortMode::Ascending ? Lhs->Type < Rhs->Type : Lhs->Type > Rhs->Type;
-		}
-		else if (ColName == ConcertBrowserUtils::SessionColName)
-		{
-			return SortMode == EColumnSortMode::Ascending ? Lhs->SessionName < Rhs->SessionName : Lhs->SessionName > Rhs->SessionName;
-		}
-		else
-		{
-			return SortMode == EColumnSortMode::Ascending ? Lhs->ServerName < Rhs->ServerName : Lhs->ServerName > Rhs->ServerName;
-		}
-	};
-
-	Sessions.Sort([&](const TSharedPtr<FConcertSessionItem>& Lhs, const TSharedPtr<FConcertSessionItem>& Rhs)
-	{
-		if (Compare(Lhs, Rhs, PrimarySortedColumn, PrimarySortMode))
-		{
-			return true; // Lhs must be before Rhs based on the primary sort order.
-		}
-		else if (Compare(Rhs, Lhs, PrimarySortedColumn, PrimarySortMode)) // Invert operands order (goal is to check if operands are equal or not)
-		{
-			return false; // Rhs must be before Lhs based on the primary sort.
-		}
-		else // Lhs == Rhs on the primary column, need to order accoding the secondary column if one is set.
-		{
-			return SecondarySortedColumn.IsNone() ? false : Compare(Lhs, Rhs, SecondarySortedColumn, SecondarySortMode);
-		}
-	});
-
-	EnsureEditableParentChildOrder();
-}
-
-void SConcertSessionBrowser::EnsureEditableParentChildOrder()
-{
-	// This is for Archiving or Restoring a session. We keep the editable row below the session to archive or restore and visually link them with small wires in UI.
-	if (EditableSessionRowParent.IsValid())
-	{
-		check(EditableSessionRow.IsValid());
-		Sessions.Remove(EditableSessionRow);
-
-		int32 ParentIndex = Sessions.IndexOfByKey(EditableSessionRowParent);
-		if (ParentIndex != INDEX_NONE)
-		{
-			Sessions.Insert(EditableSessionRow, ParentIndex + 1); // Insert the 'child' below its parent.
-		}
-	}
-}
-
-TSharedRef<ITableRow> SConcertSessionBrowser::OnGenerateSessionRowWidget(TSharedPtr<FConcertSessionItem> Item, const TSharedRef<STableViewBase>& OwnerTable)
-{
-	switch (Item->Type)
-	{
-		case FConcertSessionItem::EType::ActiveSession:
-			return MakeActiveSessionRowWidget(Item, OwnerTable);
-
-		case FConcertSessionItem::EType::ArchivedSession:
-			return MakeArchivedSessionRowWidget(Item, OwnerTable);
-
-		case FConcertSessionItem::EType::NewSession:
-			return MakeNewSessionRowWidget(Item, OwnerTable);
-
-		case FConcertSessionItem::EType::RestoreSession:
-			return MakeRestoreSessionRowWidget(Item, OwnerTable);
-
-		default:
-			check(Item->Type == FConcertSessionItem::EType::SaveSession);
-			return MakeSaveSessionRowWidget(Item, OwnerTable);
-	}
-}
-
-TSharedRef<ITableRow> SConcertSessionBrowser::MakeActiveSessionRowWidget(const TSharedPtr<FConcertSessionItem>& ActiveItem, const TSharedRef<STableViewBase>& OwnerTable)
-{
-	const FConcertSessionInfo* SessionInfo = Controller->GetActiveSessionInfo(ActiveItem->ServerAdminEndpointId, ActiveItem->SessionId);
-
-	// Add an 'Active Session' row. Clicking the row icon joins the session.
-	return SNew(SSessionRow, ActiveItem, OwnerTable, Controller->GetConcertClient())
-		.OnDoubleClickFunc([this](const TSharedPtr<FConcertSessionItem>& Item) { RequestJoinSession(Item); })
-		.OnRenameFunc([this](const TSharedPtr<FConcertSessionItem>& Item, const FString& NewName) { RequestRenameSession(Item, NewName); })
-		.HighlightText(this, &SConcertSessionBrowser::HighlightSearchText)
-		.ToolTipText(SessionInfo ? SessionInfo->ToDisplayString() : FText::GetEmpty())
-		.IsSelected_Lambda([this, ActiveItem]() { return SessionsView->GetSelectedItems().Num() == 1 && SessionsView->GetSelectedItems()[0] == ActiveItem; });
-}
-
-TSharedRef<ITableRow> SConcertSessionBrowser::MakeArchivedSessionRowWidget(const TSharedPtr<FConcertSessionItem>& ArchivedItem, const TSharedRef<STableViewBase>& OwnerTable)
-{
-	const FConcertSessionInfo* SessionInfo = Controller->GetArchivedSessionInfo(ArchivedItem->ServerAdminEndpointId, ArchivedItem->SessionId);
-
-	// Add an 'Archived Session' row. Clicking the row icon adds a 'Restore as' row to the table.
-	return SNew(SSessionRow, ArchivedItem, OwnerTable, Controller->GetConcertClient())
-		.OnDoubleClickFunc([this](const TSharedPtr<FConcertSessionItem>& Item) { InsertRestoreSessionAsEditableRow(Item); })
-		.OnRenameFunc([this](const TSharedPtr<FConcertSessionItem>& Item, const FString& NewName) { RequestRenameSession(Item, NewName); })
-		.HighlightText(this, &SConcertSessionBrowser::HighlightSearchText)
-		.ToolTipText(SessionInfo ? SessionInfo->ToDisplayString() : FText::GetEmpty())
-		.IsSelected_Lambda([this, ArchivedItem]() { return SessionsView->GetSelectedItems().Num() == 1 && SessionsView->GetSelectedItems()[0] == ArchivedItem; });
-}
-
-TSharedRef<ITableRow> SConcertSessionBrowser::MakeNewSessionRowWidget(const TSharedPtr<FConcertSessionItem>& NewItem, const TSharedRef<STableViewBase>& OwnerTable)
-{
-	// Add an editable 'New Session' row in the table to let user pick a name and a server.
-	return SNew(SNewSessionRow, NewItem, OwnerTable, Controller->GetConcertClient())
-		.GetServerFunc([this]() { return TPair<uint32, const TArray<FConcertServerInfo>&>(ServerListVersion, Controller->GetServers()); }) // Let the row pull the servers for the combo box.
-		.OnAcceptFunc([this](const TSharedPtr<FConcertSessionItem>& Item) { RequestCreateSession(Item); }) // Accepting creates the session.
-		.OnDeclineFunc([this](const TSharedPtr<FConcertSessionItem>& Item) { RemoveSessionRow(Item); })  // Declining removes the editable 'new session' row from the view.
-		.HighlightText(this, &SConcertSessionBrowser::HighlightSearchText);
-}
-
-TSharedRef<ITableRow> SConcertSessionBrowser::MakeSaveSessionRowWidget(const TSharedPtr<FConcertSessionItem>& SaveItem, const TSharedRef<STableViewBase>& OwnerTable)
-{
-	// Add an editable 'Save Session' row in the table to let the user enter an archive name.
-	return SNew(SSaveRestoreSessionRow, SaveItem, OwnerTable)
-		.OnAcceptFunc([this](const TSharedPtr<FConcertSessionItem>& Item, const FString& ArchiveName) { RequestArchiveSession(Item, ArchiveName); }) // Accepting archive the session.
-		.OnDeclineFunc([this](const TSharedPtr<FConcertSessionItem>& Item) { RemoveSessionRow(Item); }) // Declining removes the editable 'save session as' row from the view.
-		.HighlightText(this, &SConcertSessionBrowser::HighlightSearchText);
-}
-
-TSharedRef<ITableRow> SConcertSessionBrowser::MakeRestoreSessionRowWidget(const TSharedPtr<FConcertSessionItem>& RestoreItem, const TSharedRef<STableViewBase>& OwnerTable)
-{
-	// Add an editable 'Restore Session' row in the table to let the user enter a session name.
-	return SNew(SSaveRestoreSessionRow, RestoreItem, OwnerTable)
-		.OnAcceptFunc([this](const TSharedPtr<FConcertSessionItem>& Item, const FString& SessionName) { RequestRestoreSession(Item, SessionName); }) // Accepting restores the session.
-		.OnDeclineFunc([this](const TSharedPtr<FConcertSessionItem>& Item) { RemoveSessionRow(Item); }) // Declining removes the editable 'restore session as' row from the view.
-		.HighlightText(this, &SConcertSessionBrowser::HighlightSearchText);
-}
-
-void SConcertSessionBrowser::InsertNewSessionEditableRow()
-{
-	// Insert a 'new session' editable row.
-	InsertEditableSessionRow(MakeShared<FConcertSessionItem>(FConcertSessionItem::EType::NewSession, FString(), FGuid(), FString(), FGuid(), EConcertServerFlags::None), nullptr);
-}
-
-void SConcertSessionBrowser::InsertRestoreSessionAsEditableRow(const TSharedPtr<FConcertSessionItem>& ArchivedItem)
-{
-	// Insert the 'restore session as ' editable row just below the 'archived' item to restore.
-	InsertEditableSessionRow(MakeShared<FConcertSessionItem>(FConcertSessionItem::EType::RestoreSession, ArchivedItem->SessionName, ArchivedItem->SessionId, ArchivedItem->ServerName, ArchivedItem->ServerAdminEndpointId, ArchivedItem->ServerFlags), ArchivedItem);
-}
-
-void SConcertSessionBrowser::InsertArchiveSessionAsEditableRow(const TSharedPtr<FConcertSessionItem>& LiveItem)
-{
-	// Insert the 'save session as' editable row just below the 'active' item to save.
-	InsertEditableSessionRow(MakeShared<FConcertSessionItem>(FConcertSessionItem::EType::SaveSession, LiveItem->SessionName, LiveItem->SessionId, LiveItem->ServerName, LiveItem->ServerAdminEndpointId, LiveItem->ServerFlags), LiveItem);
-}
-
-void SConcertSessionBrowser::InsertEditableSessionRow(TSharedPtr<FConcertSessionItem> EditableItem, TSharedPtr<FConcertSessionItem> ParentItem)
-{
-	// Insert the new row below its parent (if any).
-	int32 ParentIndex = Sessions.IndexOfByKey(ParentItem);
-	Sessions.Insert(EditableItem, ParentIndex != INDEX_NONE ? ParentIndex + 1 : 0);
-
-	// Ensure there is only one editable row at the time, removing the row being edited (if any).
-	Sessions.Remove(EditableSessionRow);
-	EditableSessionRow = EditableItem;
-	EditableSessionRowParent = ParentItem;
-
-	// Ensure the editable row added is selected and visible.
-	SessionsView->SetSelection(EditableItem);
-	SessionsView->RequestListRefresh();
-
-	// NOTE: Ideally, I would only use RequestScrollIntoView() to scroll the new item into view, but it did not work. If an item was added into an hidden part,
-	//       it was not always scrolled correctly into view. RequestNavigateToItem() worked much better, except when inserting the very first row in the list, in
-	//       such case calling the function would give the focus to the list view (showing a white dashed line around it).
-	if (ParentIndex == INDEX_NONE)
-	{
-		SessionsView->ScrollToTop(); // Item is inserted at 0. (New session)
-	}
-	else
-	{
-		SessionsView->RequestNavigateToItem(EditableItem);
-	}
 }
 
-void SConcertSessionBrowser::RemoveSessionRow(const TSharedPtr<FConcertSessionItem>& Item)
+void SConcertClientSessionBrowser::ExtendControlButtons(FExtender& Extender)
 {
-	Sessions.Remove(Item);
-
-	// Don't keep the editable row if its 'parent' is removed. (if the session to restore or archive gets deleted in the meantime)
-	if (Item == EditableSessionRowParent)
-	{
-		Sessions.Remove(EditableSessionRow);
-		EditableSessionRow.Reset();
-	}
+	// Before separator
+	Extender.AddToolBarExtension(
+		SConcertSessionBrowser::ControlButtonExtensionHooks::BeforeSeparator,
+		EExtensionHook::First,
+		nullptr,
+		FToolBarExtensionDelegate::CreateLambda([this](FToolBarBuilder& Builder)
+		{
+			Builder.AddWidget(
+				SNew(SOverlay)
+				// Launch server.
+				+SOverlay::Slot()
+				[
+					ConcertBrowserUtils::MakeIconButton(TEXT("SimpleButton"),
+						FConcertFrontendStyle::Get()->GetBrush("Concert.NewServer"),
+						LOCTEXT("LaunchServerTooltip", "Launch a Multi-User server on your computer unless one is already running.\nThe editor UDP messaging settings will be passed to the launching server.\nThe server will use the server port found in Multi-User client settings if non 0 or use the editor udp messaging port number + 1, if non 0."),
+						TAttribute<bool>(this, &SConcertClientSessionBrowser::IsLaunchServerButtonEnabled),
+						FOnClicked::CreateSP(this, &SConcertClientSessionBrowser::OnLaunchServerButtonClicked),
+						TAttribute<EVisibility>::Create([this]() { return IsLaunchServerButtonEnabled() ? EVisibility::Visible : EVisibility::Collapsed; }))
+				]
+				// Stop server.
+				+SOverlay::Slot()
+				[
+					ConcertBrowserUtils::MakeIconButton(
+						TEXT("SimpleButton"),
+						FConcertFrontendStyle::Get()->GetBrush("Concert.CloseServer"),
+						LOCTEXT("ShutdownServerTooltip", "Shutdown the Multi-User server running on this computer."),
+						true, // Always enabled, but might be collapsed.
+						FOnClicked::CreateSP(this, &SConcertClientSessionBrowser::OnShutdownServerButtonClicked),
+						TAttribute<EVisibility>::Create([this]() { return IsLaunchServerButtonEnabled() ? EVisibility::Collapsed : EVisibility::Visible; }))
+				]
+				);
+		}));
+
+	// After separator
+	Extender.AddToolBarExtension(
+		SConcertSessionBrowser::ControlButtonExtensionHooks::AfterSeparator,
+		EExtensionHook::First,
+		nullptr,
+		FToolBarExtensionDelegate::CreateLambda([this](FToolBarBuilder& Builder)
+		{
+			TAttribute<FText> AutoJoinTooltip = TAttribute<FText>::Create([this]()
+			{
+				if (Controller->GetConcertClient()->CanAutoConnect()) // Default session and server are configured?
+				{
+					return FText::Format(LOCTEXT("JoinDefaultSessionTooltip", "Join the default session '{0}' on '{1}'"),
+						FText::FromString(Controller->GetConcertClient()->GetConfiguration()->DefaultSessionName),
+						FText::FromString(Controller->GetConcertClient()->GetConfiguration()->DefaultServerURL));
+				}
+				else
+				{
+					return LOCTEXT("JoinDefaultSessionConfiguredTooltip", "Join the default session configured in the Multi-Users settings");
+				}
+			});
 
-	// Clear the editable row state if its the one removed.
-	if (Item == EditableSessionRow)
-	{
-		EditableSessionRow.Reset();
-		EditableSessionRowParent.Reset();
-	}
+			TAttribute<FText> CancelAutoJoinTooltip = TAttribute<FText>::Create([this]()
+			{
+				return FText::Format(LOCTEXT("CancelJoinDefaultSessionTooltip", "Cancel joining the default session '{0}' on '{1}'"),
+					FText::FromString(Controller->GetConcertClient()->GetConfiguration()->DefaultSessionName),
+					FText::FromString(Controller->GetConcertClient()->GetConfiguration()->DefaultServerURL));
+			});
+			
+			Builder.AddWidget(
+				SNew(SOverlay)
+				// Auto-join
+				+SOverlay::Slot()
+				[
+					ConcertBrowserUtils::MakeIconButton(TEXT("SimpleButton"),
+						FConcertFrontendStyle::Get()->GetBrush("Concert.JoinDefaultSession"),
+						AutoJoinTooltip,
+						TAttribute<bool>(this, &SConcertClientSessionBrowser::IsAutoJoinButtonEnabled),
+						FOnClicked::CreateSP(this, &SConcertClientSessionBrowser::OnAutoJoinButtonClicked),
+						// Default button shown if both auto-join/cancel are disabled.
+						TAttribute<EVisibility>::Create([this]() { return !IsCancelAutoJoinButtonEnabled() ? EVisibility::Visible : EVisibility::Collapsed; })
+						) 
+				
+				]
+				// Cancel auto join.
+				+SOverlay::Slot()
+				[
+					ConcertBrowserUtils::MakeIconButton(TEXT("SimpleButton"),
+						FConcertFrontendStyle::Get()->GetBrush("Concert.CancelAutoJoin"),
+						CancelAutoJoinTooltip,
+						TAttribute<bool>(this, &SConcertClientSessionBrowser::IsCancelAutoJoinButtonEnabled),
+						FOnClicked::CreateSP(this, &SConcertClientSessionBrowser::OnCancelAutoJoinButtonClicked),
+						TAttribute<EVisibility>::Create([this]() { return IsCancelAutoJoinButtonEnabled() ? EVisibility::Visible : EVisibility::Collapsed; })
+						)
+				]
+			);
 
-	SessionsView->RequestListRefresh();
+			// Join
+			Builder.AddWidget(
+				ConcertBrowserUtils::MakeIconButton(TEXT("SimpleButton"), FConcertFrontendStyle::Get()->GetBrush("Concert.JoinSession"), LOCTEXT("JoinButtonTooltip", "Join the selected session"),
+					TAttribute<bool>(this, &SConcertClientSessionBrowser::IsJoinButtonEnabled),
+					FOnClicked::CreateSP(this, &SConcertClientSessionBrowser::OnJoinButtonClicked),
+					TAttribute<EVisibility>::Create([this]() { return !SessionBrowser->IsRestoreButtonEnabled() ? EVisibility::Visible : EVisibility::Collapsed; })) // Default button shown if both join/restore are disabled.
+			);
+		}));
 }
 
-TSharedRef<SWidget> SConcertSessionBrowser::MakeSessionViewOptionsBar()
+void SConcertClientSessionBrowser::ExtendSessionContextMenu(TSharedPtr<FConcertSessionItem> Item, FExtender& Extender)
 {
-	auto AddFilterMenu = [this]()
+	if (Item->Type == FConcertSessionItem::EType::ActiveSession)
 	{
-		FMenuBuilder MenuBuilder(/*bInShouldCloseWindowAfterMenuSelection=*/true, nullptr);
-
-		MenuBuilder.AddMenuEntry(
-			LOCTEXT("ActiveSessions_Label", "Active Sessions"),
-			LOCTEXT("ActiveSessions_Tooltip", "Displays Active Sessions"),
-			FSlateIcon(),
-			FUIAction(
-				FExecuteAction::CreateSP(this, &SConcertSessionBrowser::OnFilterMenuChecked, ConcertBrowserUtils::ActiveSessionsCheckBoxMenuName),
-				FCanExecuteAction::CreateLambda([] { return true; }),
-				FIsActionChecked::CreateLambda([this] { return PersitentSettings->bShowActiveSessions; })),
-			NAME_None,
-			EUserInterfaceActionType::ToggleButton
-		);
-
-		MenuBuilder.AddMenuEntry(
-			LOCTEXT("ArchivedSessions_Label", "Archived Sessions"),
-			LOCTEXT("ArchivedSessions_Tooltip", "Displays Archived Sessions"),
-			FSlateIcon(),
-			FUIAction(
-				FExecuteAction::CreateSP(this, &SConcertSessionBrowser::OnFilterMenuChecked, ConcertBrowserUtils::ArchivedSessionsCheckBoxMenuName),
-				FCanExecuteAction::CreateLambda([] { return true; }),
-				FIsActionChecked::CreateLambda([this] { return PersitentSettings->bShowArchivedSessions; })),
-			NAME_None,
-			EUserInterfaceActionType::ToggleButton
-		);
-
-		MenuBuilder.AddMenuEntry(
-			LOCTEXT("DefaultServer_Label", "Default Server Sessions"),
-			LOCTEXT("DefaultServer_Tooltip", "Displays Sessions Hosted By the Default Server"),
-			FSlateIcon(),
-			FUIAction(
-				FExecuteAction::CreateSP(this, &SConcertSessionBrowser::OnFilterMenuChecked, ConcertBrowserUtils::DefaultServerCheckBoxMenuName),
-				FCanExecuteAction::CreateLambda([] { return true; }),
-				FIsActionChecked::CreateLambda([this] { return PersitentSettings->bShowDefaultServerSessionsOnly; })),
-			NAME_None,
-			EUserInterfaceActionType::ToggleButton
-		);
-
-		return MenuBuilder.MakeWidget();
-	};
+		Extender.AddMenuExtension(
+			SConcertSessionBrowser::SessionContextMenuExtensionHooks::ManageSession,
+			EExtensionHook::First,
+			nullptr,
+			FMenuExtensionDelegate::CreateLambda([this](FMenuBuilder& Builder)
+			{
+				Builder.AddMenuEntry(
+					LOCTEXT("CtxMenuJoin", "Join"),
+					LOCTEXT("CtxMenuJoin_Tooltip", "Join the Session"),
+					FSlateIcon(),
+					FUIAction(
+						FExecuteAction::CreateLambda([this](){ OnJoinButtonClicked(); }),
+						FCanExecuteAction::CreateLambda([SelectedCount = SessionBrowser->GetSelectedItems().Num()] { return SelectedCount == 1; }),
+						FIsActionChecked::CreateLambda([this] { return false; })),
+					NAME_None,
+					EUserInterfaceActionType::Button
+					);
+			}));
+	}
+	
+}
 
+TSharedRef<SWidget> SConcertClientSessionBrowser::MakeUserAndSettings()
+{
 	return SNew(SHorizontalBox)
+		// The user "Avatar color" displayed as a small square colored by the user avatar color.
 		+SHorizontalBox::Slot()
-		.AutoWidth()
 		.VAlign(VAlign_Center)
+		.AutoWidth()
 		[
-			SNew(STextBlock).Text_Lambda([this]()
-			{
-				// Don't count the 'New Session', 'Restore Session' and 'Archive Session' editable row, they are transient rows used for inline input only.
-				int32 DisplayedSessionNum = Sessions.Num() - (EditableSessionRow.IsValid() ? 1 : 0);
-				int32 AvailableSessionNum = Controller->GetActiveSessions().Num() + Controller->GetArchivedSessions().Num();
-				int32 ServerNum = Controller->GetServers().Num();
-
-				// If all discovered session are displayed (none excluded by a filter).
-				if (DisplayedSessionNum == AvailableSessionNum)
-				{
-					if (Controller->GetServers().Num() == 0)
-					{
-						return LOCTEXT("NoServerNoFilter", "No servers found");
-					}
-					else
-					{
-						return FText::Format(LOCTEXT("NSessionNServerNoFilter", "{0} {0}|plural(one=session,other=sessions) on {1} {1}|plural(one=server,other=servers)"),
-							DisplayedSessionNum, ServerNum);
-					}
-				}
-				else // A filter is excluding at least one session.
-				{
-					if (DisplayedSessionNum == 0)
-					{
-						return FText::Format(LOCTEXT("NoSessionMatchNServer", "No matching sessions ({0} total on {1} {1}|plural(one=server,other=servers))"),
-							AvailableSessionNum, ServerNum);
-					}
-					else
-					{
-						return FText::Format(LOCTEXT("NSessionNServer", "Showing {0} of {1} {1}|plural(one=session,other=sessions) on {2} {2}|plural(one=server,other=servers)"),
-							DisplayedSessionNum, AvailableSessionNum, ServerNum);
-					}
-				}
-			})
+			SNew(STextBlock)
+				.Font(FAppStyle::Get().GetFontStyle("FontAwesome.8"))
+				.ColorAndOpacity_Lambda([this]() { return Controller->GetConcertClient()->GetClientInfo().AvatarColor; })
+				.Text(FEditorFontGlyphs::Square)
 		]
-
+				
+		// The user "Display Name".
 		+SHorizontalBox::Slot()
-		.FillWidth(1.0)
+		.VAlign(VAlign_Center)
+		.AutoWidth()
+		.Padding(3, 0, 2, 0)
 		[
-			SNew(SSpacer)
+			SNew(STextBlock)
+				.Text_Lambda([this]() { return FText::FromString(Controller->GetConcertClient()->GetClientInfo().DisplayName);} )
 		]
 
+		// The "Settings" icons.
 		+SHorizontalBox::Slot()
+		.VAlign(VAlign_Fill)
 		.AutoWidth()
 		[
-			SNew(SComboButton)
-			.ComboButtonStyle(FAppStyle::Get(), "GenericFilters.ComboButtonStyle")
-			.ForegroundColor(FLinearColor::White)
-			.ContentPadding(0)
-			.OnGetMenuContent_Lambda(AddFilterMenu)
-			.HasDownArrow(true)
-			.ContentPadding(FMargin(1, 0))
-			.ButtonContent()
+			SNew(SButton)
+			.ButtonStyle(FAppStyle::Get(), TEXT("SimpleButton"))
+			.OnClicked_Lambda([](){FModuleManager::GetModulePtr<ISettingsModule>("Settings")->ShowViewer("Project", "Plugins", "Concert"); return FReply::Handled(); })
+			.VAlign(VAlign_Center)
+			.HAlign(HAlign_Fill)
 			[
-				SNew(SHorizontalBox)
-
-				+SHorizontalBox::Slot()
-				.AutoWidth()
-				.VAlign(VAlign_Center)
-				[
-					SNew(SImage).Image(FAppStyle::Get().GetBrush("GenericViewButton")) // The eye ball image.
-				]
-
-				+SHorizontalBox::Slot()
-				.AutoWidth()
-				.Padding(2, 0, 0, 0)
-				.VAlign(VAlign_Center)
-				[
-					SNew(STextBlock).Text(LOCTEXT("ViewOptions", "View Options"))
-				]
+				SNew(STextBlock)
+				.Font(FAppStyle::Get().GetFontStyle("FontAwesome.16"))
+				.Text(FEditorFontGlyphs::Cogs)
 			]
 		];
 }
 
-void SConcertSessionBrowser::OnFilterMenuChecked(const FName MenuName)
+TSharedRef<SWidget> SConcertClientSessionBrowser::MakeOverlayedTableView(TSharedRef<SWidget> SessionTable)
 {
-	if (MenuName == ConcertBrowserUtils::ActiveSessionsCheckBoxMenuName)
-	{
-		PersitentSettings->bShowActiveSessions = !PersitentSettings->bShowActiveSessions;
-	}
-	else if (MenuName == ConcertBrowserUtils::ArchivedSessionsCheckBoxMenuName)
-	{
-		PersitentSettings->bShowArchivedSessions = !PersitentSettings->bShowArchivedSessions;
-	}
-	else if (MenuName == ConcertBrowserUtils::DefaultServerCheckBoxMenuName)
-	{
-		PersitentSettings->bShowDefaultServerSessionsOnly = !PersitentSettings->bShowDefaultServerSessionsOnly;
-	}
-	bRefreshSessionFilter = true;
+	return SNew(SOverlay)
+		+SOverlay::Slot()
+		[
+			SessionTable
+		]
+		+SOverlay::Slot()
+		.Padding(0, 20, 0, 0) // To ensure the panel is below the row header.
+		[
+			SessionDiscoveryPanel.ToSharedRef()
+		]
+		+SOverlay::Slot()
+		.Padding(0, 20, 0, 0) // To ensure the panel is below the row header.
+		[
+			ServerDiscoveryPanel.ToSharedRef()
+		];
+}
+
+EActiveTimerReturnType SConcertClientSessionBrowser::TickDiscovery(double InCurrentTime, float InDeltaTime)
+{
+	// Cache the result of this function because it is very expensive. It kills the framerate if polled every frame.
+	bLocalServerRunning = IMultiUserClientModule::Get().IsConcertServerRunning();
 
-	PersitentSettings->SaveConfig();
+	UpdateDiscovery();
+	return EActiveTimerReturnType::Continue;
 }
 
-TSharedPtr<SWidget> SConcertSessionBrowser::MakeContextualMenu()
+void SConcertClientSessionBrowser::UpdateDiscovery()
 {
-	TArray<TSharedPtr<FConcertSessionItem>> SelectedItems = SessionsView->GetSelectedItems();
-	if (SelectedItems.Num() == 0 || (SelectedItems[0]->Type != FConcertSessionItem::EType::ActiveSession && SelectedItems[0]->Type != FConcertSessionItem::EType::ArchivedSession))
+	// Check if the controller has updated data since the last 'tick' and fire new asynchronous requests for future 'tick'.
+	TPair<uint32, uint32> ServerSessionListVersions = Controller->TickServersAndSessionsDiscovery();
+	ServerListVersion = ServerSessionListVersions.Key;
+
+	if (ServerSessionListVersions.Value != DisplayedSessionListVersion) // Need to refresh the list?
 	{
-		return nullptr; // No menu for editable rows.
+		SessionBrowser->RefreshSessionList();
+		DisplayedSessionListVersion = ServerSessionListVersions.Value;
 	}
 
-	TSharedPtr<FConcertSessionItem> Item = SelectedItems[0];
+	// If an active session is selected.
+	TArray<TSharedPtr<FConcertSessionItem>> SelectedSession = SessionBrowser->GetSelectedItems();
+	if (SelectedSession.Num() && SelectedSession[0]->Type == FConcertSessionItem::EType::ActiveSession)
+	{
+		// Ensure to poll its clients.
+		uint32 CachedClientListVersion = Controller->TickClientsDiscovery(SelectedSession[0]->ServerAdminEndpointId, SelectedSession[0]->SessionId);
+		if (CachedClientListVersion != DisplayedClientListVersion) // Need to refresh the list?
+		{
+			RefreshClientList(Controller->GetClients(SelectedSession[0]->ServerAdminEndpointId, SelectedSession[0]->SessionId));
+			DisplayedClientListVersion = CachedClientListVersion;
+		}
+	}
+}
 
-	FMenuBuilder MenuBuilder(/*bInShouldCloseWindowAfterMenuSelection=*/true, nullptr);
+void SConcertClientSessionBrowser::RefreshClientList(const TArray<FConcertSessionClientInfo>& LastestClientList)
+{
+	// Remember which client is selected.
+	TArray<TSharedPtr<FConcertSessionClientInfo>> SelectedItems = ClientsView->GetSelectedItems();
 
-	// Section title.
-	MenuBuilder.BeginSection(NAME_None, Item->Type == FConcertSessionItem::EType::ActiveSession ?
-		LOCTEXT("ActiveSessionSection", "Active Session"):
-		LOCTEXT("ArchivedSessionSection", "Archived Session"));
+	// Copy the list of clients.
+	TArray<TSharedPtr<FConcertSessionClientInfo>> LatestClientPtrs;
+	Algo::Transform(LastestClientList, LatestClientPtrs, [](const FConcertSessionClientInfo& Client) { return MakeShared<FConcertSessionClientInfo>(Client); });
 
-	if (Item->Type == FConcertSessionItem::EType::ActiveSession)
+	// Merge the current list with the new list, removing client that disappeared and adding client that appeared.
+	ConcertFrontendUtils::SyncArraysByPredicate(Clients, MoveTemp(LatestClientPtrs), [](const TSharedPtr<FConcertSessionClientInfo>& ClientToFind)
 	{
-		MenuBuilder.AddMenuEntry(
-			LOCTEXT("CtxMenuJoin", "Join"),
-			LOCTEXT("CtxMenuJoin_Tooltip", "Join the Session"),
-			FSlateIcon(),
-			FUIAction(
-				FExecuteAction::CreateLambda([this](){ OnJoinButtonClicked(); }),
-				FCanExecuteAction::CreateLambda([SelectedCount = SelectedItems.Num()] { return SelectedCount == 1; }),
-				FIsActionChecked::CreateLambda([this] { return false; })),
-			NAME_None,
-			EUserInterfaceActionType::Button);
-
-		MenuBuilder.AddMenuEntry(
-			LOCTEXT("CtxMenuArchive", "Archive"),
-			LOCTEXT("CtxMenuArchive_Tooltip", "Archived the Session"),
-			FSlateIcon(),
-			FUIAction(
-				FExecuteAction::CreateLambda([this](){ OnArchiveButtonClicked(); }),
-				FCanExecuteAction::CreateLambda([SelectedCount = SelectedItems.Num()] { return SelectedCount == 1; }),
-				FIsActionChecked::CreateLambda([this] { return false; })),
-			NAME_None,
-			EUserInterfaceActionType::Button);
-	}
-	else // Archive
+		return [ClientToFind](const TSharedPtr<FConcertSessionClientInfo>& PotentialClientMatch)
+		{
+			return PotentialClientMatch->ClientEndpointId == ClientToFind->ClientEndpointId && PotentialClientMatch->ClientInfo == ClientToFind->ClientInfo;
+		};
+	});
+
+	// Sort the list item alphabetically.
+	Clients.StableSort([](const TSharedPtr<FConcertSessionClientInfo>& Lhs, const TSharedPtr<FConcertSessionClientInfo>& Rhs)
 	{
-		MenuBuilder.AddMenuEntry(
-			LOCTEXT("CtxMenuRestore", "Restore"),
-			LOCTEXT("CtxMenuRestore_Tooltip", "Restore the Session"),
-			FSlateIcon(),
-			FUIAction(
-				FExecuteAction::CreateLambda([this](){ OnRestoreButtonClicked(); }),
-				FCanExecuteAction::CreateLambda([SelectedCount = SelectedItems.Num()] { return SelectedCount == 1; }),
-				FIsActionChecked::CreateLambda([this] { return false; })),
-			NAME_None,
-			EUserInterfaceActionType::Button);
-	}
+		return Lhs->ClientInfo.DisplayName < Rhs->ClientInfo.DisplayName;
+	});
 
-	MenuBuilder.AddMenuEntry(
-		LOCTEXT("CtxMenuRename", "Rename"),
-		LOCTEXT("CtxMenuRename_Tooltip", "Rename the Session"),
-		FSlateIcon(),
-		FUIAction(
-			FExecuteAction::CreateLambda([this, Item](){ OnBeginEditingSessionName(Item); }),
-			FCanExecuteAction::CreateLambda([this] { return IsRenameButtonEnabled(); }),
-			FIsActionChecked::CreateLambda([this] { return false; })),
-		NAME_None,
-		EUserInterfaceActionType::Button);
-
-	MenuBuilder.AddMenuEntry(
-		LOCTEXT("CtxMenuDelete", "Delete"),
-		FText::Format(LOCTEXT("CtxMenuDelete_Tooltip", "Delete the {0}|plural(one=Session,other=Sessions)"), SelectedItems.Num()),
-		FSlateIcon(),
-		FUIAction(
-			FExecuteAction::CreateLambda([this](){ OnDeleteButtonClicked(); }),
-			FCanExecuteAction::CreateLambda([this] { return IsDeleteButtonEnabled(); }),
-			FIsActionChecked::CreateLambda([this] { return false; })),
-		NAME_None,
-		EUserInterfaceActionType::Button);
-
-	MenuBuilder.EndSection();
-
-	return MenuBuilder.MakeWidget();
+	// Preserve previously selected item (if any).
+	if (SelectedItems.Num())
+	{
+		ClientsView->SetSelection(SelectedItems[0]);
+	}
+	ClientsView->RequestListRefresh();
 }
 
-TSharedRef<SWidget> SConcertSessionBrowser::MakeSessionDetails(TSharedPtr<FConcertSessionItem> Item)
+TSharedRef<SWidget> SConcertClientSessionBrowser::MakeSessionDetails(TSharedPtr<FConcertSessionItem> Item)
 {
 	const FConcertSessionInfo* SessionInfo = nullptr;
 
@@ -3168,7 +1448,7 @@ TSharedRef<SWidget> SConcertSessionBrowser::MakeSessionDetails(TSharedPtr<FConce
 	return NoSessionSelectedPanel.ToSharedRef();
 }
 
-TSharedRef<SWidget> SConcertSessionBrowser::MakeActiveSessionDetails(TSharedPtr<FConcertSessionItem> Item)
+TSharedRef<SWidget> SConcertClientSessionBrowser::MakeActiveSessionDetails(TSharedPtr<FConcertSessionItem> Item)
 {
 	const FConcertSessionInfo* SessionInfo = Controller->GetActiveSessionInfo(Item->ServerAdminEndpointId, Item->SessionId);
 	if (!SessionInfo)
@@ -3281,7 +1561,7 @@ TSharedRef<SWidget> SConcertSessionBrowser::MakeActiveSessionDetails(TSharedPtr<
 	return Widget;
 }
 
-TSharedRef<SWidget> SConcertSessionBrowser::MakeArchivedSessionDetails(TSharedPtr<FConcertSessionItem> Item)
+TSharedRef<SWidget> SConcertClientSessionBrowser::MakeArchivedSessionDetails(TSharedPtr<FConcertSessionItem> Item)
 {
 	const FConcertSessionInfo* SessionInfo = Controller->GetArchivedSessionInfo(Item->ServerAdminEndpointId, Item->SessionId);
 	if (!SessionInfo)
@@ -3324,7 +1604,7 @@ TSharedRef<SWidget> SConcertSessionBrowser::MakeArchivedSessionDetails(TSharedPt
 	return Widget;
 }
 
-void SConcertSessionBrowser::PopulateSessionInfoGrid(SGridPanel& Grid, const FConcertSessionInfo& SessionInfo)
+void SConcertClientSessionBrowser::PopulateSessionInfoGrid(SGridPanel& Grid, const FConcertSessionInfo& SessionInfo)
 {
 	// Local function to populate the details grid.
 	auto AddDetailRow = [](SGridPanel& Grid, int32 Row, const FText& Label, const FText& Value)
@@ -3366,7 +1646,7 @@ void SConcertSessionBrowser::PopulateSessionInfoGrid(SGridPanel& Grid, const FCo
 	AddDetailRow(Grid, Row++, LOCTEXT("ServerEndPointId", "Server Endpoint ID:"), FText::FromString(SessionInfo.ServerEndpointId.ToString()));
 }
 
-TSharedRef<ITableRow> SConcertSessionBrowser::OnGenerateClientRowWidget(TSharedPtr<FConcertSessionClientInfo> Item, const TSharedRef<STableViewBase>& OwnerTable)
+TSharedRef<ITableRow> SConcertClientSessionBrowser::OnGenerateClientRowWidget(TSharedPtr<FConcertSessionClientInfo> Item, const TSharedRef<STableViewBase>& OwnerTable)
 {
 	return SNew(STableRow<TSharedPtr<FConcertSessionClientInfo>>, OwnerTable)
 	[
@@ -3396,16 +1676,8 @@ TSharedRef<ITableRow> SConcertSessionBrowser::OnGenerateClientRowWidget(TSharedP
 	];
 }
 
-void SConcertSessionBrowser::OnSessionSelectionChanged(TSharedPtr<FConcertSessionItem> SelectedSession, ESelectInfo::Type SelectInfo)
+void SConcertClientSessionBrowser::OnSessionSelectionChanged(TSharedPtr<FConcertSessionItem> SelectedSession)
 {
-	// Cancel editing the row to create, archive or restore a session (if any), unless the row was selected in code.
-	if (EditableSessionRow.IsValid() && SelectInfo != ESelectInfo::Direct)
-	{
-		check (SelectedSession != EditableSessionRow); // User should not be able to reselect an editable row as we remove it as soon as it is unselected.
-		RemoveSessionRow(EditableSessionRow);
-		check(!EditableSessionRow.IsValid() && !EditableSessionRowParent.IsValid()); // Expect to be cleared by RemoveSessionRow().
-	}
-
 	// Clear the list of clients (if any)
 	Clients.Reset();
 
@@ -3413,154 +1685,95 @@ void SConcertSessionBrowser::OnSessionSelectionChanged(TSharedPtr<FConcertSessio
 	SessionDetailsView->SetContent(MakeSessionDetails(SelectedSession));
 }
 
-bool SConcertSessionBrowser::IsNewButtonEnabled() const
-{
-	return Controller->GetServers().Num() > 0;
-}
-
-bool SConcertSessionBrowser::IsJoinButtonEnabled() const
-{
-	TArray<TSharedPtr<FConcertSessionItem>> SelectedItems = SessionsView->GetSelectedItems();
-	return SelectedItems.Num() == 1 && SelectedItems[0]->Type == FConcertSessionItem::EType::ActiveSession;
-}
-
-bool SConcertSessionBrowser::IsRestoreButtonEnabled() const
-{
-	TArray<TSharedPtr<FConcertSessionItem>> SelectedItems = SessionsView->GetSelectedItems();
-	return SelectedItems.Num() == 1 && SelectedItems[0]->Type == FConcertSessionItem::EType::ArchivedSession;
-}
-
-bool SConcertSessionBrowser::IsArchiveButtonEnabled() const
-{
-	TArray<TSharedPtr<FConcertSessionItem>> SelectedItems = SessionsView->GetSelectedItems();
-	return SelectedItems.Num() == 1 && SelectedItems[0]->Type == FConcertSessionItem::EType::ActiveSession;
-}
-
-bool SConcertSessionBrowser::IsRenameButtonEnabled() const
+void SConcertClientSessionBrowser::OnSessionDoubleClicked(TSharedPtr<FConcertSessionItem> SelectedSession)
 {
-	TArray<TSharedPtr<FConcertSessionItem>> SelectedItems = SessionsView->GetSelectedItems();
-	if (SelectedItems.Num() != 1)
+	switch (SelectedSession->Type)
 	{
-		return false;
+	case FConcertSessionItem::EType::ActiveSession: 
+		RequestJoinSession(SelectedSession);
+		break;
+	case FConcertSessionItem::EType::ArchivedSession:
+		SessionBrowser->InsertRestoreSessionAsEditableRow(SelectedSession);
+		break;
+	default: 
+		checkNoEntry();
 	}
-
-	return (SelectedItems[0]->Type == FConcertSessionItem::EType::ActiveSession && Controller->CanRenameActiveSession(SelectedItems[0]->ServerAdminEndpointId, SelectedItems[0]->SessionId)) ||
-	       (SelectedItems[0]->Type == FConcertSessionItem::EType::ArchivedSession && Controller->CanRenameArchivedSession(SelectedItems[0]->ServerAdminEndpointId, SelectedItems[0]->SessionId));
 }
 
-bool SConcertSessionBrowser::IsDeleteButtonEnabled() const
+bool SConcertClientSessionBrowser::IsLaunchServerButtonEnabled() const
 {
-	TArray<TSharedPtr<FConcertSessionItem>> SelectedItems = SessionsView->GetSelectedItems();
-	if (SelectedItems.Num() == 0)
-	{
-		return false;
-	}
-
-	return (SelectedItems[0]->Type == FConcertSessionItem::EType::ActiveSession && Controller->CanDeleteActiveSession(SelectedItems[0]->ServerAdminEndpointId, SelectedItems[0]->SessionId)) ||
-	       (SelectedItems[0]->Type == FConcertSessionItem::EType::ArchivedSession && Controller->CanDeleteArchivedSession(SelectedItems[0]->ServerAdminEndpointId, SelectedItems[0]->SessionId));
+	return !bLocalServerRunning;
 }
 
-bool SConcertSessionBrowser::IsLaunchServerButtonEnabled() const
+bool SConcertClientSessionBrowser::IsJoinButtonEnabled() const
 {
-	return !bLocalServerRunning;
+	TArray<TSharedPtr<FConcertSessionItem>> SelectedItems = SessionBrowser->GetSelectedItems();
+	return SelectedItems.Num() == 1 && SelectedItems[0]->Type == FConcertSessionItem::EType::ActiveSession;
 }
 
-bool SConcertSessionBrowser::IsAutoJoinButtonEnabled() const
+bool SConcertClientSessionBrowser::IsAutoJoinButtonEnabled() const
 {
 	return Controller->GetConcertClient()->CanAutoConnect() && !Controller->GetConcertClient()->IsAutoConnecting();
 }
 
-bool SConcertSessionBrowser::IsCancelAutoJoinButtonEnabled() const
+bool SConcertClientSessionBrowser::IsCancelAutoJoinButtonEnabled() const
 {
 	return Controller->GetConcertClient()->IsAutoConnecting();
 }
 
-FReply SConcertSessionBrowser::OnNewButtonClicked()
-{
-	InsertNewSessionEditableRow();
-	return FReply::Handled();
-}
-
-FReply SConcertSessionBrowser::OnJoinButtonClicked()
+FReply SConcertClientSessionBrowser::OnNewButtonClicked()
 {
-	TArray<TSharedPtr<FConcertSessionItem>> SelectedItems = SessionsView->GetSelectedItems();
-	if (SelectedItems.Num() == 1)
-	{
-		RequestJoinSession(SelectedItems[0]);
-	}
+	SessionBrowser->InsertNewSessionEditableRow();
 	return FReply::Handled();
 }
 
-FReply SConcertSessionBrowser::OnRestoreButtonClicked()
+FReply SConcertClientSessionBrowser::OnLaunchServerButtonClicked()
 {
-	TArray<TSharedPtr<FConcertSessionItem>> SelectedItems = SessionsView->GetSelectedItems();
-	if (SelectedItems.Num() == 1)
-	{
-		InsertRestoreSessionAsEditableRow(SelectedItems[0]);
-	}
-
+	IMultiUserClientModule::Get().LaunchConcertServer();
+	bLocalServerRunning = IMultiUserClientModule::Get().IsConcertServerRunning(); // Immediately update the cache state to avoid showing buttons enabled for a split second.
 	return FReply::Handled();
 }
 
-FReply SConcertSessionBrowser::OnArchiveButtonClicked()
+FReply SConcertClientSessionBrowser::OnShutdownServerButtonClicked()
 {
-	TArray<TSharedPtr<FConcertSessionItem>> SelectedItems = SessionsView->GetSelectedItems();
-	if (SelectedItems.Num() == 1)
+	if (bLocalServerRunning)
 	{
-		InsertArchiveSessionAsEditableRow(SelectedItems[0]);
+		IMultiUserClientModule::Get().ShutdownConcertServer();
 	}
-
 	return FReply::Handled();
 }
 
-FReply SConcertSessionBrowser::OnDeleteButtonClicked()
+FReply SConcertClientSessionBrowser::OnAutoJoinButtonClicked()
 {
-	TArray<TSharedPtr<FConcertSessionItem>> SelectedItems = SessionsView->GetSelectedItems();
-	for (const TSharedPtr<FConcertSessionItem>& Item : SelectedItems)
-	{
-		RequestDeleteSession(Item);
-	}
+	IConcertClientPtr ConcertClient = Controller->GetConcertClient();
 
+	// Start the 'auto connect' routine. It will try until it succeeded or gets canceled. Creating or Joining a session automatically cancels it.
+	ConcertClient->StartAutoConnect();
 	return FReply::Handled();
 }
 
-void SConcertSessionBrowser::OnBeginEditingSessionName(TSharedPtr<FConcertSessionItem> Item)
-{
-	Item->OnBeginEditSessionNameRequest.Broadcast(); // Signal the row widget to enter in edit mode.
-}
-
-FReply SConcertSessionBrowser::OnLaunchServerButtonClicked()
+FReply SConcertClientSessionBrowser::OnCancelAutoJoinButtonClicked()
 {
-	IMultiUserClientModule::Get().LaunchConcertServer();
-	bLocalServerRunning = IMultiUserClientModule::Get().IsConcertServerRunning(); // Immediately update the cache state to avoid showing buttons enabled for a split second.
+	Controller->GetConcertClient()->StopAutoConnect();
 	return FReply::Handled();
 }
 
-FReply SConcertSessionBrowser::OnShutdownServerButtonClicked()
+FReply SConcertClientSessionBrowser::OnJoinButtonClicked()
 {
-	if (bLocalServerRunning)
+	TArray<TSharedPtr<FConcertSessionItem>> SelectedItems = SessionBrowser->GetSelectedItems();
+	if (SelectedItems.Num() == 1)
 	{
-		IMultiUserClientModule::Get().ShutdownConcertServer();
+		RequestJoinSession(SelectedItems[0]);
 	}
 	return FReply::Handled();
 }
 
-FReply SConcertSessionBrowser::OnAutoJoinButtonClicked()
-{
-	IConcertClientPtr ConcertClient = Controller->GetConcertClient();
-
-	// Start the 'auto connect' routine. It will try until it succeeded or gets canceled. Creating or Joining a session automatically cancels it.
-	ConcertClient->StartAutoConnect();
-	return FReply::Handled();
-}
-
-FReply SConcertSessionBrowser::OnCancelAutoJoinButtonClicked()
+void SConcertClientSessionBrowser::RequestJoinSession(const TSharedPtr<FConcertSessionItem>& LiveItem)
 {
-	Controller->GetConcertClient()->StopAutoConnect();
-	return FReply::Handled();
+	Controller->JoinSession(LiveItem->ServerAdminEndpointId, LiveItem->SessionId);
 }
 
-void SConcertSessionBrowser::Tick(const FGeometry& AllottedGeometry, const double InCurrentTime, const float InDeltaTime)
+void SConcertClientSessionBrowser::Tick(const FGeometry& AllottedGeometry, const double InCurrentTime, const float InDeltaTime)
 {
 	// If the model received updates.
 	if (Controller->GetAndClearDiscoveryUpdateFlag())
@@ -3579,105 +1792,11 @@ void SConcertSessionBrowser::Tick(const FGeometry& AllottedGeometry, const doubl
 	// Should refresh the session filter?
 	if (bRefreshSessionFilter)
 	{
-		RefreshSessionList();
+		SessionBrowser->RefreshSessionList();
 		bRefreshSessionFilter = false;
 	}
 }
 
-FReply SConcertSessionBrowser::OnKeyDown(const FGeometry& MyGeometry, const FKeyEvent& InKeyEvent)
-{
-	// NOTE: When an 'editable row' text box has the focus the keys are grabbed by the text box but
-	//       if the editable row is still selected, but the text field doesn't have the focus anymore
-	//       the keys will end up here if the browser has the focus.
-
-	if (InKeyEvent.GetKey() == EKeys::Delete && !EditableSessionRow.IsValid()) // Delete selected row(s) unless the selected row is an 'editable' one.
-	{
-		for (TSharedPtr<FConcertSessionItem>& Item : SessionsView->GetSelectedItems())
-		{
-			RequestDeleteSession(Item);
-		}
-		return FReply::Handled();
-	}
-	else if (InKeyEvent.GetKey() == EKeys::Escape && EditableSessionRow.IsValid()) // Cancel 'new session', 'archive session' or 'restore session' action.
-	{
-		RemoveSessionRow(EditableSessionRow);
-		check(!EditableSessionRow.IsValid() && !EditableSessionRowParent.IsValid()); // Expect to be cleared by RemoveSessionRow().
-		return FReply::Handled();
-	}
-	else if (InKeyEvent.GetKey() == EKeys::F2 && !EditableSessionRow.IsValid())
-	{
-		TArray<TSharedPtr<FConcertSessionItem>> SelectedItems = SessionsView->GetSelectedItems();
-		if (SelectedItems.Num() == 1)
-		{
-			SelectedItems[0]->OnBeginEditSessionNameRequest.Broadcast(); // Broadcast the request.
-		}
-	}
-
-	return FReply::Unhandled();
-}
-
-void SConcertSessionBrowser::RequestCreateSession(const TSharedPtr<FConcertSessionItem>& NewItem)
-{
-	Controller->CreateSession(NewItem->ServerAdminEndpointId, NewItem->SessionName);
-	RemoveSessionRow(NewItem); // The row used to edit the session name and pick the server.
-}
-
-void SConcertSessionBrowser::RequestJoinSession(const TSharedPtr<FConcertSessionItem>& LiveItem)
-{
-	Controller->JoinSession(LiveItem->ServerAdminEndpointId, LiveItem->SessionId);
-}
-
-void SConcertSessionBrowser::RequestArchiveSession(const TSharedPtr<FConcertSessionItem>& SaveItem, const FString& ArchiveName)
-{
-	Controller->ArchiveSession(SaveItem->ServerAdminEndpointId, SaveItem->SessionId, ArchiveName, FConcertSessionFilter());
-	RemoveSessionRow(SaveItem); // The row used to edit the archive name.
-}
-
-void SConcertSessionBrowser::RequestRestoreSession(const TSharedPtr<FConcertSessionItem>& RestoreItem, const FString& SessionName)
-{
-	Controller->RestoreSession(RestoreItem->ServerAdminEndpointId, RestoreItem->SessionId, SessionName, FConcertSessionFilter());
-	RemoveSessionRow(RestoreItem); // The row used to edit the restore as name.
-}
-
-void SConcertSessionBrowser::RequestRenameSession(const TSharedPtr<FConcertSessionItem>& RenamedItem, const FString& NewName)
-{
-	if (RenamedItem->Type == FConcertSessionItem::EType::ActiveSession)
-	{
-		Controller->RenameActiveSession(RenamedItem->ServerAdminEndpointId, RenamedItem->SessionId, NewName);
-	}
-	else if (RenamedItem->Type == FConcertSessionItem::EType::ArchivedSession)
-	{
-		Controller->RenameArchivedSession(RenamedItem->ServerAdminEndpointId, RenamedItem->SessionId, NewName);
-	}
-
-	// Display the new name until the server response is received. If the server refuses the new name, the discovery will reset the
-	// name (like if another client renamed it back) and the user will get a toast saying the rename failed.
-	RenamedItem->SessionName = NewName;
-}
-
-void SConcertSessionBrowser::RequestDeleteSession(const TSharedPtr<FConcertSessionItem>& DeletedItem)
-{
-	const FText SessionNameInText = FText::FromString(DeletedItem->SessionName);
-	const FText SeverNameInText = FText::FromString(DeletedItem->ServerName);
-	const FText ConfirmationMessage = FText::Format(LOCTEXT("DeleteSessionConfirmationMessage", "Do you really want to delete the session \"{0}\" from the server \"{1}\"?"), SessionNameInText, SeverNameInText);
-	const FText ConfirmationTitle = LOCTEXT("DeleteSessionConfirmationTitle", "Delete Session Confirmation");
-
-	if (FMessageDialog::Open(EAppMsgType::YesNo, ConfirmationMessage, &ConfirmationTitle) == EAppReturnType::Yes) // Confirmed?
-	{
-		if (DeletedItem->Type == FConcertSessionItem::EType::ActiveSession)
-		{
-			Controller->DeleteActiveSession(DeletedItem->ServerAdminEndpointId, DeletedItem->SessionId);
-		}
-		else if (DeletedItem->Type == FConcertSessionItem::EType::ArchivedSession)
-		{
-			Controller->DeleteArchivedSession(DeletedItem->ServerAdminEndpointId, DeletedItem->SessionId);
-		}
-
-		UpdateDiscovery(); // Don't wait up to 1s, kick discovery right now.
-	}
-}
-
-
 void SConcertBrowser::Construct(const FArguments& InArgs, const TSharedRef<SDockTab>& ConstructUnderMajorTab, const TSharedPtr<SWindow>& ConstructUnderWindow, TWeakPtr<IConcertSyncClient> InSyncClient)
 {
 	if (!MultiUserClientUtils::HasServerCompatibleCommunicationPluginEnabled())
@@ -3722,7 +1841,7 @@ void SConcertBrowser::AttachChildWidget(EConcertConnectionStatus ConnectionStatu
 		}
 		else if (ConnectionStatus == EConcertConnectionStatus::Disconnected)
 		{
-			ChildSlot.AttachWidget(SNew(SConcertSessionBrowser, ConcertSyncClient->GetConcertClient(), SearchedText));
+			ChildSlot.AttachWidget(SNew(SConcertClientSessionBrowser, ConcertSyncClient->GetConcertClient(), SearchedText));
 		}
 	}
 }
