Fix compile issue on Linux

diff --git a/Engine/Plugins/Developer/Concert/ConcertUI/ConcertSharedSlate/Source/ConcertSharedSlate/Public/SessionBrowser/IConcertSessionBrowserController.h b/Engine/Plugins/Developer/Concert/ConcertUI/ConcertSharedSlate/Source/ConcertSharedSlate/Public/SessionBrowser/IConcertSessionBrowserController.h
--- a/Engine/Plugins/Developer/Concert/ConcertUI/ConcertSharedSlate/Source/ConcertSharedSlate/Public/SessionBrowser/IConcertSessionBrowserController.h
+++ b/Engine/Plugins/Developer/Concert/ConcertUI/ConcertSharedSlate/Source/ConcertSharedSlate/Public/SessionBrowser/IConcertSessionBrowserController.h
@@ -19,12 +19,22 @@ public:
 		FConcertServerInfo ServerInfo;
 		FConcertSessionInfo SessionInfo;
 		TArray<FConcertSessionClientInfo> Clients;
+
+		FActiveSessionInfo(FConcertServerInfo ServerInfo, FConcertSessionInfo SessionInfo)
+			: ServerInfo(MoveTemp(ServerInfo))
+			, SessionInfo(MoveTemp(SessionInfo))
+		{}
 	};
 
 	struct FArchivedSessionInfo
 	{
 		FConcertServerInfo ServerInfo;
 		FConcertSessionInfo SessionInfo;
+
+		FArchivedSessionInfo(FConcertServerInfo ServerInfo, FConcertSessionInfo SessionInfo)
+			: ServerInfo(MoveTemp(ServerInfo))
+			, SessionInfo(MoveTemp(SessionInfo))
+		{}
 	};
 
 public:
diff --git a/Engine/Plugins/Developer/Concert/MultiUserClient/Source/MultiUserClient/Private/Widgets/SConcertBrowser.cpp b/Engine/Plugins/Developer/Concert/MultiUserClient/Source/MultiUserClient/Private/Widgets/SConcertBrowser.cpp
--- a/Engine/Plugins/Developer/Concert/MultiUserClient/Source/MultiUserClient/Private/Widgets/SConcertBrowser.cpp
+++ b/Engine/Plugins/Developer/Concert/MultiUserClient/Source/MultiUserClient/Private/Widgets/SConcertBrowser.cpp
@@ -165,12 +165,20 @@ private:
 	{
 		/** Raised when the UI and the cache values may be out of sync if a rename failed (UI assumed it succeeded) */
 		bool bSessionNameDirty = false;
+
+		FClientActiveSessionInfo(FConcertServerInfo ServerInfo, FConcertSessionInfo SessionInfo)
+			: FActiveSessionInfo(MoveTemp(ServerInfo), MoveTemp(SessionInfo))
+		{}
 	};
 
 	struct FClientArchivedSessionInfo : public FArchivedSessionInfo
 	{
 		/** Raised when the UI and the cache values may be out of sync if a rename failed (UI assumed it succeeded) */
-		bool bSessionNameDirty = false; 
+		bool bSessionNameDirty = false;
+
+		FClientArchivedSessionInfo(FConcertServerInfo ServerInfo, FConcertSessionInfo SessionInfo)
+			: FArchivedSessionInfo(MoveTemp(ServerInfo), MoveTemp(SessionInfo))
+		{}
 	};
 
 	void UpdateSessionsAsync();
@@ -665,7 +673,7 @@ void FConcertClientSessionBrowserController::UpdateActiveSessionsAsync(const FCo
 						{ return SessionInfo.ServerInstanceId == MatchCandidate->ServerInfo.InstanceInfo.InstanceId && SessionInfo.SessionId == MatchCandidate->SessionInfo.SessionId; }))
 					{
 						// This is a newly discovered session, add it to the list.
-						ActiveSessions.Add(MakeShared<FClientActiveSessionInfo>(FClientActiveSessionInfo{ServerInfo, SessionInfo}));
+						ActiveSessions.Add(MakeShared<FClientActiveSessionInfo>(ServerInfo, SessionInfo));
 						OnActiveSessionDiscovered(*ActiveSessions.Last());
 					}
 
@@ -749,7 +757,7 @@ void FConcertClientSessionBrowserController::UpdateArchivedSessionsAsync(const F
 						{ return ServerInfo.InstanceInfo.InstanceId == MatchCandidate->ServerInfo.InstanceInfo.InstanceId && SessionInfo.SessionId == MatchCandidate->SessionInfo.SessionId; }))
 					{
 						// This is a newly discovered archive, add it to the list.
-						ArchivedSessions.Add(MakeShared<FClientArchivedSessionInfo>(FClientArchivedSessionInfo{ServerInfo, SessionInfo}));
+						ArchivedSessions.Add(MakeShared<FClientArchivedSessionInfo>(ServerInfo, SessionInfo));
 						OnArchivedSessionDiscovered(*ArchivedSessions.Last());
 					}
 				}
