Fix non-unity issue

diff --git a/Engine/Plugins/Developer/Concert/MultiUserClient/Source/MultiUserClient/Private/Widgets/SConcertBrowser.cpp b/Engine/Plugins/Developer/Concert/MultiUserClient/Source/MultiUserClient/Private/Widgets/SConcertBrowser.cpp
--- a/Engine/Plugins/Developer/Concert/MultiUserClient/Source/MultiUserClient/Private/Widgets/SConcertBrowser.cpp
+++ b/Engine/Plugins/Developer/Concert/MultiUserClient/Source/MultiUserClient/Private/Widgets/SConcertBrowser.cpp
@@ -20,10 +20,15 @@
 #include "SessionBrowser/SConcertSessionBrowser.h"
 
 #include "Algo/Transform.h"
-#include "Styling/AppStyle.h"
+#include "Framework/Application/SlateApplication.h"
+#include "Framework/Commands/UIAction.h"
+#include "Framework/Commands/UICommandInfo.h"
+#include "Framework/MultiBox/MultiBoxBuilder.h"
 #include "Internationalization/Regex.h"
 #include "Misc/AsyncTaskNotification.h"
 #include "Misc/TextFilter.h"
+#include "Styling/AppStyle.h"
+#include "Textures/SlateIcon.h"
 #include "Widgets/Images/SImage.h"
 #include "Widgets/Images/SThrobber.h"
 #include "Widgets/Input/SButton.h"
@@ -36,6 +41,7 @@
 #include "Widgets/Layout/SUniformGridPanel.h"
 #include "Widgets/Text/SInlineEditableTextBlock.h"
 #include "Widgets/Text/STextBlock.h"
+#include "Widgets/SWindow.h"
 #include "Widgets/Views/SListView.h"
 #include "Widgets/Views/STableRow.h"
 
