Hide IConcertSyncClient from SSessionHistory

diff --git a/Engine/Plugins/Developer/Concert/MultiUserClient/Source/MultiUserClient/Private/ConcertWorkspaceUI.cpp b/Engine/Plugins/Developer/Concert/MultiUserClient/Source/MultiUserClient/Private/ConcertWorkspaceUI.cpp
--- a/Engine/Plugins/Developer/Concert/MultiUserClient/Source/MultiUserClient/Private/ConcertWorkspaceUI.cpp
+++ b/Engine/Plugins/Developer/Concert/MultiUserClient/Source/MultiUserClient/Private/ConcertWorkspaceUI.cpp
@@ -33,6 +33,7 @@
 
 #include "ToolMenuContext.h"
 #include "ToolMenus.h"
+#include "Widgets/ClientSessionHistoryController.h"
 
 #include "Widgets/DeclarativeSyntaxSupport.h"
 #include "Widgets/SConcertSandboxPersistWidget.h"
@@ -43,7 +44,7 @@
 #include "Widgets/Images/SImage.h"
 #include "Widgets/Text/STextBlock.h"
 #include "Widgets/Docking/SDockTab.h"
-#include "Widgets/SSessionHistory.h"
+#include "Widgets/SSessionHistoryWrapper.h"
 #include "Widgets/Notifications/SNotificationList.h"
 
 #define LOCTEXT_NAMESPACE "ConcertFrontend"
@@ -834,12 +835,17 @@ void FConcertWorkspaceUI::ExecuteUnlockResources(TArray<FName> InResourceNames)
 
 void FConcertWorkspaceUI::ExecuteViewHistory(TArray<FName> InResourceNames)
 {
-	FGlobalTabmanager::Get()->RestoreFrom(AssetHistoryLayout.ToSharedRef(), nullptr);
-
-	for (const FName& ResourceName : InResourceNames)
+	if (const TSharedPtr<IConcertSyncClient> ClientPin = SyncClient.Pin())
 	{
-		FGlobalTabmanager::Get()->InsertNewDocumentTab(ConcertHistoryTabName, FTabManager::ESearchPreference::PreferLiveTab, CreateHistoryTab(ResourceName));
+		FGlobalTabmanager::Get()->RestoreFrom(AssetHistoryLayout.ToSharedRef(), nullptr);
+		
+		const TSharedRef<IConcertSyncClient> ClientRef = ClientPin.ToSharedRef();
+		for (const FName& ResourceName : InResourceNames)
+		{
+			FGlobalTabmanager::Get()->InsertNewDocumentTab(ConcertHistoryTabName, FTabManager::ESearchPreference::PreferLiveTab, CreateHistoryTab(ResourceName, ClientRef));
+		}
 	}
+	
 }
 
 bool FConcertWorkspaceUI::IsAssetModifiedByOtherClients(const FName& AssetName, int32* OutOtherClientsWithModifNum, TArray<FConcertClientInfo>* OutOtherClientsWithModifInfo, int32 OtherClientsWithModifMaxFetchNum) const
@@ -1007,8 +1013,10 @@ void FConcertWorkspaceUI::OnPreSequencerInit(TSharedRef<ISequencer> InSequencer,
 	}
 }
 
-TSharedRef<SDockTab> FConcertWorkspaceUI::CreateHistoryTab(const FName& ResourceName) const
+TSharedRef<SDockTab> FConcertWorkspaceUI::CreateHistoryTab(const FName& ResourceName, const TSharedRef<IConcertSyncClient>& SyncClient)
 {
+	const TSharedRef<FClientSessionHistoryController> SessionHistoryController =
+		MakeShared<FClientSessionHistoryController>(SyncClient, ResourceName);
 	return SNew(SDockTab)
 		.TabRole(ETabRole::DocumentTab)
 		.ContentPadding(FMargin(3.0f))
@@ -1045,8 +1053,7 @@ TSharedRef<SDockTab> FConcertWorkspaceUI::CreateHistoryTab(const FName& Resource
 				.BorderImage(FEditorStyle::GetBrush("ToolPanel.GroupBorder"))
 				.Padding(FMargin(4.0f, 2.0f))
 				[
-					SNew(SSessionHistory, SyncClient.Pin())
-					.PackageFilter(ResourceName)
+					SNew(SSessionHistoryWrapper, SessionHistoryController)
 				]
 			]
 		];
diff --git a/Engine/Plugins/Developer/Concert/MultiUserClient/Source/MultiUserClient/Private/ConcertWorkspaceUI.h b/Engine/Plugins/Developer/Concert/MultiUserClient/Source/MultiUserClient/Private/ConcertWorkspaceUI.h
--- a/Engine/Plugins/Developer/Concert/MultiUserClient/Source/MultiUserClient/Private/ConcertWorkspaceUI.h
+++ b/Engine/Plugins/Developer/Concert/MultiUserClient/Source/MultiUserClient/Private/ConcertWorkspaceUI.h
@@ -129,7 +129,7 @@ private:
 	void OnPreSequencerInit(TSharedRef<ISequencer> InSequencer, TSharedRef<ISequencerObjectChangeListener> InObjectChangeListener, const FSequencerInitParams& InitParams);
 
 	/** Create an asset history tab filtered with a resource name. */
-	TSharedRef<SDockTab> CreateHistoryTab(const FName& ResourceName) const;
+	static TSharedRef<SDockTab> CreateHistoryTab(const FName& ResourceName, const TSharedRef<IConcertSyncClient>& SyncClient);
 
 	/** Workspace this is a view of. */
 	TWeakPtr<IConcertClientWorkspace> ClientWorkspace;
diff --git a/Engine/Plugins/Developer/Concert/MultiUserClient/Source/MultiUserClient/Private/Widgets/AbstractSessionHistoryController.cpp b/Engine/Plugins/Developer/Concert/MultiUserClient/Source/MultiUserClient/Private/Widgets/AbstractSessionHistoryController.cpp
new file mode 100644
--- /dev/null
+++ b/Engine/Plugins/Developer/Concert/MultiUserClient/Source/MultiUserClient/Private/Widgets/AbstractSessionHistoryController.cpp
@@ -0,0 +1,32 @@
+﻿// Copyright Epic Games, Inc. All Rights Reserved.
+
+#include "AbstractSessionHistoryController.h"
+
+#include "SSessionHistory.h"
+#include "Widgets/DeclarativeSyntaxSupport.h"
+
+FAbstractSessionHistoryController::FAbstractSessionHistoryController(FName PackageFilter)
+	: SessionHistory(MakeSessionHistory(PackageFilter))
+{}
+
+void FAbstractSessionHistoryController::ReloadActivities()
+{
+	constexpr int64 MaximumNumberOfActivities = SSessionHistory::MaximumNumberOfActivities;
+	TMap<FGuid, FConcertClientInfo> EndpointClientInfoMap;
+	TArray<FConcertSessionActivity> FetchedActivities;
+	GetActivities(MaximumNumberOfActivities, EndpointClientInfoMap, FetchedActivities);
+
+	// MoveTemp strictly not needed - but it will be faster on Debug builds
+	GetSessionHistory()->ReloadActivities(
+		MoveTemp(EndpointClientInfoMap),
+		MoveTemp(FetchedActivities)
+		);
+}
+
+TSharedRef<SSessionHistory> FAbstractSessionHistoryController::MakeSessionHistory(FName PackageFilter) const
+{
+	return SNew(SSessionHistory)
+		.PackageFilter(PackageFilter)
+		.GetPackageEvent([this](const FConcertSessionActivity& Activity, FConcertSyncPackageEventMetaData& OutPackageEvent) { return GetPackageEvent(Activity, OutPackageEvent); })
+		.GetTransactionEvent([this](const FConcertSessionActivity& Activity) { return GetTransactionEvent(Activity); });
+}
diff --git a/Engine/Plugins/Developer/Concert/MultiUserClient/Source/MultiUserClient/Private/Widgets/AbstractSessionHistoryController.h b/Engine/Plugins/Developer/Concert/MultiUserClient/Source/MultiUserClient/Private/Widgets/AbstractSessionHistoryController.h
new file mode 100644
--- /dev/null
+++ b/Engine/Plugins/Developer/Concert/MultiUserClient/Source/MultiUserClient/Private/Widgets/AbstractSessionHistoryController.h
@@ -0,0 +1,44 @@
+﻿// Copyright Epic Games, Inc. All Rights Reserved.
+
+#pragma once
+
+#include "CoreMinimal.h"
+
+struct FConcertClientInfo;
+struct FConcertSyncTransactionEvent;
+struct FConcertSyncPackageEventMetaData;
+struct FConcertSessionActivity;
+class IConcertSyncClient;
+class SSessionHistory;
+
+/** Acts as controller of SSessionHistory in model-view-component pattern. */
+class FAbstractSessionHistoryController : public TSharedFromThis<FAbstractSessionHistoryController>
+{
+public:
+
+	FAbstractSessionHistoryController(FName PackageFilter = EName::None);
+	virtual ~FAbstractSessionHistoryController() = default;
+	
+	/** Fetches activities from the server and updates the list view. */
+	void ReloadActivities();
+	
+	const TSharedRef<SSessionHistory>& GetSessionHistory() const { return SessionHistory; }
+
+protected:
+
+	/** Gets the latest activities*/
+	virtual void GetActivities(int64 MaximumNumberOfActivities, TMap<FGuid, FConcertClientInfo>& EndpointClientInfoMap, TArray<FConcertSessionActivity>& FetchedActivities) const = 0;
+	
+	/** Returns the specified package event (without the package data itself) if available. */
+	virtual bool GetPackageEvent(const FConcertSessionActivity& Activity, FConcertSyncPackageEventMetaData& OutPackageEvent) const = 0;
+
+	/** Returns the specified package event if available. */
+	virtual TFuture<TOptional<FConcertSyncTransactionEvent>> GetTransactionEvent(const FConcertSessionActivity& Activity) const = 0;
+	
+private:
+
+	/** The widget being managed */
+	TSharedRef<SSessionHistory> SessionHistory;
+
+	TSharedRef<SSessionHistory> MakeSessionHistory(FName PackageFilter) const;
+};
\ No newline at end of file
diff --git a/Engine/Plugins/Developer/Concert/MultiUserClient/Source/MultiUserClient/Private/Widgets/ClientSessionHistoryController.cpp b/Engine/Plugins/Developer/Concert/MultiUserClient/Source/MultiUserClient/Private/Widgets/ClientSessionHistoryController.cpp
new file mode 100644
--- /dev/null
+++ b/Engine/Plugins/Developer/Concert/MultiUserClient/Source/MultiUserClient/Private/Widgets/ClientSessionHistoryController.cpp
@@ -0,0 +1,110 @@
+﻿// Copyright Epic Games, Inc. All Rights Reserved.
+
+#include "ClientSessionHistoryController.h"
+
+#include "ConcertSyncSessionTypes.h"
+#include "IConcertClientWorkspace.h"
+#include "IConcertSyncClient.h"
+
+#include "SSessionHistory.h"
+
+FClientSessionHistoryController::FClientSessionHistoryController(TSharedRef<IConcertSyncClient> Client, FName PackageFilter)
+	: FAbstractSessionHistoryController(PackageFilter)
+	, Client(MoveTemp(Client))
+{
+	SubscribeToWorkspaceEvents();
+}
+
+FClientSessionHistoryController::~FClientSessionHistoryController()
+{
+	if (const TSharedPtr<IConcertSyncClient> ClientPtr = Client.Pin())
+	{
+		ClientPtr->OnWorkspaceStartup().RemoveAll(this);
+		ClientPtr->OnWorkspaceShutdown().RemoveAll(this);
+	}
+	
+	if (const TSharedPtr<IConcertClientWorkspace> WorkspacePtr = Workspace.Pin())
+	{
+		WorkspacePtr->OnActivityAddedOrUpdated().RemoveAll(this);
+		WorkspacePtr->OnWorkspaceSynchronized().RemoveAll(this);
+	}
+}
+
+void FClientSessionHistoryController::GetActivities(int64 MaximumNumberOfActivities, TMap<FGuid, FConcertClientInfo>& EndpointClientInfoMap, TArray<FConcertSessionActivity>& FetchedActivities) const
+{
+	if (const TSharedPtr<IConcertClientWorkspace> WorkspacePtr = Workspace.Pin())
+	{
+		const int64 LastActivityId = WorkspacePtr->GetLastActivityId();
+		const int64 FirstActivityIdToFetch = FMath::Max<int64>(1, LastActivityId - MaximumNumberOfActivities);
+		WorkspacePtr->GetActivities(FirstActivityIdToFetch, MaximumNumberOfActivities, EndpointClientInfoMap, FetchedActivities);
+	}
+}
+
+bool FClientSessionHistoryController::GetPackageEvent(const FConcertSessionActivity& Activity, FConcertSyncPackageEventMetaData& OutPackageEvent) const
+{
+	if (TSharedPtr<IConcertClientWorkspace> WorkspacePtr = Workspace.Pin())
+	{
+		// Don't request the package data, the widget only display the meta-data.
+		return WorkspacePtr->FindPackageEvent(Activity.Activity.EventId, OutPackageEvent);
+	}
+
+	return false; // The data is not available.
+}
+
+TFuture<TOptional<FConcertSyncTransactionEvent>> FClientSessionHistoryController::GetTransactionEvent(const FConcertSessionActivity& Activity) const
+{
+	if (TSharedPtr<IConcertClientWorkspace> WorkspacePtr = Workspace.Pin())
+	{
+		// Ask to get the full transaction to display which properties changed.
+		return WorkspacePtr->FindOrRequestTransactionEvent(Activity.Activity.EventId, /*bMetaDataOnly*/false);
+	}
+
+	return MakeFulfilledPromise<TOptional<FConcertSyncTransactionEvent>>().GetFuture();
+}
+
+void FClientSessionHistoryController::SubscribeToWorkspaceEvents()
+{
+	TSharedPtr<IConcertSyncClient> ClientPtr = Client.Pin();
+	checkf(ClientPtr, TEXT("We're supposed to be in the constructor where the Client should still be valid!"));
+	
+	// Must use AddRaw instead of AddSP because we're called from constructor (before TSharedFromThis is initialised)
+	ClientPtr->OnWorkspaceStartup().AddRaw(this, &FClientSessionHistoryController::HandleWorkspaceStartup);
+	ClientPtr->OnWorkspaceShutdown().AddRaw(this, &FClientSessionHistoryController::HandleWorkspaceShutdown);
+
+	if (TSharedPtr<IConcertClientWorkspace> WorkspacePtr = ClientPtr->GetWorkspace())
+	{
+		Workspace = WorkspacePtr;
+		RegisterWorkspaceHandler();
+		ReloadActivities();
+	}
+}
+
+void FClientSessionHistoryController::RegisterWorkspaceHandler()
+{
+	TSharedPtr<IConcertClientWorkspace> WorkspacePtr = Workspace.Pin();
+	if (WorkspacePtr.IsValid())
+	{
+		// Must use AddRaw instead of AddSP because we may be called from constructor (before TSharedFromThis is initialised)
+		WorkspacePtr->OnActivityAddedOrUpdated().AddRaw(this, &FClientSessionHistoryController::HandleActivityAddedOrUpdated);
+		WorkspacePtr->OnWorkspaceSynchronized().AddRaw(this, &FClientSessionHistoryController::ReloadActivities);
+	}
+}
+void FClientSessionHistoryController::HandleWorkspaceStartup(const TSharedPtr<IConcertClientWorkspace>& NewWorkspace)
+{
+	Workspace = NewWorkspace;
+	RegisterWorkspaceHandler();
+}
+
+void FClientSessionHistoryController::HandleWorkspaceShutdown(const TSharedPtr<IConcertClientWorkspace>& WorkspaceShuttingDown)
+{
+	if (WorkspaceShuttingDown == Workspace)
+	{
+		Workspace.Reset();
+		ReloadActivities();
+	}
+}
+
+void FClientSessionHistoryController::HandleActivityAddedOrUpdated(const FConcertClientInfo& InClientInfo, const FConcertSyncActivity& InActivity, const FStructOnScope& InActivitySummary)
+{
+	GetSessionHistory()->HandleActivityAddedOrUpdated(InClientInfo, InActivity, InActivitySummary);
+}
diff --git a/Engine/Plugins/Developer/Concert/MultiUserClient/Source/MultiUserClient/Private/Widgets/ClientSessionHistoryController.h b/Engine/Plugins/Developer/Concert/MultiUserClient/Source/MultiUserClient/Private/Widgets/ClientSessionHistoryController.h
new file mode 100644
--- /dev/null
+++ b/Engine/Plugins/Developer/Concert/MultiUserClient/Source/MultiUserClient/Private/Widgets/ClientSessionHistoryController.h
@@ -0,0 +1,46 @@
+﻿// Copyright Epic Games, Inc. All Rights Reserved.
+
+#pragma once
+
+#include "CoreMinimal.h"
+#include "AbstractSessionHistoryController.h"
+
+struct FConcertSyncActivity;
+struct FConcertClientInfo;
+class IConcertClientWorkspace;
+class IConcertSyncClient;
+
+/** Manages SSessionHistory for IConcertClients */
+class FClientSessionHistoryController : public FAbstractSessionHistoryController
+{
+public:
+
+	FClientSessionHistoryController(TSharedRef<IConcertSyncClient> Client, FName PackageFilter = EName::None);
+	virtual ~FClientSessionHistoryController() override;
+
+protected:
+
+	//~ Begin FAbstractSessionHistoryController Interface
+	virtual void GetActivities(int64 MaximumNumberOfActivities, TMap<FGuid, FConcertClientInfo>& EndpointClientInfoMap, TArray<FConcertSessionActivity>& FetchedActivities) const override;
+	virtual bool GetPackageEvent(const FConcertSessionActivity& Activity, FConcertSyncPackageEventMetaData& OutPackageEvent) const override;
+	virtual TFuture<TOptional<FConcertSyncTransactionEvent>> GetTransactionEvent( const FConcertSessionActivity& Activity) const override;
+	//~ End FAbstractSessionHistoryController Interface
+	
+private:
+
+	/** So we can unsubscribe from events when destroyed */
+	TWeakPtr<IConcertSyncClient> Client;
+	
+	/** Holds a weak pointer to the current workspace. */ 
+	TWeakPtr<IConcertClientWorkspace> Workspace;
+	
+	void SubscribeToWorkspaceEvents();
+	
+	/** Registers callbacks with the current workspace. */
+	void RegisterWorkspaceHandler();
+
+	// Workspace events
+	void HandleWorkspaceStartup(const TSharedPtr<IConcertClientWorkspace>& NewWorkspace);
+	void HandleWorkspaceShutdown(const TSharedPtr<IConcertClientWorkspace>& WorkspaceShuttingDown);
+	void HandleActivityAddedOrUpdated(const FConcertClientInfo& InClientInfo, const FConcertSyncActivity& InActivity, const FStructOnScope& InActivitySummary);
+};
diff --git a/Engine/Plugins/Developer/Concert/MultiUserClient/Source/MultiUserClient/Private/Widgets/SActiveSession.cpp b/Engine/Plugins/Developer/Concert/MultiUserClient/Source/MultiUserClient/Private/Widgets/SActiveSession.cpp
--- a/Engine/Plugins/Developer/Concert/MultiUserClient/Source/MultiUserClient/Private/Widgets/SActiveSession.cpp
+++ b/Engine/Plugins/Developer/Concert/MultiUserClient/Source/MultiUserClient/Private/Widgets/SActiveSession.cpp
@@ -1,28 +1,30 @@
 // Copyright Epic Games, Inc. All Rights Reserved.
 
 #include "SActiveSession.h"
-#include "Widgets/Docking/SDockTab.h"
+
+#include "ConcertActionDefinition.h"
+#include "ConcertClientFrontendUtils.h"
+#include "ConcertFrontendUtils.h"
+#include "ConcertMessageData.h"
+#include "ClientSessionHistoryController.h"
+#include "IConcertClientPresenceManager.h"
+#include "IConcertClient.h"
+#include "IConcertSyncClient.h"
+#include "IMultiUserClientModule.h"
+
+#include "Algo/Transform.h"
+#include "EditorFontGlyphs.h"
+#include "FileHelpers.h"
+#include "Misc/PackageName.h"
+#include "Styling/AppStyle.h"
+#include "SSessionHistory.h"
 #include "Widgets/Text/STextBlock.h"
 #include "Widgets/Images/SImage.h"
 #include "Widgets/Input/SHyperlink.h"
 #include "Widgets/Layout/SExpandableArea.h"
 #include "Widgets/Views/STableViewBase.h"
 #include "Widgets/Views/STableRow.h"
-#include "Framework/Docking/WorkspaceItem.h"
-#include "Algo/Transform.h"
-#include "Misc/PackageName.h"
-#include "EditorFontGlyphs.h"
-#include "Styling/AppStyle.h"
-#include "FileHelpers.h"
-#include "IConcertClient.h"
-#include "IConcertSyncClient.h"
-#include "IMultiUserClientModule.h"
-#include "IConcertClientPresenceManager.h"
-#include "ConcertActionDefinition.h"
-#include "ConcertFrontendUtils.h"
-#include "ConcertMessageData.h"
-#include "SSessionHistory.h"
-#include "ConcertClientFrontendUtils.h"
+#include "Widgets/SSessionHistoryWrapper.h"
 
 #define LOCTEXT_NAMESPACE "SActiveSession"
 
@@ -329,6 +331,8 @@ private:
 void SActiveSession::Construct(const FArguments& InArgs, TSharedPtr<IConcertSyncClient> InConcertSyncClient)
 {
 	WeakConcertSyncClient = InConcertSyncClient;
+	SessionHistoryController = MakeShared<FClientSessionHistoryController>(InConcertSyncClient.ToSharedRef());
+	
 	if (InConcertSyncClient.IsValid())
 	{
 		IConcertClientRef ConcertClient = InConcertSyncClient->GetConcertClient();
@@ -501,7 +505,7 @@ void SActiveSession::Construct(const FArguments& InArgs, TSharedPtr<IConcertSync
 					]
 					.BodyContent()
 					[
-						SNew(SSessionHistory, WeakConcertSyncClient.Pin())
+						SessionHistoryController->GetSessionHistory()
 					]
 				]
 			]
@@ -530,9 +534,9 @@ void SActiveSession::HandleSessionStartup(TSharedRef<IConcertClientSession> InCl
 
 	UpdateSessionClientListView();
 
-	if (SessionHistory.IsValid())
+	if (SessionHistoryController.IsValid())
 	{
-		SessionHistory->Refresh();
+		SessionHistoryController->ReloadActivities();
 	}
 }
 
@@ -549,9 +553,9 @@ void SActiveSession::HandleSessionShutdown(TSharedRef<IConcertClientSession> InC
 			ClientsListView->RequestListRefresh();
 		}
 
-		if (SessionHistory.IsValid())
+		if (SessionHistoryController.IsValid())
 		{
-			SessionHistory->Refresh();
+			SessionHistoryController->ReloadActivities();
 		}
 	}
 }
diff --git a/Engine/Plugins/Developer/Concert/MultiUserClient/Source/MultiUserClient/Private/Widgets/SActiveSession.h b/Engine/Plugins/Developer/Concert/MultiUserClient/Source/MultiUserClient/Private/Widgets/SActiveSession.h
--- a/Engine/Plugins/Developer/Concert/MultiUserClient/Source/MultiUserClient/Private/Widgets/SActiveSession.h
+++ b/Engine/Plugins/Developer/Concert/MultiUserClient/Source/MultiUserClient/Private/Widgets/SActiveSession.h
@@ -3,10 +3,10 @@
 #pragma once
 
 #include "CoreMinimal.h"
+#include "ClientSessionHistoryController.h"
 #include "Widgets/DeclarativeSyntaxSupport.h"
 #include "Widgets/SCompoundWidget.h"
 #include "Widgets/Views/SListView.h"
-#include "IConcertModule.h"
 #include "ConcertMessages.h"
 
 class IConcertClientSession;
@@ -83,6 +83,7 @@ private:
 	void OnHistoryAreaExpansionChanged(bool bExpanded) { bHistoryAreaExpanded = bExpanded; }
 
 private:
+	
 	/** Pointer on the client sync. */
 	TWeakPtr<IConcertSyncClient> WeakConcertSyncClient;
 
@@ -99,7 +100,7 @@ private:
 	TSharedPtr<FConcertSessionClientInfo> ClientInfo;
 
 	/** Holds a concert activity log. */
-	TSharedPtr<class SSessionHistory> SessionHistory;
+	TSharedPtr<FClientSessionHistoryController> SessionHistoryController;
 
 	/** The 'Clients' expandable area. */
 	TSharedPtr<SExpandableArea> ClientArea;
diff --git a/Engine/Plugins/Developer/Concert/MultiUserClient/Source/MultiUserClient/Private/Widgets/SSessionHistory.cpp b/Engine/Plugins/Developer/Concert/MultiUserClient/Source/MultiUserClient/Private/Widgets/SSessionHistory.cpp
--- a/Engine/Plugins/Developer/Concert/MultiUserClient/Source/MultiUserClient/Private/Widgets/SSessionHistory.cpp
+++ b/Engine/Plugins/Developer/Concert/MultiUserClient/Source/MultiUserClient/Private/Widgets/SSessionHistory.cpp
@@ -42,7 +42,7 @@ namespace ConcertSessionHistoryUI
 	}
 }
 
-void SSessionHistory::Construct(const FArguments& InArgs, TSharedPtr<IConcertSyncClient> InConcertSyncClient)
+void SSessionHistory::Construct(const FArguments& InArgs)
 {
 	PackageNameFilter = InArgs._PackageFilter;
 
@@ -50,8 +50,8 @@ void SSessionHistory::Construct(const FArguments& InArgs, TSharedPtr<IConcertSyn
 	ActivityListViewOptions = MakeShared<FConcertSessionActivitiesOptions>();
 
 	SAssignNew(ActivityListView, SConcertSessionActivities)
-		.OnGetPackageEvent([this](const FConcertSessionActivity& Activity, FConcertSyncPackageEventMetaData& OutEvent) { return GetPackageEvent(Activity, OutEvent); })
-		.OnGetTransactionEvent([this](const FConcertSessionActivity& Activity) { return GetTransactionEvent(Activity); })
+		.OnGetPackageEvent(InArgs._GetPackageEvent)
+		.OnGetTransactionEvent(InArgs._GetTransactionEvent)
 		.OnMapActivityToClient([this](FGuid ClientId){ return EndpointClientInfoMap.Find(ClientId); })
 		.HighlightText(this, &SSessionHistory::HighlightSearchedText)
 		.TimeFormat(ActivityListViewOptions.Get(), &FConcertSessionActivitiesOptions::GetTimeFormat)
@@ -101,67 +101,21 @@ void SSessionHistory::Construct(const FArguments& InArgs, TSharedPtr<IConcertSyn
 				TAttribute<int32>(ActivityListView.Get(), &SConcertSessionActivities::GetDisplayedActivityNum))
 		]
 	];
-
-	if (InConcertSyncClient.IsValid())
-	{
-		InConcertSyncClient->OnWorkspaceStartup().AddSP(this, &SSessionHistory::HandleWorkspaceStartup);
-		InConcertSyncClient->OnWorkspaceShutdown().AddSP(this, &SSessionHistory::HandleWorkspaceShutdown);
-
-		if (TSharedPtr<IConcertClientWorkspace> WorkspacePtr = InConcertSyncClient->GetWorkspace())
-		{
-			Workspace = WorkspacePtr;
-			RegisterWorkspaceHandler();
-			ReloadActivities();
-		}
-	}
-}
-
-void SSessionHistory::Refresh()
-{
-	ReloadActivities();
-}
-
-void SSessionHistory::OnSearchTextChanged(const FText& InSearchText)
-{
-	SearchedText = InSearchText;
-	SearchBox->SetError(ActivityListView->UpdateTextFilter(InSearchText));
-}
-
-void SSessionHistory::OnSearchTextCommitted(const FText& InSearchText, ETextCommit::Type CommitType)
-{
-	if (!InSearchText.EqualTo(SearchedText))
-	{
-		OnSearchTextChanged(InSearchText);
-	}
-}
-
-FText SSessionHistory::HighlightSearchedText() const
-{
-	return SearchedText;
 }
 
-void SSessionHistory::ReloadActivities()
+void SSessionHistory::ReloadActivities(TMap<FGuid, FConcertClientInfo> InEndpointClientInfoMap, TArray<FConcertSessionActivity> InFetchedActivities)
 {
-	EndpointClientInfoMap.Reset();
+	EndpointClientInfoMap = MoveTemp(InEndpointClientInfoMap);
 	ActivityMap.Reset();
 	ActivityListView->Reset(); // Careful, don't reset the shared ptr.
 
-	if (TSharedPtr<IConcertClientWorkspace> WorkspacePtr = Workspace.Pin())
+	for (FConcertSessionActivity& FetchedActivity : InFetchedActivities)
 	{
-		const int64 LastActivityId = WorkspacePtr->GetLastActivityId();
-		const int64 FirstActivityIdToFetch = FMath::Max<int64>(1, LastActivityId - MaximumNumberOfActivities);
-
-		TArray<FConcertSessionActivity> FetchedActivities;
-		WorkspacePtr->GetActivities(FirstActivityIdToFetch, MaximumNumberOfActivities, EndpointClientInfoMap, FetchedActivities);
-
-		for (FConcertSessionActivity& FetchedActivity : FetchedActivities)
+		if (ConcertSessionHistoryUI::PackageNamePassesFilter(PackageNameFilter, FetchedActivity.ActivitySummary))
 		{
-			if (ConcertSessionHistoryUI::PackageNamePassesFilter(PackageNameFilter, FetchedActivity.ActivitySummary))
-			{
-				TSharedRef<FConcertSessionActivity> NewActivity = MakeShared<FConcertSessionActivity>(MoveTemp(FetchedActivity));
-				ActivityMap.Add(NewActivity->Activity.ActivityId, NewActivity);
-				ActivityListView->Append(NewActivity);
-			}
+			TSharedRef<FConcertSessionActivity> NewActivity = MakeShared<FConcertSessionActivity>(MoveTemp(FetchedActivity));
+			ActivityMap.Add(NewActivity->Activity.ActivityId, NewActivity);
+			ActivityListView->Append(NewActivity);
 		}
 	}
 }
@@ -193,52 +147,23 @@ void SSessionHistory::HandleActivityAddedOrUpdated(const FConcertClientInfo& InC
 	}
 }
 
-void SSessionHistory::HandleWorkspaceStartup(const TSharedPtr<IConcertClientWorkspace>& NewWorkspace)
-{
-	Workspace = NewWorkspace;
-	RegisterWorkspaceHandler();
-}
-
-void SSessionHistory::HandleWorkspaceShutdown(const TSharedPtr<IConcertClientWorkspace>& WorkspaceShuttingDown)
-{
-	if (WorkspaceShuttingDown == Workspace)
-	{
-		Workspace.Reset();
-		ReloadActivities();
-	}
-}
-
-void SSessionHistory::RegisterWorkspaceHandler()
+void SSessionHistory::OnSearchTextChanged(const FText& InSearchText)
 {
-	TSharedPtr<IConcertClientWorkspace> WorkspacePtr = Workspace.Pin();
-	if (WorkspacePtr.IsValid())
-	{
-		WorkspacePtr->OnActivityAddedOrUpdated().AddSP(this, &SSessionHistory::HandleActivityAddedOrUpdated);
-		WorkspacePtr->OnWorkspaceSynchronized().AddSP(this, &SSessionHistory::ReloadActivities);
-	}
+	SearchedText = InSearchText;
+	SearchBox->SetError(ActivityListView->UpdateTextFilter(InSearchText));
 }
 
-bool SSessionHistory::GetPackageEvent(const FConcertSessionActivity& Activity, FConcertSyncPackageEventMetaData& OutPackageEvent) const
+void SSessionHistory::OnSearchTextCommitted(const FText& InSearchText, ETextCommit::Type CommitType)
 {
-	if (TSharedPtr<IConcertClientWorkspace> WorkspacePtr = Workspace.Pin())
+	if (!InSearchText.EqualTo(SearchedText))
 	{
-		// Don't request the package data, the widget only display the meta-data.
-		return WorkspacePtr->FindPackageEvent(Activity.Activity.EventId, OutPackageEvent);
+		OnSearchTextChanged(InSearchText);
 	}
-
-	return false; // The data is not available.
 }
 
-TFuture<TOptional<FConcertSyncTransactionEvent>> SSessionHistory::GetTransactionEvent(const FConcertSessionActivity& Activity) const
+FText SSessionHistory::HighlightSearchedText() const
 {
-	if (TSharedPtr<IConcertClientWorkspace> WorkspacePtr = Workspace.Pin())
-	{
-		// Ask to get the full transaction to display which properties changed.
-		return WorkspacePtr->FindOrRequestTransactionEvent(Activity.Activity.EventId, /*bMetaDataOnly*/false);
-	}
-
-	return MakeFulfilledPromise<TOptional<FConcertSyncTransactionEvent>>().GetFuture();
+	return SearchedText;
 }
 
-
 #undef LOCTEXT_NAMESPACE /* SSessionHistory */
diff --git a/Engine/Plugins/Developer/Concert/MultiUserClient/Source/MultiUserClient/Private/Widgets/SSessionHistory.h b/Engine/Plugins/Developer/Concert/MultiUserClient/Source/MultiUserClient/Private/Widgets/SSessionHistory.h
--- a/Engine/Plugins/Developer/Concert/MultiUserClient/Source/MultiUserClient/Private/Widgets/SSessionHistory.h
+++ b/Engine/Plugins/Developer/Concert/MultiUserClient/Source/MultiUserClient/Private/Widgets/SSessionHistory.h
@@ -3,14 +3,10 @@
 #pragma once
 
 #include "CoreMinimal.h"
-#include "IConcertClientWorkspace.h"
-#include "ConcertSyncSessionTypes.h"
-#include "Async/Future.h"
+#include "SConcertSessionActivities.h"
 #include "Widgets/DeclarativeSyntaxSupport.h"
 #include "Widgets/SCompoundWidget.h"
-#include "Widgets/Views/SListView.h"
 
-class IConcertSyncClient;
 class SSearchBox;
 class SConcertSessionActivities;
 class FConcertSessionActivitiesOptions;
@@ -18,8 +14,14 @@ class FConcertSessionActivitiesOptions;
 class SSessionHistory : public SCompoundWidget
 {
 public:
+
+	/** Maximum number of activities displayed on screen. */ 
+	static constexpr int64 MaximumNumberOfActivities = 1000;
+	
 	SLATE_BEGIN_ARGS(SSessionHistory) {}
 		SLATE_ARGUMENT(FName, PackageFilter)
+		SLATE_ARGUMENT(SConcertSessionActivities::FGetPackageEvent, GetPackageEvent)
+		SLATE_ARGUMENT(SConcertSessionActivities::FGetTransactionEvent, GetTransactionEvent)
 	SLATE_END_ARGS()
 
 	/**
@@ -29,29 +31,15 @@ public:
 	* @param ConstructUnderMajorTab The major tab which will contain the Session History widget.
 	* @param InConcertSyncClient Pointer on the concert sync client.
 	*/
-	void Construct(const FArguments& InArgs, TSharedPtr<IConcertSyncClient> InConcertSyncClient);
-
-	/** Fetches the activities and updates the UI. */
-	void Refresh();
-
-private:
-	/** Callback for selecting an activity in the list view. */
-	void HandleSelectionChanged(TSharedPtr<FConcertSessionActivity> InSessionActivity, ESelectInfo::Type SelectInfo);
+	void Construct(const FArguments& InArgs);
 
 	/** Fetches activities from the server and updates the list view. */
-	void ReloadActivities();
-
+	void ReloadActivities(TMap<FGuid, FConcertClientInfo> EndpointClientInfoMap, TArray<FConcertSessionActivity> FetchedActivities);
+	
 	/** Callback for handling the a new or updated activity item. */ 
 	void HandleActivityAddedOrUpdated(const FConcertClientInfo& InClientInfo, const FConcertSyncActivity& InActivity, const FStructOnScope& InActivitySummary);
 
-	/** Callback for handling the startup of a workspace.  */
-	void HandleWorkspaceStartup(const TSharedPtr<IConcertClientWorkspace>& NewWorkspace);
-
-	/** Callback for handling the shutdown of a workspace. */
-	void HandleWorkspaceShutdown(const TSharedPtr<IConcertClientWorkspace>& WorkspaceShuttingDown);
-
-	/** Registers callbacks with the current workspace. */
-	void RegisterWorkspaceHandler();
+private:
 
 	/** Invoked when the text in the search box widget changes. */
 	void OnSearchTextChanged(const FText& InSearchText);
@@ -69,9 +57,6 @@ private:
 	TFuture<TOptional<FConcertSyncTransactionEvent>> GetTransactionEvent(const FConcertSessionActivity& Activity) const;
 
 private:
-
-	/** Maximum number of activities displayed on screen. */ 
-	static const int64 MaximumNumberOfActivities = 1000;
 	
 	/** Holds the map of endpoint IDs to client info. */
 	TMap<FGuid, FConcertClientInfo> EndpointClientInfoMap;
@@ -85,9 +70,6 @@ private:
 	/** Controls the activity list view options */
 	TSharedPtr<FConcertSessionActivitiesOptions> ActivityListViewOptions;
 
-	/** Holds a weak pointer to the current workspace. */ 
-	TWeakPtr<IConcertClientWorkspace> Workspace;
-
 	/** The widget used to enter the text to search. */
 	TSharedPtr<SSearchBox> SearchBox;
 
diff --git a/Engine/Plugins/Developer/Concert/MultiUserClient/Source/MultiUserClient/Private/Widgets/SSessionHistoryWrapper.h b/Engine/Plugins/Developer/Concert/MultiUserClient/Source/MultiUserClient/Private/Widgets/SSessionHistoryWrapper.h
new file mode 100644
--- /dev/null
+++ b/Engine/Plugins/Developer/Concert/MultiUserClient/Source/MultiUserClient/Private/Widgets/SSessionHistoryWrapper.h
@@ -0,0 +1,32 @@
+﻿// Copyright Epic Games, Inc. All Rights Reserved.
+
+#pragma once
+
+#include "CoreMinimal.h"
+#include "AbstractSessionHistoryController.h"
+#include "SSessionHistory.h"
+#include "Widgets/SCompoundWidget.h"
+#include "Widgets/DeclarativeSyntaxSupport.h"
+
+/** Shows a SSessionHistory widget - keeping alive the controller as long as this widget lives. Useful for closeable tabs. */
+class SSessionHistoryWrapper : public SCompoundWidget
+{
+public:
+
+	SLATE_BEGIN_ARGS(SSessionHistoryWrapper) {}
+	SLATE_END_ARGS()
+
+	void Construct(const FArguments& InArgs, TSharedRef<FAbstractSessionHistoryController> InController)
+	{
+		Controller = InController;
+		ChildSlot
+		[
+			Controller->GetSessionHistory()
+		];
+	}
+
+private:
+
+	/** Manages SessionHistory */
+	TSharedPtr<FAbstractSessionHistoryController> Controller;
+};
