Open archived session in tab

diff --git a/Engine/Plugins/Developer/Concert/ConcertUI/ConcertServerUI/Source/ConcertServerUI/Private/Widgets/Browser/ConcertServerSessionBrowserController.cpp b/Engine/Plugins/Developer/Concert/ConcertUI/ConcertServerUI/Source/ConcertServerUI/Private/Widgets/Browser/ConcertServerSessionBrowserController.cpp
--- a/Engine/Plugins/Developer/Concert/ConcertUI/ConcertServerUI/Source/ConcertServerUI/Private/Widgets/Browser/ConcertServerSessionBrowserController.cpp
+++ b/Engine/Plugins/Developer/Concert/ConcertUI/ConcertServerUI/Source/ConcertServerUI/Private/Widgets/Browser/ConcertServerSessionBrowserController.cpp
@@ -147,7 +147,8 @@ TSharedRef<SDockTab> FConcertServerSessionBrowserController::SpawnSessionBrowser
 		.TabRole(MajorTab)
 		[
 			SAssignNew(ConcertBrowser, SConcertServerSessionBrowser, SharedThis(this))
-				.DoubleClickSession(this, &FConcertServerSessionBrowserController::OpenSession)
+				.DoubleClickLiveSession(this, &FConcertServerSessionBrowserController::OpenSession)
+				.DoubleClickArchivedSession(this, &FConcertServerSessionBrowserController::OpenSession)
 		];
 
 	FGlobalTabmanager::Get()->SetMainTab(DockTab);
@@ -164,7 +165,7 @@ void FConcertServerSessionBrowserController::RefreshSessionList()
 	}
 }
 
-void FConcertServerSessionBrowserController::OpenSession(TSharedPtr<FConcertSessionItem> SessionItem)
+void FConcertServerSessionBrowserController::OpenSession(const TSharedPtr<FConcertSessionItem>& SessionItem)
 {
 	Owner.Pin()->OpenSessionTab(SessionItem->SessionId);
 }
diff --git a/Engine/Plugins/Developer/Concert/ConcertUI/ConcertServerUI/Source/ConcertServerUI/Private/Widgets/Browser/ConcertServerSessionBrowserController.h b/Engine/Plugins/Developer/Concert/ConcertUI/ConcertServerUI/Source/ConcertServerUI/Private/Widgets/Browser/ConcertServerSessionBrowserController.h
--- a/Engine/Plugins/Developer/Concert/ConcertUI/ConcertServerUI/Source/ConcertServerUI/Private/Widgets/Browser/ConcertServerSessionBrowserController.h
+++ b/Engine/Plugins/Developer/Concert/ConcertUI/ConcertServerUI/Source/ConcertServerUI/Private/Widgets/Browser/ConcertServerSessionBrowserController.h
@@ -72,7 +72,7 @@ private:
 	void RefreshSessionList();
 
 	// Session actions
-	void OpenSession(TSharedPtr<FConcertSessionItem> SessionItem);
+	void OpenSession(const TSharedPtr<FConcertSessionItem>& SessionItem);
 	void RenameSession(const FGuid& ServerAdminEndpointId, const FGuid& SessionId, const FString& NewName);
 	void DeleteSession(const FGuid& ServerAdminEndpointId, const FGuid& SessionId);
 
diff --git a/Engine/Plugins/Developer/Concert/ConcertUI/ConcertServerUI/Source/ConcertServerUI/Private/Widgets/Browser/SConcertServerSessionBrowser.cpp b/Engine/Plugins/Developer/Concert/ConcertUI/ConcertServerUI/Source/ConcertServerUI/Private/Widgets/Browser/SConcertServerSessionBrowser.cpp
--- a/Engine/Plugins/Developer/Concert/ConcertUI/ConcertServerUI/Source/ConcertServerUI/Private/Widgets/Browser/SConcertServerSessionBrowser.cpp
+++ b/Engine/Plugins/Developer/Concert/ConcertUI/ConcertServerUI/Source/ConcertServerUI/Private/Widgets/Browser/SConcertServerSessionBrowser.cpp
@@ -55,14 +55,15 @@ TSharedRef<SWidget> SConcertServerSessionBrowser::MakeSessionTableView(const FAr
 {
 	SearchText = MakeShared<FText>();
 	return SAssignNew(SessionBrowser, SConcertSessionBrowser, Controller.Pin().ToSharedRef(), SearchText)
-		.OnSessionDoubleClicked(InArgs._DoubleClickSession)
+		.OnLiveSessionDoubleClicked(InArgs._DoubleClickLiveSession)
+		.OnArchivedSessionDoubleClicked(InArgs._DoubleClickArchivedSession)
 		.OnRequestedDeleteSession(this, &SConcertServerSessionBrowser::RequestDeleteSession)
 		// Pretend a modal dialog said no - RequestDeleteSession will show non-modal dialog
 		.CanDeleteArchivedSession_Lambda([](TSharedPtr<FConcertSessionItem>) { return false; })
 		.CanDeleteActiveSession_Lambda([](TSharedPtr<FConcertSessionItem>) { return false; });
 }
 
-void SConcertServerSessionBrowser::RequestDeleteSession(TSharedPtr<FConcertSessionItem> SessionItem) 
+void SConcertServerSessionBrowser::RequestDeleteSession(const TSharedPtr<FConcertSessionItem>& SessionItem) 
 {
 	if (const TSharedPtr<SMessageDialog> PinnedDeleteSessionDialog = DeleteSessionDialog.Pin())
 	{
@@ -110,7 +111,7 @@ void SConcertServerSessionBrowser::UnregisterFromOnRootWindowClosed() const
 	}
 }
 
-void SConcertServerSessionBrowser::DeleteArchivedSessionWithNonModalQuestion(TSharedPtr<FConcertSessionItem> SessionItem)
+void SConcertServerSessionBrowser::DeleteArchivedSessionWithNonModalQuestion(const TSharedPtr<FConcertSessionItem>& SessionItem)
 {
 	const FText Message = FText::Format(
 	LOCTEXT("DeleteArchivedDescription", "Deleting a session will cause all associated data to be removed.\n\nDelete {0}?"),
@@ -140,7 +141,7 @@ void SConcertServerSessionBrowser::DeleteArchivedSessionWithNonModalQuestion(TSh
 	Dialog->Show();
 }
 
-void SConcertServerSessionBrowser::DeleteActiveSessionWithNonModalQuestion(TSharedPtr<FConcertSessionItem> SessionItem)
+void SConcertServerSessionBrowser::DeleteActiveSessionWithNonModalQuestion(const TSharedPtr<FConcertSessionItem>& SessionItem)
 {
 	const int32 NumUsers = Controller.Pin()->GetNumConnectedClients(SessionItem->SessionId);
 	const FText Message = FText::Format(
diff --git a/Engine/Plugins/Developer/Concert/ConcertUI/ConcertServerUI/Source/ConcertServerUI/Private/Widgets/Browser/SConcertServerSessionBrowser.h b/Engine/Plugins/Developer/Concert/ConcertUI/ConcertServerUI/Source/ConcertServerUI/Private/Widgets/Browser/SConcertServerSessionBrowser.h
--- a/Engine/Plugins/Developer/Concert/ConcertUI/ConcertServerUI/Source/ConcertServerUI/Private/Widgets/Browser/SConcertServerSessionBrowser.h
+++ b/Engine/Plugins/Developer/Concert/ConcertUI/ConcertServerUI/Source/ConcertServerUI/Private/Widgets/Browser/SConcertServerSessionBrowser.h
@@ -17,9 +17,10 @@ class SConcertServerSessionBrowser : public SCompoundWidget
 public:
 	
 	SLATE_BEGIN_ARGS(SConcertServerSessionBrowser) { }
-		SLATE_EVENT(FSessionDelegate, DoubleClickSession)
+		SLATE_EVENT(FSessionDelegate, DoubleClickLiveSession)
+		SLATE_EVENT(FSessionDelegate, DoubleClickArchivedSession)
 	SLATE_END_ARGS()
-	~SConcertServerSessionBrowser();
+	virtual ~SConcertServerSessionBrowser() override;
 
 	void Construct(const FArguments& InArgs, TSharedRef<FConcertServerSessionBrowserController> InController);
 
@@ -49,9 +50,9 @@ private:
 	
 	TSharedRef<SWidget> MakeSessionTableView(const FArguments& InArgs);
 
-	void RequestDeleteSession(TSharedPtr<FConcertSessionItem> SessionItem);
+	void RequestDeleteSession(const TSharedPtr<FConcertSessionItem>& SessionItem);
 	void OnRootWindowClosed(const TSharedRef<SWindow>&) const;
 	void UnregisterFromOnRootWindowClosed() const;
-	void DeleteArchivedSessionWithNonModalQuestion(TSharedPtr<FConcertSessionItem> SessionItem);
-	void DeleteActiveSessionWithNonModalQuestion(TSharedPtr<FConcertSessionItem> SessionItem);
+	void DeleteArchivedSessionWithNonModalQuestion(const TSharedPtr<FConcertSessionItem>& SessionItem);
+	void DeleteActiveSessionWithNonModalQuestion(const TSharedPtr<FConcertSessionItem>& SessionItem);
 };
diff --git a/Engine/Plugins/Developer/Concert/ConcertUI/ConcertServerUI/Source/ConcertServerUI/Private/Widgets/ConcertServerWindowController.cpp b/Engine/Plugins/Developer/Concert/ConcertUI/ConcertServerUI/Source/ConcertServerUI/Private/Widgets/ConcertServerWindowController.cpp
--- a/Engine/Plugins/Developer/Concert/ConcertUI/ConcertServerUI/Source/ConcertServerUI/Private/Widgets/ConcertServerWindowController.cpp
+++ b/Engine/Plugins/Developer/Concert/ConcertUI/ConcertServerUI/Source/ConcertServerUI/Private/Widgets/ConcertServerWindowController.cpp
@@ -14,7 +14,10 @@
 #include "HAL/PlatformApplicationMisc.h"
 #include "Misc/ConfigCacheIni.h"
 #include "OutputLog/Public/OutputLogModule.h"
-#include "Session/ConcertServerSessionTab.h"
+#include "Session/LiveConcertSessionTab.h"
+
+#include "Tabs/Archived/ArchivedConcertSessionTab.h"
+
 #include "Widgets/SWindow.h"
 
 #define LOCTEXT_NAMESPACE "UnrealMultiUserUI"
@@ -73,30 +76,37 @@ void FConcertServerWindowController::CreateWindow()
 
 	RootWindow->SetOnWindowClosed(FOnWindowClosed::CreateRaw(this, &FConcertServerWindowController::OnWindowClosed));
 	RootWindow->ShowWindow();
-	const bool bForceWindowToFront = true;
+	constexpr bool bForceWindowToFront = true;
 	RootWindow->BringToFront(bForceWindowToFront);
 }
 
 void FConcertServerWindowController::OpenSessionTab(const FGuid& SessionId)
 {
-	if (TSharedPtr<FConcertServerSessionTab> SessionTab = GetOrRegisterSessionTab(SessionId))
+	if (TSharedPtr<FAbstractConcertSessionTab> SessionTab = GetOrRegisterSessionTab(SessionId))
 	{
 		SessionTab->OpenSessionTab();
 	}
 }
 
-TSharedPtr<FConcertServerSessionTab> FConcertServerWindowController::GetOrRegisterSessionTab(const FGuid& SessionId)
+TSharedPtr<FAbstractConcertSessionTab> FConcertServerWindowController::GetOrRegisterSessionTab(const FGuid& SessionId)
 {
-	if (TSharedPtr<IConcertServerSession> Session = ServerInstance->GetConcertServer()->GetLiveSession(SessionId))
+	if (TSharedRef<FAbstractConcertSessionTab>* FoundId = RegisteredSessions.Find(SessionId))
 	{
-		const FConcertSessionInfo& SessionInfo = Session->GetSessionInfo();
-		if (TSharedRef<FConcertServerSessionTab>* FoundId = RegisteredSessions.Find(SessionInfo.SessionId))
-		{
-			return *FoundId;
-		}
+		return *FoundId;
+	}
 	
-		const TSharedRef<FConcertServerSessionTab> SessionTab = MakeShared<FConcertServerSessionTab>(Session.ToSharedRef(), ServerInstance.ToSharedRef(), RootWindow.ToSharedRef());
-		RegisteredSessions.Add(SessionInfo.SessionId, SessionTab);
+	if (TSharedPtr<IConcertServerSession> Session = ServerInstance->GetConcertServer()->GetLiveSession(SessionId))
+	{
+		const TSharedRef<FLiveConcertSessionTab> SessionTab = MakeShared<FLiveConcertSessionTab>(Session.ToSharedRef(), ServerInstance.ToSharedRef(), RootWindow.ToSharedRef());
+		RegisteredSessions.Add(SessionId, SessionTab);
+		return SessionTab;
+	}
+
+	const bool bIsArchivedSession = ServerInstance->GetConcertServer()->GetArchivedSessionInfo(SessionId).IsSet();
+	if (bIsArchivedSession)
+	{
+		const TSharedRef<FArchivedConcertSessionTab> SessionTab = MakeShared<FArchivedConcertSessionTab>(SessionId, ServerInstance.ToSharedRef(), RootWindow.ToSharedRef());
+		RegisteredSessions.Add(SessionId, SessionTab);
 		return SessionTab;
 	}
 	
diff --git a/Engine/Plugins/Developer/Concert/ConcertUI/ConcertServerUI/Source/ConcertServerUI/Private/Widgets/ConcertServerWindowController.h b/Engine/Plugins/Developer/Concert/ConcertUI/ConcertServerUI/Source/ConcertServerUI/Private/Widgets/ConcertServerWindowController.h
--- a/Engine/Plugins/Developer/Concert/ConcertUI/ConcertServerUI/Source/ConcertServerUI/Private/Widgets/ConcertServerWindowController.h
+++ b/Engine/Plugins/Developer/Concert/ConcertUI/ConcertServerUI/Source/ConcertServerUI/Private/Widgets/ConcertServerWindowController.h
@@ -5,7 +5,8 @@
 #include "CoreMinimal.h"
 #include "Framework/Docking/TabManager.h"
 
-class FConcertServerSessionTab;
+class FAbstractConcertSessionTab;
+class FLiveConcertSessionTab;
 class IConcertServerSession;
 class IConcertSyncServer;
 class FConcertServerSessionBrowserController;
@@ -34,7 +35,7 @@ public:
 	FConcertServerWindowController(const FConcertServerWindowInitParams& Params);
 	void CreateWindow();
 
-	/** Opens or draws attention to the tab for the given live session ID */
+	/** Opens or draws attention to the tab for the given live or archived session ID */
 	void OpenSessionTab(const FGuid& SessionId);
 	
 private:
@@ -48,7 +49,7 @@ private:
 
 	/** The main window being managed */
 	TSharedPtr<SWindow> RootWindow;
-	TMap<FGuid, TSharedRef<FConcertServerSessionTab>> RegisteredSessions;
+	TMap<FGuid, TSharedRef<FAbstractConcertSessionTab>> RegisteredSessions;
 	
 	/** Manages the session browser */
 	TSharedPtr<FConcertServerSessionBrowserController> SessionBrowserController;
@@ -56,7 +57,7 @@ private:
 	void InitComponents();
 
 	/** Gets the manager for a session tab if the session ID is valid */
-	TSharedPtr<FConcertServerSessionTab> GetOrRegisterSessionTab(const FGuid& SessionId);
+	TSharedPtr<FAbstractConcertSessionTab> GetOrRegisterSessionTab(const FGuid& SessionId);
 
 	void OnWindowClosed(const TSharedRef<SWindow>& Window);
 	void SaveLayout() const;
diff --git a/Engine/Plugins/Developer/Concert/ConcertUI/ConcertServerUI/Source/ConcertServerUI/Private/Widgets/Session/LiveConcertSessionTab.cpp b/Engine/Plugins/Developer/Concert/ConcertUI/ConcertServerUI/Source/ConcertServerUI/Private/Widgets/Session/LiveConcertSessionTab.cpp
new file mode 100644
--- /dev/null
+++ b/Engine/Plugins/Developer/Concert/ConcertUI/ConcertServerUI/Source/ConcertServerUI/Private/Widgets/Session/LiveConcertSessionTab.cpp
@@ -0,0 +1,48 @@
+﻿// Copyright Epic Games, Inc. All Rights Reserved.
+
+#include "LiveConcertSessionTab.h"
+
+#include "IConcertSession.h"
+#include "Widgets/Session/PackageViewer/ConcertSessionPackageViewerController.h"
+#include "Widgets/Session/SConcertSessionInspector.h"
+#include "Widgets/StatusBar/SConcertStatusBar.h"
+
+#include "Widgets/Docking/SDockTab.h"
+
+FLiveConcertSessionTab::FLiveConcertSessionTab(TSharedRef<IConcertServerSession> InspectedSession, TSharedRef<IConcertSyncServer> SyncServer, TAttribute<TSharedRef<SWindow>> ConstructUnderWindow)
+	: FAbstractConcertSessionTab(SyncServer)
+	, InspectedSession(MoveTemp(InspectedSession))
+	, ConstructUnderWindow(MoveTemp(ConstructUnderWindow))
+	, SessionHistoryController(MakeShared<FServerSessionHistoryController>(InspectedSession, SyncServer))
+	, PackageViewerController(MakeShared<FConcertSessionPackageViewerController>(InspectedSession, SyncServer))
+{}
+
+FGuid FLiveConcertSessionTab::GetSessionID() const
+{
+	return InspectedSession->GetSessionInfo().SessionId;
+}
+
+void FLiveConcertSessionTab::CreateDockContent(const TSharedRef<SDockTab>& InDockTab)
+{
+	const SConcertSessionInspector::FRequiredArgs WidgetArgs
+	{
+		InDockTab,
+		ConstructUnderWindow.Get(),
+		SessionHistoryController->GetSessionHistory(),
+		PackageViewerController->GetPackageViewer()
+	};
+	InDockTab->SetContent(
+		SNew(SConcertSessionInspector, WidgetArgs)
+			.StatusBar()
+			[
+				SNew(SConcertStatusBar, *GetTabId())
+			]
+		);
+}
+
+void FLiveConcertSessionTab::OnOpenTab()
+{
+	SessionHistoryController->ReloadActivities();
+	PackageViewerController->ReloadActivities();
+}
+
diff --git a/Engine/Plugins/Developer/Concert/ConcertUI/ConcertServerUI/Source/ConcertServerUI/Private/Widgets/Session/LiveConcertSessionTab.h b/Engine/Plugins/Developer/Concert/ConcertUI/ConcertServerUI/Source/ConcertServerUI/Private/Widgets/Session/LiveConcertSessionTab.h
new file mode 100644
--- /dev/null
+++ b/Engine/Plugins/Developer/Concert/ConcertUI/ConcertServerUI/Source/ConcertServerUI/Private/Widgets/Session/LiveConcertSessionTab.h
@@ -0,0 +1,46 @@
+﻿// Copyright Epic Games, Inc. All Rights Reserved.
+
+#pragma once
+
+#include "CoreMinimal.h"
+#include "ServerSessionHistoryController.h"
+#include "Widgets/Tabs/AbstractConcertSessionTab.h"
+
+class FConcertSessionPackageViewerController;
+class IConcertSyncServer;
+class FServerSessionHistoryController;
+class IConcertServerSession;
+class SDockTab;
+class SWindow;
+
+/**
+ * Manages the tab that contains the UI for a session.
+ * It has access to the controllers and views needed for displaying a session.
+ */
+class FLiveConcertSessionTab : public FAbstractConcertSessionTab
+{
+public:
+
+	FLiveConcertSessionTab(TSharedRef<IConcertServerSession> InspectedSession, TSharedRef<IConcertSyncServer> SyncServer, TAttribute<TSharedRef<SWindow>> ConstructUnderWindow);
+
+protected:
+
+	//~ Begin FAbstractConcertSessionTab Interface
+	virtual FGuid GetSessionID() const override;
+	virtual void CreateDockContent(const TSharedRef<SDockTab>& InDockTab) override;
+	virtual void OnOpenTab() override;
+	//~ End FAbstractConcertSessionTab Interface
+	
+private:
+	
+	/** The session being inspected */
+	const TSharedRef<IConcertServerSession> InspectedSession;
+
+	/** Used to determine owning window during construction */
+	TAttribute<TSharedRef<SWindow>> ConstructUnderWindow;
+
+	/** Manages the session history widget */
+	const TSharedRef<FServerSessionHistoryController> SessionHistoryController;
+	/** Manages the package viewer widget */
+	const TSharedRef<FConcertSessionPackageViewerController> PackageViewerController;
+};
diff --git a/Engine/Plugins/Developer/Concert/ConcertUI/ConcertServerUI/Source/ConcertServerUI/Private/Widgets/Session/SConcertSessionInspector.cpp b/Engine/Plugins/Developer/Concert/ConcertUI/ConcertServerUI/Source/ConcertServerUI/Private/Widgets/Session/SConcertSessionInspector.cpp
--- a/Engine/Plugins/Developer/Concert/ConcertUI/ConcertServerUI/Source/ConcertServerUI/Private/Widgets/Session/SConcertSessionInspector.cpp
+++ b/Engine/Plugins/Developer/Concert/ConcertUI/ConcertServerUI/Source/ConcertServerUI/Private/Widgets/Session/SConcertSessionInspector.cpp
@@ -5,6 +5,8 @@
 #include "Framework/Docking/TabManager.h"
 #include "Session/History/SSessionHistory.h"
 #include "Widgets/Docking/SDockTab.h"
+#include "Widgets/Layout/SBorder.h"
+#include "Widgets/SBoxPanel.h"
 
 #define LOCTEXT_NAMESPACE "SConcertSessionInspector"
 
@@ -22,13 +24,6 @@ void SConcertSessionInspector::Construct(const FArguments& InArgs, const FRequir
 		[
 			SNew(SVerticalBox)
 
-			// Toolbar
-			+SVerticalBox::Slot()
-			.AutoHeight()
-			[
-				SNullWidget::NullWidget
-			]
-
 			// Content
 			+SVerticalBox::Slot()
 			.FillHeight(1.f)
diff --git a/Engine/Plugins/Developer/Concert/ConcertUI/ConcertServerUI/Source/ConcertServerUI/Private/Widgets/Tabs/AbstractConcertSessionTab.cpp b/Engine/Plugins/Developer/Concert/ConcertUI/ConcertServerUI/Source/ConcertServerUI/Private/Widgets/Tabs/AbstractConcertSessionTab.cpp
new file mode 100644
--- /dev/null
+++ b/Engine/Plugins/Developer/Concert/ConcertUI/ConcertServerUI/Source/ConcertServerUI/Private/Widgets/Tabs/AbstractConcertSessionTab.cpp
@@ -0,0 +1,65 @@
+﻿// Copyright Epic Games, Inc. All Rights Reserved.
+
+#include "AbstractConcertSessionTab.h"
+
+#include "IConcertServer.h"
+#include "IConcertSyncServer.h"
+#include "Widgets/ConcertServerTabs.h"
+
+#include "Framework/Docking/TabManager.h"
+#include "Widgets/Docking/SDockTab.h"
+
+namespace UE::ConcertServerUI::Private
+{
+	static TOptional<FString> GetSessionName(const TSharedRef<IConcertSyncServer>& SyncServer, const FGuid SessionId)
+	{
+		if (const TSharedPtr<IConcertServerSession> LiveSessionInfo = SyncServer->GetConcertServer()->GetLiveSession(SessionId))
+		{
+			return LiveSessionInfo->GetSessionInfo().SessionName;
+		}
+		if (const TOptional<FConcertSessionInfo> ArchivedSessionInfo = SyncServer->GetConcertServer()->GetArchivedSessionInfo(SessionId))
+		{
+			return ArchivedSessionInfo->SessionName;
+		}
+		return {};
+	}
+}
+
+FAbstractConcertSessionTab::FAbstractConcertSessionTab(TSharedRef<IConcertSyncServer> SyncServer)
+	: SyncServer(MoveTemp(SyncServer))
+{}
+
+void FAbstractConcertSessionTab::OpenSessionTab()
+{
+	const TSharedRef<FGlobalTabmanager>& TabManager = FGlobalTabmanager::Get();
+	const FTabId TabId { *GetTabId() };
+
+	EnsureInitDockTab();
+	if (TabManager->FindExistingLiveTab(TabId))
+	{
+		TabManager->DrawAttention(DockTab.ToSharedRef());
+	}
+	else
+	{
+		const FTabManager::FLastMajorOrNomadTab Search(ConcertServerTabs::GetSessionBrowserTabId());
+		FGlobalTabmanager::Get()->InsertNewDocumentTab(*GetTabId(), Search, DockTab.ToSharedRef());
+
+		OnOpenTab();
+	}
+}
+
+void FAbstractConcertSessionTab::EnsureInitDockTab()
+{
+	if (!DockTab.IsValid())
+	{
+		const TOptional<FString> SessionName = UE::ConcertServerUI::Private::GetSessionName(SyncServer, GetSessionID());
+		check(SessionName);
+		
+		const FText Title = FText::FromString(*SessionName);
+		DockTab = SNew(SDockTab)
+			.Label(Title)
+			.TabRole(MajorTab);
+		
+		CreateDockContent(DockTab.ToSharedRef());
+	}
+}
diff --git a/Engine/Plugins/Developer/Concert/ConcertUI/ConcertServerUI/Source/ConcertServerUI/Private/Widgets/Tabs/AbstractConcertSessionTab.h b/Engine/Plugins/Developer/Concert/ConcertUI/ConcertServerUI/Source/ConcertServerUI/Private/Widgets/Tabs/AbstractConcertSessionTab.h
new file mode 100644
--- /dev/null
+++ b/Engine/Plugins/Developer/Concert/ConcertUI/ConcertServerUI/Source/ConcertServerUI/Private/Widgets/Tabs/AbstractConcertSessionTab.h
@@ -0,0 +1,42 @@
+﻿// Copyright Epic Games, Inc. All Rights Reserved.
+
+#pragma once
+
+#include "CoreMinimal.h"
+
+class IConcertSyncServer;
+class SDockTab;
+class SWindow;
+
+/** Shared functionality for a tab that manages a session */
+class FAbstractConcertSessionTab
+{
+public:
+
+    FAbstractConcertSessionTab(TSharedRef<IConcertSyncServer> SyncServer);
+	virtual ~FAbstractConcertSessionTab() = default;
+	
+	/** Opens or draws attention to the given tab */
+	void OpenSessionTab();
+	
+protected:
+
+	virtual FGuid GetSessionID() const = 0;
+	virtual void CreateDockContent(const TSharedRef<SDockTab>& DockTab) = 0;
+	virtual void OnOpenTab() = 0;
+	
+	/** Generates a tab ID for FTabManager::InsertNewDocumentTab */
+	FString GetTabId() const { return GetSessionID().ToString(); }
+	const TSharedRef<IConcertSyncServer>& GetSyncServer() const { return SyncServer; }
+	
+private:
+
+	/** Used to look up session name */
+	const TSharedRef<IConcertSyncServer> SyncServer;
+	
+	/** The tab containing the UI for InspectedSession */
+	TSharedPtr<SDockTab> DockTab;
+
+	/** Inits DockTab if it has not yet been inited */
+	void EnsureInitDockTab();
+};
diff --git a/Engine/Plugins/Developer/Concert/ConcertUI/ConcertServerUI/Source/ConcertServerUI/Private/Widgets/Tabs/Archived/ArchivedConcertSessionTab.cpp b/Engine/Plugins/Developer/Concert/ConcertUI/ConcertServerUI/Source/ConcertServerUI/Private/Widgets/Tabs/Archived/ArchivedConcertSessionTab.cpp
new file mode 100644
--- /dev/null
+++ b/Engine/Plugins/Developer/Concert/ConcertUI/ConcertServerUI/Source/ConcertServerUI/Private/Widgets/Tabs/Archived/ArchivedConcertSessionTab.cpp
@@ -0,0 +1,31 @@
+﻿// Copyright Epic Games, Inc. All Rights Reserved.
+
+#include "Widgets/Tabs/Archived/ArchivedConcertSessionTab.h"
+
+#include "Widgets/Tabs/Archived/SConcertArchivedSessionInspector.h"
+#include "Widgets/StatusBar/SConcertStatusBar.h"
+
+FArchivedConcertSessionTab::FArchivedConcertSessionTab(const FGuid& InspectedSessionID, TSharedRef<IConcertSyncServer> SyncServer, TAttribute<TSharedRef<SWindow>> ConstructUnderWindow)
+	: FAbstractConcertSessionTab(SyncServer)
+	, InspectedSessionID(InspectedSessionID)
+	, ConstructUnderWindow(ConstructUnderWindow)
+{}
+
+FGuid FArchivedConcertSessionTab::GetSessionID() const
+{
+	return InspectedSessionID;
+}
+
+void FArchivedConcertSessionTab::CreateDockContent(const TSharedRef<SDockTab>& InDockTab)
+{
+	InDockTab->SetContent(
+		SNew(SConcertArchivedSessionInspector, SConcertArchivedSessionInspector::FRequiredArgs(InDockTab, ConstructUnderWindow.Get()))
+			.StatusBar()
+			[
+				SNew(SConcertStatusBar, *GetTabId())
+			]
+		);
+}
+
+void FArchivedConcertSessionTab::OnOpenTab()
+{}
diff --git a/Engine/Plugins/Developer/Concert/ConcertUI/ConcertServerUI/Source/ConcertServerUI/Private/Widgets/Tabs/Archived/ArchivedConcertSessionTab.h b/Engine/Plugins/Developer/Concert/ConcertUI/ConcertServerUI/Source/ConcertServerUI/Private/Widgets/Tabs/Archived/ArchivedConcertSessionTab.h
new file mode 100644
--- /dev/null
+++ b/Engine/Plugins/Developer/Concert/ConcertUI/ConcertServerUI/Source/ConcertServerUI/Private/Widgets/Tabs/Archived/ArchivedConcertSessionTab.h
@@ -0,0 +1,35 @@
+﻿// Copyright Epic Games, Inc. All Rights Reserved.
+
+#pragma once
+
+#include "CoreMinimal.h"
+#include "SConcertArchivedSessionInspector.h"
+
+#include "Widgets/Tabs/AbstractConcertSessionTab.h"
+
+/** Manages the tab for an archived session.  */
+class FArchivedConcertSessionTab : public FAbstractConcertSessionTab
+{
+public:
+
+	FArchivedConcertSessionTab(const FGuid& InspectedSessionID, TSharedRef<IConcertSyncServer> SyncServer, TAttribute<TSharedRef<SWindow>> ConstructUnderWindow);
+
+protected:
+
+	//~ Begin FAbstractConcertSessionTab Interface
+	virtual FGuid GetSessionID() const override;
+	virtual void CreateDockContent(const TSharedRef<SDockTab>& InDockTab) override;
+	virtual void OnOpenTab() override;
+	//~ End FAbstractConcertSessionTab Interface
+	
+private:
+
+	/** The inspected session's ID */
+	FGuid InspectedSessionID;
+
+	/** Used later to obtain the window into which to add the tab */
+	TAttribute<TSharedRef<SWindow>> ConstructUnderWindow;
+
+	/** Displays session */
+	TSharedPtr<SConcertArchivedSessionInspector> Inspector;
+};
diff --git a/Engine/Plugins/Developer/Concert/ConcertUI/ConcertServerUI/Source/ConcertServerUI/Private/Widgets/Tabs/Archived/SConcertArchivedSessionInspector.cpp b/Engine/Plugins/Developer/Concert/ConcertUI/ConcertServerUI/Source/ConcertServerUI/Private/Widgets/Tabs/Archived/SConcertArchivedSessionInspector.cpp
new file mode 100644
--- /dev/null
+++ b/Engine/Plugins/Developer/Concert/ConcertUI/ConcertServerUI/Source/ConcertServerUI/Private/Widgets/Tabs/Archived/SConcertArchivedSessionInspector.cpp
@@ -0,0 +1,89 @@
+﻿// Copyright Epic Games, Inc. All Rights Reserved.
+
+#include "SConcertArchivedSessionInspector.h"
+
+#include "Framework/Docking/TabManager.h"
+#include "Session/History/SSessionHistory.h"
+#include "Widgets/Docking/SDockTab.h"
+#include "Widgets/Layout/SBorder.h"
+#include "Widgets/SBoxPanel.h"
+#include "Widgets/SNullWidget.h"
+
+#define LOCTEXT_NAMESPACE "SConcertArchivedSessionInspector"
+
+const FName SConcertArchivedSessionInspector::HistoryTabId("HistoryTabId");
+
+void SConcertArchivedSessionInspector::Construct(const FArguments& InArgs, const FRequiredArgs& InRequiredArgs)
+{
+	ChildSlot
+	[
+		SNew(SBorder)
+		.BorderImage(FAppStyle::Get().GetBrush("ToolPanel.GroupBorder"))
+		.Padding(FMargin(1.0f, 2.0f))
+		[
+			SNew(SVerticalBox)
+
+			// Toolbar
+			+SVerticalBox::Slot()
+			.AutoHeight()
+			[
+				SNullWidget::NullWidget
+			]
+
+			// Content
+			+SVerticalBox::Slot()
+			.FillHeight(1.f)
+			[
+				SNew(SBorder)
+				.BorderImage(FAppStyle::Get().GetBrush("Brushes.Title"))
+				.Padding(FMargin(0.f, 0.f, 0.f, 5.0f)) // Visually separate status bar from tabs
+				[
+					CreateTabs(InRequiredArgs)
+				]
+			]
+
+			// Status bar
+			+SVerticalBox::Slot()
+			.AutoHeight()
+			.HAlign(HAlign_Fill)
+			[
+				InArgs._StatusBar.Widget
+			]
+		]
+	];
+}
+
+TSharedRef<SWidget> SConcertArchivedSessionInspector::CreateTabs(const FRequiredArgs& RequiredArgs)
+{
+	TabManager = FGlobalTabmanager::Get()->NewTabManager(RequiredArgs.ConstructUnderMajorTab);
+	
+	TabManager->RegisterTabSpawner(HistoryTabId, FOnSpawnTab::CreateSP(this, &SConcertArchivedSessionInspector::SpawnActivityHistory))
+		.SetDisplayName(LOCTEXT("ActivityHistoryLabel", "History"));
+
+	const TSharedRef<FTabManager::FLayout> Layout = FTabManager::NewLayout("ConcertArchivedSessionLayout_v0.1")
+		->AddArea
+		(
+			FTabManager::NewPrimaryArea()
+				->SetOrientation(Orient_Horizontal)
+				->Split
+				(
+					FTabManager::NewStack()
+					->SetSizeCoefficient(0.5f)
+					->AddTab(HistoryTabId, ETabState::OpenedTab)
+				)
+		);
+
+	return TabManager->RestoreFrom(Layout, RequiredArgs.ConstructUnderWindow).ToSharedRef();
+}
+
+TSharedRef<SDockTab> SConcertArchivedSessionInspector::SpawnActivityHistory(const FSpawnTabArgs& Args)
+{
+	return SNew(SDockTab)
+		.Label(LOCTEXT("ActivityHistoryLabel", "History"))
+		.TabRole(PanelTab)
+		[
+			SNullWidget::NullWidget
+		]; 
+}
+
+#undef LOCTEXT_NAMESPACE
diff --git a/Engine/Plugins/Developer/Concert/ConcertUI/ConcertServerUI/Source/ConcertServerUI/Private/Widgets/Tabs/Archived/SConcertArchivedSessionInspector.h b/Engine/Plugins/Developer/Concert/ConcertUI/ConcertServerUI/Source/ConcertServerUI/Private/Widgets/Tabs/Archived/SConcertArchivedSessionInspector.h
new file mode 100644
--- /dev/null
+++ b/Engine/Plugins/Developer/Concert/ConcertUI/ConcertServerUI/Source/ConcertServerUI/Private/Widgets/Tabs/Archived/SConcertArchivedSessionInspector.h
@@ -0,0 +1,41 @@
+﻿// Copyright Epic Games, Inc. All Rights Reserved.
+
+#pragma once
+
+#include "CoreMinimal.h"
+#include "Widgets/DeclarativeSyntaxSupport.h"
+#include "Widgets/SCompoundWidget.h"
+
+/** Designed as content for a tab. Displays information about an archived session. */
+class SConcertArchivedSessionInspector : public SCompoundWidget
+{
+public:
+
+	struct FRequiredArgs
+	{
+		TSharedRef<SDockTab> ConstructUnderMajorTab;
+		TSharedRef<SWindow> ConstructUnderWindow;
+
+		FRequiredArgs(TSharedRef<SDockTab> ConstructUnderMajorTab, TSharedRef<SWindow> ConstructUnderWindow)
+			: ConstructUnderMajorTab(MoveTemp(ConstructUnderMajorTab))
+			, ConstructUnderWindow(MoveTemp(ConstructUnderWindow))
+		{}
+	};
+	
+	static const FName HistoryTabId;
+
+	SLATE_BEGIN_ARGS(SConcertArchivedSessionInspector)
+	{}
+		SLATE_NAMED_SLOT(FArguments, StatusBar)
+	SLATE_END_ARGS()
+
+	void Construct(const FArguments& InArgs, const FRequiredArgs& InRequiredArgs);
+
+private:
+	
+	/** Holds the tab manager that manages the front-end's tabs. */
+	TSharedPtr<FTabManager> TabManager;
+	
+	TSharedRef<SWidget> CreateTabs(const FRequiredArgs& RequiredArgs);
+	TSharedRef<SDockTab> SpawnActivityHistory(const FSpawnTabArgs& Args);
+};
diff --git a/Engine/Plugins/Developer/Concert/ConcertUI/ConcertSharedSlate/Source/ConcertSharedSlate/Private/Session/Browser/SConcertSessionBrowser.cpp b/Engine/Plugins/Developer/Concert/ConcertUI/ConcertSharedSlate/Source/ConcertSharedSlate/Private/Session/Browser/SConcertSessionBrowser.cpp
--- a/Engine/Plugins/Developer/Concert/ConcertUI/ConcertSharedSlate/Source/ConcertSharedSlate/Private/Session/Browser/SConcertSessionBrowser.cpp
+++ b/Engine/Plugins/Developer/Concert/ConcertUI/ConcertSharedSlate/Source/ConcertSharedSlate/Private/Session/Browser/SConcertSessionBrowser.cpp
@@ -45,7 +45,12 @@ void SConcertSessionBrowser::Construct(const FArguments& InArgs, TSharedRef<ICon
 
 	ExtendSessionContextMenu = InArgs._ExtendSessionContextMenu;
 	OnSessionClicked = InArgs._OnSessionClicked;
-	OnSessionDoubleClicked = InArgs._OnSessionDoubleClicked;
+	OnLiveSessionDoubleClicked = InArgs._OnLiveSessionDoubleClicked;
+	OnArchivedSessionDoubleClicked = InArgs._OnArchivedSessionDoubleClicked;
+	if (!OnArchivedSessionDoubleClicked.IsBound())
+	{
+		OnArchivedSessionDoubleClicked = FSessionDelegate::CreateSP(this, &SConcertSessionBrowser::InsertRestoreSessionAsEditableRowInternal);
+	}
 	OnRequestedDeleteSession = InArgs._OnRequestedDeleteSession;
 	CanDeleteArchivedSession = InArgs._CanDeleteArchivedSession;
 	CanDeleteActiveSession = InArgs._CanDeleteActiveSession;
@@ -488,7 +493,7 @@ TSharedRef<ITableRow> SConcertSessionBrowser::MakeActiveSessionRowWidget(const T
 
 	// Add an 'Active Session' row. Clicking the row icon joins the session.
 	return SNew(SSessionRow, ActiveItem, OwnerTable)
-		.OnDoubleClickFunc([this](const TSharedPtr<FConcertSessionItem>& Item) { OnSessionDoubleClicked.ExecuteIfBound(Item); })
+		.OnDoubleClickFunc([this](const TSharedPtr<FConcertSessionItem>& Item) { OnLiveSessionDoubleClicked.ExecuteIfBound(Item); })
 		.OnRenameFunc([this](const TSharedPtr<FConcertSessionItem>& Item, const FString& NewName) { RequestRenameSession(Item, NewName); })
 		.IsDefaultSession([this](TSharedPtr<FConcertSessionItem> ItemPin)
 		{
@@ -507,7 +512,7 @@ TSharedRef<ITableRow> SConcertSessionBrowser::MakeArchivedSessionRowWidget(const
 
 	// Add an 'Archived Session' row. Clicking the row icon adds a 'Restore as' row to the table.
 	return SNew(SSessionRow, ArchivedItem, OwnerTable)
-		.OnDoubleClickFunc([this](const TSharedPtr<FConcertSessionItem>& Item) { InsertRestoreSessionAsEditableRowInternal(Item); })
+		.OnDoubleClickFunc([this](const TSharedPtr<FConcertSessionItem>& Item) { OnArchivedSessionDoubleClicked.Execute(Item); })
 		.OnRenameFunc([this](const TSharedPtr<FConcertSessionItem>& Item, const FString& NewName) { RequestRenameSession(Item, NewName); })
 		.IsDefaultSession([this](TSharedPtr<FConcertSessionItem> ItemPin)
 		{
diff --git a/Engine/Plugins/Developer/Concert/ConcertUI/ConcertSharedSlate/Source/ConcertSharedSlate/Public/Session/Browser/SConcertSessionBrowser.h b/Engine/Plugins/Developer/Concert/ConcertUI/ConcertSharedSlate/Source/ConcertSharedSlate/Public/Session/Browser/SConcertSessionBrowser.h
--- a/Engine/Plugins/Developer/Concert/ConcertUI/ConcertSharedSlate/Source/ConcertSharedSlate/Public/Session/Browser/SConcertSessionBrowser.h
+++ b/Engine/Plugins/Developer/Concert/ConcertUI/ConcertSharedSlate/Source/ConcertSharedSlate/Public/Session/Browser/SConcertSessionBrowser.h
@@ -21,11 +21,11 @@ class STableViewBase;
 struct FConcertSessionClientInfo;
 struct FConcertSessionInfo;
 
-DECLARE_DELEGATE_RetVal_OneParam(TSharedRef<SWidget>, FExtendSessionTable, TSharedRef<SWidget> /* TableView */);
+DECLARE_DELEGATE_RetVal_OneParam(TSharedRef<SWidget>, FExtendSessionTable, const TSharedRef<SWidget>& /* TableView */);
 DECLARE_DELEGATE_OneParam(FExtenderDelegate, FExtender&);
-DECLARE_DELEGATE_TwoParams(FExtendSessionContextMenu, TSharedPtr<FConcertSessionItem>, FExtender&)
-DECLARE_DELEGATE_OneParam(FSessionDelegate, TSharedPtr<FConcertSessionItem> /*Session*/);
-DECLARE_DELEGATE_RetVal_OneParam(bool, FCanRemoveSession, TSharedPtr<FConcertSessionItem> /* SessionItem */);
+DECLARE_DELEGATE_TwoParams(FExtendSessionContextMenu, const TSharedPtr<FConcertSessionItem>&, FExtender&)
+DECLARE_DELEGATE_OneParam(FSessionDelegate, const TSharedPtr<FConcertSessionItem>& /*Session*/);
+DECLARE_DELEGATE_RetVal_OneParam(bool, FCanRemoveSession, const TSharedPtr<FConcertSessionItem>& /* SessionItem */);
 
 /**
  * Enables the user to browse/search/filter/sort active and archived sessions, create new session,
@@ -67,10 +67,12 @@ public:
 	/** Custom slot placed to the right of the search bar */
 	SLATE_NAMED_SLOT(FArguments, RightOfSearchBar)
 	
-	/** Called when this session is clicked */
+	/** Called when a live or archived session is clicked */
 	SLATE_EVENT(FSessionDelegate, OnSessionClicked)
-	/** Called when this session is double-clicked */
-	SLATE_EVENT(FSessionDelegate, OnSessionDoubleClicked)
+	/** Called when a live session is double-clicked */
+	SLATE_EVENT(FSessionDelegate, OnLiveSessionDoubleClicked)
+	/** Called when an archived session is double-clicked. If unset, the default behaviour is to rename the session. */
+	SLATE_EVENT(FSessionDelegate, OnArchivedSessionDoubleClicked)
 	/** Called after a user has requested to delete a session */
 	SLATE_EVENT(FSessionDelegate, OnRequestedDeleteSession)
 
@@ -189,8 +191,12 @@ private:
 	TAttribute<FString> DefaultServerUrl;
 	
 	FExtendSessionContextMenu ExtendSessionContextMenu;
+	/** Called when a live or archived session is clicked */
 	FSessionDelegate OnSessionClicked;
-	FSessionDelegate OnSessionDoubleClicked;
+	FSessionDelegate OnLiveSessionDoubleClicked;
+	/** If unset, the default behaviour is to rename the session. */
+	FSessionDelegate OnArchivedSessionDoubleClicked;
+	/** Called after a live or archived session was requested to be deleted (it may or may not have been deleted).*/
 	FSessionDelegate OnRequestedDeleteSession;
 	FCanRemoveSession CanDeleteArchivedSession;
 	FCanRemoveSession CanDeleteActiveSession;
diff --git a/Engine/Plugins/Developer/Concert/MultiUserClient/Source/MultiUserClient/Private/Widgets/SConcertBrowser.cpp b/Engine/Plugins/Developer/Concert/MultiUserClient/Source/MultiUserClient/Private/Widgets/SConcertBrowser.cpp
--- a/Engine/Plugins/Developer/Concert/MultiUserClient/Source/MultiUserClient/Private/Widgets/SConcertBrowser.cpp
+++ b/Engine/Plugins/Developer/Concert/MultiUserClient/Source/MultiUserClient/Private/Widgets/SConcertBrowser.cpp
@@ -1008,9 +1008,9 @@ private:
 	// Layout the 'session|details' split view.
 	TSharedRef<SWidget> MakeBrowserContent(TSharedPtr<FText> InSearchText);
 	void ExtendControlButtons(FExtender& Extender);
-	void ExtendSessionContextMenu(TSharedPtr<FConcertSessionItem> Item, FExtender& Extender);
+	void ExtendSessionContextMenu(const TSharedPtr<FConcertSessionItem>& Item, FExtender& Extender);
 	TSharedRef<SWidget> MakeUserAndSettings();
-	TSharedRef<SWidget> MakeOverlayedTableView(TSharedRef<SWidget> TableView);
+	TSharedRef<SWidget> MakeOverlayedTableView(const TSharedRef<SWidget>& TableView);
 	
 	// Layouts the session detail panel.
 	TSharedRef<SWidget> MakeSessionDetails(TSharedPtr<FConcertSessionItem> Item);
@@ -1035,9 +1035,9 @@ private:
 	void RequestJoinSession(const TSharedPtr<FConcertSessionItem>& ActiveItem);
 
 	// Manipulates the sessions view (the array and the UI).
-	void OnSessionSelectionChanged(TSharedPtr<FConcertSessionItem> SelectedSession);
-	void OnSessionDoubleClicked(TSharedPtr<FConcertSessionItem> SelectedSession);
-	bool ConfirmDeleteSessionWithDialog(TSharedPtr<FConcertSessionItem> SessionItem);
+	void OnSessionSelectionChanged(const TSharedPtr<FConcertSessionItem>& SessionItem);
+	void OnSessionDoubleClicked(const TSharedPtr<FConcertSessionItem>& SessionItem);
+	bool ConfirmDeleteSessionWithDialog(const TSharedPtr<FConcertSessionItem>& SessionItem);
 
 	// Update server/session/clients lists.
 	EActiveTimerReturnType TickDiscovery(double InCurrentTime, float InDeltaTime);
@@ -1173,7 +1173,7 @@ TSharedRef<SWidget> SConcertClientSessionBrowser::MakeBrowserContent(TSharedPtr<
 						MakeUserAndSettings()
 					]
 					.OnSessionClicked_Raw(this, &SConcertClientSessionBrowser::OnSessionSelectionChanged)
-					.OnSessionDoubleClicked_Raw(this, &SConcertClientSessionBrowser::OnSessionDoubleClicked)
+					.OnLiveSessionDoubleClicked_Raw(this, &SConcertClientSessionBrowser::OnSessionDoubleClicked)
 					.OnRequestedDeleteSession_Lambda([this](auto) { UpdateDiscovery(); /* Don't wait up to 1s, kick discovery right now */ })
 					.CanDeleteArchivedSession(this, &SConcertClientSessionBrowser::ConfirmDeleteSessionWithDialog)
 					.CanDeleteActiveSession(this, &SConcertClientSessionBrowser::ConfirmDeleteSessionWithDialog)
@@ -1292,7 +1292,7 @@ void SConcertClientSessionBrowser::ExtendControlButtons(FExtender& Extender)
 		}));
 }
 
-void SConcertClientSessionBrowser::ExtendSessionContextMenu(TSharedPtr<FConcertSessionItem> Item, FExtender& Extender)
+void SConcertClientSessionBrowser::ExtendSessionContextMenu(const TSharedPtr<FConcertSessionItem>& Item, FExtender& Extender)
 {
 	if (Item->Type == FConcertSessionItem::EType::ActiveSession)
 	{
@@ -1360,7 +1360,7 @@ TSharedRef<SWidget> SConcertClientSessionBrowser::MakeUserAndSettings()
 		];
 }
 
-TSharedRef<SWidget> SConcertClientSessionBrowser::MakeOverlayedTableView(TSharedRef<SWidget> SessionTable)
+TSharedRef<SWidget> SConcertClientSessionBrowser::MakeOverlayedTableView(const TSharedRef<SWidget>& SessionTable)
 {
 	return SNew(SOverlay)
 		+SOverlay::Slot()
@@ -1693,7 +1693,7 @@ TSharedRef<ITableRow> SConcertClientSessionBrowser::OnGenerateClientRowWidget(TS
 	];
 }
 
-void SConcertClientSessionBrowser::OnSessionSelectionChanged(TSharedPtr<FConcertSessionItem> SelectedSession)
+void SConcertClientSessionBrowser::OnSessionSelectionChanged(const TSharedPtr<FConcertSessionItem>& SelectedSession)
 {
 	// Clear the list of clients (if any)
 	Clients.Reset();
@@ -1702,7 +1702,7 @@ void SConcertClientSessionBrowser::OnSessionSelectionChanged(TSharedPtr<FConcert
 	SessionDetailsView->SetContent(MakeSessionDetails(SelectedSession));
 }
 
-void SConcertClientSessionBrowser::OnSessionDoubleClicked(TSharedPtr<FConcertSessionItem> SelectedSession)
+void SConcertClientSessionBrowser::OnSessionDoubleClicked(const TSharedPtr<FConcertSessionItem>& SelectedSession)
 {
 	switch (SelectedSession->Type)
 	{
@@ -1717,7 +1717,7 @@ void SConcertClientSessionBrowser::OnSessionDoubleClicked(TSharedPtr<FConcertSes
 	}
 }
 
-bool SConcertClientSessionBrowser::ConfirmDeleteSessionWithDialog(TSharedPtr<FConcertSessionItem> SessionItem)
+bool SConcertClientSessionBrowser::ConfirmDeleteSessionWithDialog(const TSharedPtr<FConcertSessionItem>& SessionItem)
 {
 	const FText SessionNameInText = FText::FromString(SessionItem->SessionName);
 	const FText SeverNameInText = FText::FromString(SessionItem->ServerName);
