Allow users customize the launch flags for concert server.

diff --git a/Engine/Plugins/Developer/Concert/ConcertApp/MultiUserClient/Source/MultiUserClient/Private/MultiUserClientModule.cpp b/Engine/Plugins/Developer/Concert/ConcertApp/MultiUserClient/Source/MultiUserClient/Private/MultiUserClientModule.cpp
--- a/Engine/Plugins/Developer/Concert/ConcertApp/MultiUserClient/Source/MultiUserClient/Private/MultiUserClientModule.cpp
+++ b/Engine/Plugins/Developer/Concert/ConcertApp/MultiUserClient/Source/MultiUserClient/Private/MultiUserClientModule.cpp
@@ -1064,7 +1064,7 @@ private:
 	/**
 	 * Launch a Concert collaboration server on the local machine.
 	 */
-	virtual void LaunchConcertServer() override
+	virtual void LaunchConcertServer(TOptional<FServerLaunchOverrides> LaunchOverrides = {}) override
 	{
 		FAsyncTaskNotificationConfig NotificationConfig;
 		NotificationConfig.bKeepOpenOnFailure = true;
@@ -1105,7 +1105,7 @@ private:
 			return;
 		}
 
-		FString CmdLineArgs = GenerateConcertServerCommandLine();
+		FString CmdLineArgs = GenerateConcertServerCommandLine(LaunchOverrides);
 		FProcHandle ProcHandle = FPlatformProcess::CreateProc(*ServerPath, *CmdLineArgs, true, false, false, nullptr, 0, nullptr, nullptr, nullptr);
 		if (ProcHandle.IsValid())
 		{
@@ -1203,7 +1203,7 @@ private:
 
 			// Launch Server
 			CommandList->MapAction(FConcertUICommands::Get().LaunchServer,
-				FExecuteAction::CreateRaw(this, &FMultiUserClientModule::LaunchConcertServer)
+				FExecuteAction::CreateLambda([this](){ LaunchConcertServer();} )
 			);
 
 			UToolMenu* Menu = UToolMenus::Get()->ExtendMenu("LevelEditor.LevelEditorToolBar.User");
@@ -1295,7 +1295,7 @@ private:
 #endif
 
 private:
-	FString GenerateConcertServerCommandLine() const
+	FString GenerateConcertServerCommandLine(TOptional<FServerLaunchOverrides> LaunchOverrides) const
 	{
 		// @note FH: Those settings are UDP Messaging specific
 
@@ -1344,6 +1344,14 @@ private:
 				}
 			}
 
+			if (LaunchOverrides)
+			{
+				if (!LaunchOverrides->ServerName.IsEmpty())
+				{
+					CmdLine += FString::Printf(TEXT(" -CONCERTSERVER=\"%s\""), *(LaunchOverrides->ServerName));
+				}
+			}
+
 			if (ClientConfig && ClientConfig->bRunWithSlate)
 			{
 				CmdLine += TEXT(" -WITHSLATE");
diff --git a/Engine/Plugins/Developer/Concert/ConcertApp/MultiUserClient/Source/MultiUserClient/Public/IMultiUserClientModule.h b/Engine/Plugins/Developer/Concert/ConcertApp/MultiUserClient/Source/MultiUserClient/Public/IMultiUserClientModule.h
--- a/Engine/Plugins/Developer/Concert/ConcertApp/MultiUserClient/Source/MultiUserClient/Public/IMultiUserClientModule.h
+++ b/Engine/Plugins/Developer/Concert/ConcertApp/MultiUserClient/Source/MultiUserClient/Public/IMultiUserClientModule.h
@@ -8,6 +8,12 @@
 
 class IConcertSyncClient;
 
+struct FServerLaunchOverrides
+{
+	/** The name of the server upon launch */
+	FString ServerName;
+};
+
 /**
  * Interface for the Multi-User module.
  */
@@ -71,7 +77,7 @@ public:
 	 * Launches a server (if none are running) on the local machine. On success, the server is launched. On
 	 * failure, an asynchronous notification (banner) is displayed to the user.
 	 */
-	virtual void LaunchConcertServer() = 0;
+	virtual void LaunchConcertServer(TOptional<FServerLaunchOverrides> LaunchOverrides = {}) = 0;
 
 	/**
 	 * Shuts down all local servers running. Do nothing if no servers are running.
