Simplify usage of Transform using new way of tieing tuples

diff --git a/Engine/Plugins/Developer/Concert/ConcertApp/MultiUserClient/Source/MultiUserClient/Private/Widgets/SConcertBrowser.cpp b/Engine/Plugins/Developer/Concert/ConcertApp/MultiUserClient/Source/MultiUserClient/Private/Widgets/SConcertBrowser.cpp
--- a/Engine/Plugins/Developer/Concert/ConcertApp/MultiUserClient/Source/MultiUserClient/Private/Widgets/SConcertBrowser.cpp
+++ b/Engine/Plugins/Developer/Concert/ConcertApp/MultiUserClient/Source/MultiUserClient/Private/Widgets/SConcertBrowser.cpp
@@ -15,6 +15,8 @@
 #include "SActiveSession.h"
 #include "SConcertSessionRecovery.h"
 
+#include "Algo/TiedTupleOutput.h"
+
 #include "Session/Browser/ConcertBrowserUtils.h"
 #include "Session/Browser/ConcertSessionItem.h"
 #include "Session/Browser/IConcertSessionBrowserController.h"
@@ -1715,14 +1717,10 @@ bool SConcertClientSessionBrowser::ConfirmDeleteSessionWithDialog(const TArray<T
 {
 	TSet<FString> SessionNames;
 	TSet<FString> UniqueServers;
-	
-	Algo::Transform(SessionItems, SessionNames, [](const TSharedPtr<FConcertSessionItem>& Item)
-	{
-		return Item->SessionName;
-	});
-	Algo::Transform(SessionItems, UniqueServers, [](const TSharedPtr<FConcertSessionItem>& Item)
+
+	Algo::Transform(SessionItems, Algo::TieTupleAdd(SessionNames, UniqueServers), [](const TSharedPtr<FConcertSessionItem>& Item)
 	{
-		return Item->ServerName;
+		return MakeTuple(Item->SessionName, Item->ServerName);
 	});
 	
 	const FText ConfirmationMessage = FText::Format(
