Changing EditorFontGlyph usage to icons

diff --git a/Engine/Plugins/Developer/Concert/ConcertApp/MultiUserClient/Source/MultiUserClient/Private/Widgets/SConcertBrowser.cpp b/Engine/Plugins/Developer/Concert/ConcertApp/MultiUserClient/Source/MultiUserClient/Private/Widgets/SConcertBrowser.cpp
--- a/Engine/Plugins/Developer/Concert/ConcertApp/MultiUserClient/Source/MultiUserClient/Private/Widgets/SConcertBrowser.cpp
+++ b/Engine/Plugins/Developer/Concert/ConcertApp/MultiUserClient/Source/MultiUserClient/Private/Widgets/SConcertBrowser.cpp
@@ -1330,10 +1330,9 @@ TSharedRef<SWidget> SConcertClientSessionBrowser::MakeUserAndSettings()
 		.VAlign(VAlign_Center)
 		.AutoWidth()
 		[
-			SNew(STextBlock)
-				.Font(FAppStyle::Get().GetFontStyle("FontAwesome.8"))
-				.ColorAndOpacity_Lambda([this]() { return Controller->GetConcertClient()->GetClientInfo().AvatarColor; })
-				.Text(FEditorFontGlyphs::Square)
+			SNew(SImage)
+			.ColorAndOpacity_Lambda([this]() { return Controller->GetConcertClient()->GetClientInfo().AvatarColor; })
+			.Image(FAppStyle::GetBrush("Icons.FilledCircle"))
 		]
 				
 		// The user "Display Name".
@@ -1357,9 +1356,8 @@ TSharedRef<SWidget> SConcertClientSessionBrowser::MakeUserAndSettings()
 			.VAlign(VAlign_Center)
 			.HAlign(HAlign_Fill)
 			[
-				SNew(STextBlock)
-				.Font(FAppStyle::Get().GetFontStyle("FontAwesome.16"))
-				.Text(FEditorFontGlyphs::Cogs)
+				SNew(SImage)
+				.Image(FAppStyle::GetBrush("Icons.Settings"))
 			]
 		];
 }
@@ -1671,10 +1669,9 @@ TSharedRef<ITableRow> SConcertClientSessionBrowser::OnGenerateClientRowWidget(TS
 		.VAlign(VAlign_Center)
 		.AutoWidth()
 		[
-			SNew(STextBlock)
-			.Font(FAppStyle::Get().GetFontStyle("FontAwesome.8"))
+			SNew(SImage)
 			.ColorAndOpacity(Item->ClientInfo.AvatarColor)
-			.Text(FEditorFontGlyphs::Square)
+			.Image(FAppStyle::GetBrush("Icons.FilledCircle"))
 		]
 
 		// The user "Display Name".
