Create separate program UnrealMultiUserSlateServer

Main implementation by Dominik with build and switchboard support added by me.

diff --git a/Engine/Plugins/Developer/Concert/ConcertApp/MultiUserClient/Source/MultiUserClient/Private/MultiUserClientModule.cpp b/Engine/Plugins/Developer/Concert/ConcertApp/MultiUserClient/Source/MultiUserClient/Private/MultiUserClientModule.cpp
--- a/Engine/Plugins/Developer/Concert/ConcertApp/MultiUserClient/Source/MultiUserClient/Private/MultiUserClientModule.cpp
+++ b/Engine/Plugins/Developer/Concert/ConcertApp/MultiUserClient/Source/MultiUserClient/Private/MultiUserClientModule.cpp
@@ -41,8 +41,9 @@
 #include "Framework/MultiBox/MultiBoxBuilder.h"
 #include "Framework/Notifications/NotificationManager.h"
 
-#include "Styling/AppStyle.h"
+#include "EditorStyleSet.h"
 #include "Delegates/IDelegateInstance.h"
+#include "Interfaces/IEditorStyleModule.h"
 #include "IDetailCustomization.h"
 #include "DetailLayoutBuilder.h"
 #include "DetailCategoryBuilder.h"
@@ -75,6 +76,7 @@
 
 static const FName ConcertBrowserTabName("ConcertBrowser");
 static const TCHAR MultiUserServerAppName[] = TEXT("UnrealMultiUserServer");
+static const TCHAR MultiUserSlateServerAppName[] = TEXT("UnrealMultiUserSlateServer");
 
 #define LOCTEXT_NAMESPACE "MultiUserClient"
 
@@ -452,7 +454,7 @@ public:
 				.Padding(FMargin(0, 1, 0, 1)) // To ensure the text has same size as the default one.
 				[
 					SAssignNew(OutTextBox, SEditableTextBox)
-					.Font(FAppStyle::GetFontStyle(TEXT("PropertyWindow.NormalFont")))
+					.Font(FAppStyle::GetFontStyle((TEXT("PropertyWindow.NormalFont"))))
 					.Text(InitialValue)
 					.SelectAllTextWhenFocused(true)
 					.SelectAllTextOnCommit(true)
@@ -788,6 +790,9 @@ private:
 		// Initialize Style
 		FConcertFrontendStyle::Initialize();
 
+		// Multi-User front end currently relies on EditorStyle being loaded
+		FModuleManager::LoadModuleChecked<IEditorStyleModule>("EditorStyle");
+
 #if WITH_EDITOR
 		RegisterTabSpawner(WorkspaceMenu::GetMenuStructure().GetLevelEditorCategory());
 
@@ -1036,10 +1041,16 @@ private:
 		return DockTab;
 	}
 
-	static FString GetMultiUserServerExePathname()
+	static FString GetConfiguredMultiUserServerExePathname()
 	{
-		// Find concert server location for our build configuration
-		FString MultiUserServerName(MultiUserServerAppName);
+		return GetMultiUserServerExePathname(GetMutableDefault<UConcertClientConfig>()->ServerType);
+	}
+
+	static FString GetMultiUserServerExePathname(const EConcertServerType ServerType)
+	{
+		check(ServerType == EConcertServerType::Console || ServerType == EConcertServerType::Slate);
+		
+		const FString MultiUserServerName(ServerType == EConcertServerType::Console ? MultiUserServerAppName : MultiUserSlateServerAppName);
 		FString ServerPath = FPlatformProcess::GenerateApplicationPath(MultiUserServerName, FApp::GetBuildConfiguration());
 
 		// Validate it exists and fall back to development if it doesn't.
@@ -1053,8 +1064,11 @@ private:
 
 	virtual bool IsConcertServerRunning() override
 	{
-		FString ServerPath = GetMultiUserServerExePathname();
-		return FPlatformProcess::IsApplicationRunning(*FPaths::GetCleanFilename(ServerPath));
+		// Check both because user can change config
+		const FString ConsoleServerPath = GetMultiUserServerExePathname(EConcertServerType::Console);
+		const FString SlateServerPath = GetMultiUserServerExePathname(EConcertServerType::Slate);
+		return FPlatformProcess::IsApplicationRunning(*FPaths::GetCleanFilename(ConsoleServerPath))
+			|| FPlatformProcess::IsApplicationRunning(*FPaths::GetCleanFilename(SlateServerPath));
 	}
 
 	/**
@@ -1076,7 +1090,7 @@ private:
 		}
 
 		// Find concert server location for our build configuration
-		FString ServerPath = GetMultiUserServerExePathname();
+		FString ServerPath = GetConfiguredMultiUserServerExePathname();
 
 		FText LaunchMultiUserErrorTitle = LOCTEXT("LaunchUnrealMultiUserServerErrorTitle", "Failed to Launch the Unreal Multi-User Server");
 		if (!IFileManager::Get().FileExists(*ServerPath))
@@ -1118,7 +1132,7 @@ private:
 
 	void ShutdownConcertServer() override
 	{
-		FString ServerPath = GetMultiUserServerExePathname();
+		FString ServerPath = GetConfiguredMultiUserServerExePathname();
 		ServerPath = FPaths::ConvertRelativePathToFull(ServerPath);
 		FPaths::NormalizeDirectoryName(ServerPath); // Need to have all slashes the same side.
 
@@ -1347,11 +1361,6 @@ private:
 					CmdLine += FString::Printf(TEXT(" -CONCERTSERVER=\"%s\""), *(LaunchOverrides->ServerName));
 				}
 			}
-
-			if (ClientConfig && ClientConfig->bRunWithSlate)
-			{
-				CmdLine += TEXT(" -WITHSLATE");
-			}
 		}
 		return CmdLine;
 	}
diff --git a/Engine/Plugins/Developer/Concert/ConcertApp/MultiUserServer/MultiUserServer.uplugin b/Engine/Plugins/Developer/Concert/ConcertApp/MultiUserServer/MultiUserServer.uplugin
--- a/Engine/Plugins/Developer/Concert/ConcertApp/MultiUserServer/MultiUserServer.uplugin
+++ b/Engine/Plugins/Developer/Concert/ConcertApp/MultiUserServer/MultiUserServer.uplugin
@@ -15,13 +15,13 @@
 	"Hidden": true,
 	"IsBetaVersion": true,
 	"Installed": false,
-	"SupportedPrograms": [ "UnrealMultiUserServer", "UnrealRecoverySvc", "CrashReportClientEditor" ],
+	"SupportedPrograms": [ "UnrealMultiUserSlateServer" ],
 	"Modules": [
         {
           "Name": "MultiUserServer",
           "Type": "UncookedOnly",
           "LoadingPhase": "Default",
-          "ProgramAllowList": [ "UnrealMultiUserServer", "UnrealRecoverySvc", "CrashReportClientEditor" ]
+          "ProgramAllowList": [ "UnrealMultiUserSlateServer" ]
         }
 	],
 	"Plugins": [
diff --git a/Engine/Plugins/Developer/Concert/ConcertMain/ConcertMain.uplugin b/Engine/Plugins/Developer/Concert/ConcertMain/ConcertMain.uplugin
--- a/Engine/Plugins/Developer/Concert/ConcertMain/ConcertMain.uplugin
+++ b/Engine/Plugins/Developer/Concert/ConcertMain/ConcertMain.uplugin
@@ -15,19 +15,19 @@
 	"Hidden": true,
 	"IsBetaVersion": true,
 	"Installed": false,
-	"SupportedPrograms": [ "UnrealMultiUserServer", "UnrealRecoverySvc", "CrashReportClientEditor" ],
+	"SupportedPrograms": [ "UnrealMultiUserServer", "UnrealMultiUserSlateServer", "UnrealRecoverySvc", "CrashReportClientEditor" ],
 	"Modules": [
 		{
 			"Name": "Concert",
 			"Type": "UncookedOnly",
 			"LoadingPhase": "PreDefault",
-			"ProgramAllowList": [ "UnrealMultiUserServer", "UnrealRecoverySvc", "CrashReportClientEditor" ]
+			"ProgramAllowList": [ "UnrealMultiUserServer", "UnrealMultiUserSlateServer", "UnrealRecoverySvc", "CrashReportClientEditor" ]
 		},
 		{
 			"Name": "ConcertTransport",
 			"Type": "UncookedOnly",
 			"LoadingPhase": "PreDefault",
-			"ProgramAllowList": [ "UnrealMultiUserServer", "UnrealRecoverySvc", "CrashReportClientEditor" ]
+			"ProgramAllowList": [ "UnrealMultiUserServer", "UnrealMultiUserSlateServer", "UnrealRecoverySvc", "CrashReportClientEditor" ]
 		}
 	]
 }
diff --git a/Engine/Plugins/Developer/Concert/ConcertMain/Source/Concert/Public/ConcertSettings.h b/Engine/Plugins/Developer/Concert/ConcertMain/Source/Concert/Public/ConcertSettings.h
--- a/Engine/Plugins/Developer/Concert/ConcertMain/Source/Concert/Public/ConcertSettings.h
+++ b/Engine/Plugins/Developer/Concert/ConcertMain/Source/Concert/Public/ConcertSettings.h
@@ -287,6 +287,13 @@ struct FConcertSourceControlSettings
 	EConcertSourceValidationMode ValidationMode;
 };
 
+UENUM()
+enum class EConcertServerType
+{
+	Console,
+	Slate
+};
+
 UCLASS(config=Engine)
 class CONCERT_API UConcertClientConfig : public UObject
 {
@@ -331,9 +338,13 @@ public:
 	UPROPERTY(config, EditAnywhere, Category="Client Settings")
 	bool bRetryAutoConnectOnError = false;
 
-	/** If true, MU servers will be launched with Slate UI. This appends -WITHSLATE to the cmd.  */
+	/**
+	 * Determines which server executable with be launched:
+	 *		Console -> UnrealMultiUserServer.exe
+	 *		Slate	-> UnrealMultiUserSlateServer.exe
+	 */
 	UPROPERTY(config, EditAnywhere, Category="Client Settings")
-	bool bRunWithSlate = false;
+	EConcertServerType ServerType = EConcertServerType::Console;
 
 	/** 
 	 * Default server url (just a name for now) to look for on auto or default connect. 
diff --git a/Engine/Plugins/Developer/Concert/ConcertSync/ConcertSyncCore/ConcertSyncCore.uplugin b/Engine/Plugins/Developer/Concert/ConcertSync/ConcertSyncCore/ConcertSyncCore.uplugin
--- a/Engine/Plugins/Developer/Concert/ConcertSync/ConcertSyncCore/ConcertSyncCore.uplugin
+++ b/Engine/Plugins/Developer/Concert/ConcertSync/ConcertSyncCore/ConcertSyncCore.uplugin
@@ -15,13 +15,13 @@
 	"Hidden": true,
 	"IsBetaVersion": true,
 	"Installed": false,
-	"SupportedPrograms": [ "UnrealMultiUserServer", "UnrealRecoverySvc", "CrashReportClientEditor" ],
+	"SupportedPrograms": [ "UnrealMultiUserServer", "UnrealMultiUserSlateServer", "UnrealRecoverySvc", "CrashReportClientEditor" ],
 	"Modules": [
 		{
 			"Name": "ConcertSyncCore",
 			"Type": "UncookedOnly",
 			"LoadingPhase": "PreDefault",
-			"ProgramAllowList": [ "UnrealMultiUserServer", "UnrealRecoverySvc", "CrashReportClientEditor" ]
+			"ProgramAllowList": [ "UnrealMultiUserServer", "UnrealMultiUserSlateServer", "UnrealRecoverySvc", "CrashReportClientEditor" ]
 		}
 	],
 	"Plugins": [
diff --git a/Engine/Plugins/Developer/Concert/ConcertSync/ConcertSyncServer/ConcertSyncServer.uplugin b/Engine/Plugins/Developer/Concert/ConcertSync/ConcertSyncServer/ConcertSyncServer.uplugin
--- a/Engine/Plugins/Developer/Concert/ConcertSync/ConcertSyncServer/ConcertSyncServer.uplugin
+++ b/Engine/Plugins/Developer/Concert/ConcertSync/ConcertSyncServer/ConcertSyncServer.uplugin
@@ -15,13 +15,13 @@
 	"Hidden": true,
 	"IsBetaVersion": true,
 	"Installed": false,
-	"SupportedPrograms": [ "UnrealMultiUserServer", "UnrealRecoverySvc", "CrashReportClientEditor" ],
+	"SupportedPrograms": [ "UnrealMultiUserServer", "UnrealMultiUserSlateServer", "UnrealRecoverySvc", "CrashReportClientEditor" ],
 	"Modules": [
 		{
 			"Name": "ConcertSyncServer",
 			"Type": "UncookedOnly",
 			"LoadingPhase": "Default",
-			"ProgramAllowList": [ "UnrealMultiUserServer", "UnrealRecoverySvc", "CrashReportClientEditor" ]
+			"ProgramAllowList": [ "UnrealMultiUserServer", "UnrealMultiUserSlateServer", "UnrealRecoverySvc", "CrashReportClientEditor" ]
 		}
 	],
 	"Plugins": [
diff --git a/Engine/Plugins/Developer/Concert/ConcertSync/ConcertSyncServer/Source/ConcertSyncServer/Public/ConcertSyncServerLoop.h b/Engine/Plugins/Developer/Concert/ConcertSync/ConcertSyncServer/Source/ConcertSyncServer/Public/ConcertSyncServerLoop.h
--- a/Engine/Plugins/Developer/Concert/ConcertSync/ConcertSyncServer/Source/ConcertSyncServer/Public/ConcertSyncServerLoop.h
+++ b/Engine/Plugins/Developer/Concert/ConcertSync/ConcertSyncServer/Source/ConcertSyncServer/Public/ConcertSyncServerLoop.h
@@ -8,6 +8,6 @@
 /**
  * Blocking main loop for running a Concert Sync server application.
  */
-int32 ConcertSyncServerLoop(int32 ArgC, TCHAR** ArgV, const FConcertSyncServerLoopInitArgs& InitArgs);
+int32 ConcertSyncServerLoop(const TCHAR* CommandLine, const FConcertSyncServerLoopInitArgs& InitArgs);
 
 #include "ConcertSyncServerLoop.inl"
diff --git a/Engine/Plugins/Developer/Concert/ConcertSync/ConcertSyncServer/Source/ConcertSyncServer/Public/ConcertSyncServerLoop.inl b/Engine/Plugins/Developer/Concert/ConcertSync/ConcertSyncServer/Source/ConcertSyncServer/Public/ConcertSyncServerLoop.inl
--- a/Engine/Plugins/Developer/Concert/ConcertSync/ConcertSyncServer/Source/ConcertSyncServer/Public/ConcertSyncServerLoop.inl
+++ b/Engine/Plugins/Developer/Concert/ConcertSync/ConcertSyncServer/Source/ConcertSyncServer/Public/ConcertSyncServerLoop.inl
@@ -6,6 +6,7 @@
 #include "IConcertSyncServer.h"
 #include "IConcertSyncServerModule.h"
 
+#include "CoreGlobals.h"
 #include "LaunchEngineLoop.h"
 #include "Containers/Ticker.h"
 #include "Misc/ConfigCacheIni.h"
@@ -18,20 +19,10 @@
 
 DEFINE_LOG_CATEGORY_STATIC(LogSyncServer, Log, All);
 
-
-void ConcertLogCommandLineArguments( int ArgC, TCHAR** ArgV )
-{
-	FString CmdLine;
-	for ( int i = 0 ; i < ArgC; i ++ )
-	{
-		CmdLine += " ";
-		CmdLine += ArgV[i];
-	}
-	UE_LOG( LogSyncServer, Display, TEXT( "Launch : %s" ), *CmdLine );
-}
-
-int32 ConcertSyncServerLoop(int32 ArgC, TCHAR** ArgV, const FConcertSyncServerLoopInitArgs& InitArgs)
+int32 ConcertSyncServerLoop(const TCHAR* CommandLine, const FConcertSyncServerLoopInitArgs& InitArgs)
 {
+	FTaskTagScope Scope(ETaskTag::EGameThread);
+	
 	// Validate the init settings
 	checkf(InitArgs.IdealFramerate > 0, TEXT("IdealFramerate must be greater than zero!"));
 	checkf(InitArgs.SessionFlags != EConcertSyncSessionFlags::None, TEXT("SessionFlags cannot be None!"));
@@ -40,13 +31,13 @@ int32 ConcertSyncServerLoop(int32 ArgC, TCHAR** ArgV, const FConcertSyncServerLo
 
 	// start up the main loop, adding some extra command line arguments:
 	//	-Messaging enables MessageBus transports
-	int32 Result = GEngineLoop.PreInit(ArgC, ArgV, TEXT(" -Messaging"));
+	int32 Result = GEngineLoop.PreInit(*FString::Printf(TEXT("%s %s"), CommandLine, TEXT(" -Messaging")));
 	check(GConfig && GConfig->IsReadyForUse());
 
 	if (InitArgs.bShowConsole)
 	{
 		GLogConsole->Show(true);
-		ConcertLogCommandLineArguments( ArgC, ArgV );
+		UE_LOG(LogSyncServer, Display, TEXT( "Launch : %s" ), CommandLine);
 	}
 
 	FModuleManager::Get().StartProcessingNewlyLoadedObjects();
diff --git a/Engine/Plugins/Developer/Concert/ConcertUI/ConcertSharedSlate/ConcertSharedSlate.uplugin b/Engine/Plugins/Developer/Concert/ConcertUI/ConcertSharedSlate/ConcertSharedSlate.uplugin
--- a/Engine/Plugins/Developer/Concert/ConcertUI/ConcertSharedSlate/ConcertSharedSlate.uplugin
+++ b/Engine/Plugins/Developer/Concert/ConcertUI/ConcertSharedSlate/ConcertSharedSlate.uplugin
@@ -15,13 +15,13 @@
 	"Hidden": true,
 	"IsBetaVersion": true,
 	"Installed": false,
-	"SupportedPrograms": [ "UnrealMultiUserServer", "UnrealRecoverySvc", "CrashReportClientEditor" ],
+	"SupportedPrograms": [ "UnrealMultiUserServer", "UnrealMultiUserSlateServer", "UnrealRecoverySvc", "CrashReportClientEditor" ],
 	"Modules": [
         {
           "Name": "ConcertSharedSlate",
           "Type": "UncookedOnly",
           "LoadingPhase": "Default",
-          "ProgramAllowList": [ "UnrealMultiUserServer", "UnrealRecoverySvc", "CrashReportClientEditor" ]
+          "ProgramAllowList": [ "UnrealMultiUserServer", "UnrealMultiUserSlateServer", "UnrealRecoverySvc", "CrashReportClientEditor" ]
         }
 	],
 	"Plugins": [
diff --git a/Engine/Plugins/Messaging/UdpMessaging/UdpMessaging.uplugin b/Engine/Plugins/Messaging/UdpMessaging/UdpMessaging.uplugin
--- a/Engine/Plugins/Messaging/UdpMessaging/UdpMessaging.uplugin
+++ b/Engine/Plugins/Messaging/UdpMessaging/UdpMessaging.uplugin
@@ -18,6 +18,7 @@
 	[
 		"UnrealFrontend",
 		"UnrealMultiUserServer",
+		"UnrealMultiUserSlateServer",
 		"UnrealRecoverySvc",
 		"UnrealLightmass",
 		"CrashReportClientEditor"
@@ -34,6 +35,7 @@
 			"ProgramAllowList": [
 				"UnrealFrontend",
 				"UnrealMultiUserServer",
+				"UnrealMultiUserSlateServer",
 				"UnrealRecoverySvc",
 				"UnrealLightmass",
 				"CrashReportClientEditor"
diff --git a/Engine/Plugins/Runtime/Database/SQLiteCore/SQLiteCore.uplugin b/Engine/Plugins/Runtime/Database/SQLiteCore/SQLiteCore.uplugin
--- a/Engine/Plugins/Runtime/Database/SQLiteCore/SQLiteCore.uplugin
+++ b/Engine/Plugins/Runtime/Database/SQLiteCore/SQLiteCore.uplugin
@@ -14,7 +14,7 @@
 	"CanContainContent": false,
 	"IsBetaVersion": false,
 	"Installed": false,
-	"SupportedPrograms": [ "UnrealMultiUserServer", "UnrealRecoverySvc" ],
+	"SupportedPrograms": [ "UnrealMultiUserServer", "UnrealMultiUserSlateServer", "UnrealRecoverySvc" ],
 	"Modules": [
 		{
 			"Name": "SQLiteCore",
diff --git a/Engine/Plugins/VirtualProduction/Switchboard/Source/Switchboard/switchboard/config.py b/Engine/Plugins/VirtualProduction/Switchboard/Source/Switchboard/switchboard/config.py
--- a/Engine/Plugins/VirtualProduction/Switchboard/Source/Switchboard/switchboard/config.py
+++ b/Engine/Plugins/VirtualProduction/Switchboard/Source/Switchboard/switchboard/config.py
@@ -2182,6 +2182,11 @@ class Config(object):
                 "Multiuser Executable Name",
                 data.get('multiuser_exe', 'UnrealMultiUserServer')
             ),
+            "multiuserslate_exe": StringSetting(
+                "multiuserslate_exe",
+                "Multiuser Slate Executable Name",
+                data.get('multiuserslate_exe', 'UnrealMultiUserSlateServer')
+            ),
             "muserver_auto_launch": BoolSetting(
                 "muserver_auto_launch",
                 "Auto Launch",
@@ -2219,6 +2224,7 @@ class Config(object):
         self.MUSERVER_ENDPOINT = self.mu_settings["muserver_endpoint"]
         self.MUSERVER_MULTICAST_ENDPOINT = self.mu_settings["udpmessaging_multicast_endpoint"]
         self.MULTIUSER_SERVER_EXE = self.mu_settings["multiuser_exe"]
+        self.MULTIUSER_SLATE_SERVER_EXE = self.mu_settings["multiuserslate_exe"]
         self.MUSERVER_AUTO_LAUNCH = self.mu_settings["muserver_auto_launch"]
         self.MUSERVER_SLATE_MODE = self.mu_settings["muserver_slate_mode"]
         self.MUSERVER_CLEAN_HISTORY = self.mu_settings["muserver_clean_history"]
@@ -2236,6 +2242,7 @@ class Config(object):
         data["muserver_server_name"] = self.MUSERVER_SERVER_NAME.get_value()
         data["muserver_endpoint"] = self.MUSERVER_ENDPOINT.get_value()
         data["multiuser_exe"] = self.MULTIUSER_SERVER_EXE.get_value()
+        data["multiuserslate_exe"] = self.MULTIUSER_SLATE_SERVER_EXE.get_value()
         data["muserver_auto_launch"] = self.MUSERVER_AUTO_LAUNCH.get_value()
         data["muserver_slate_mode"] = self.MUSERVER_SLATE_MODE.get_value()
         data["muserver_clean_history"] = self.MUSERVER_CLEAN_HISTORY.get_value()
@@ -2461,13 +2468,23 @@ class Config(object):
         return maps
 
     def multiuser_server_path(self):
-        return self.engine_exe_path(
-            self.ENGINE_DIR.get_value(), self.MULTIUSER_SERVER_EXE.get_value())
+        if self.MUSERVER_SLATE_MODE.get_value():
+            return self.engine_exe_path(
+                self.ENGINE_DIR.get_value(), self.MULTIUSER_SLATE_SERVER_EXE.get_value())
+        else:
+            return self.engine_exe_path(
+                self.ENGINE_DIR.get_value(), self.MULTIUSER_SERVER_EXE.get_value())
 
     def multiuser_server_session_directory_path(self):
+        if self.MUSERVER_SLATE_MODE.get_value():
+            return os.path.join(self.ENGINE_DIR.get_value(), "Programs", "UnrealMultiUserSlateServer", "Intermediate", "MultiUser")
+
         return os.path.join(self.ENGINE_DIR.get_value(), "Programs", "UnrealMultiUserServer", "Intermediate", "MultiUser")
-        
+
     def multiuser_server_log_path(self):
+        if self.MUSERVER_SLATE_MODE.get_value():
+            return os.path.join(self.ENGINE_DIR.get_value(), "Programs", "UnrealMultiUserSlateServer", "Saved", "Logs", "UnrealMultiUserSlateServer.log")
+        # else we get the path to the console server.
         return os.path.join(self.ENGINE_DIR.get_value(), "Programs", "UnrealMultiUserServer", "Saved", "Logs", "UnrealMultiUserServer.log")
 
     def listener_path(self):
diff --git a/Engine/Plugins/VirtualProduction/Switchboard/Source/Switchboard/switchboard/switchboard_application.py b/Engine/Plugins/VirtualProduction/Switchboard/Source/Switchboard/switchboard/switchboard_application.py
--- a/Engine/Plugins/VirtualProduction/Switchboard/Source/Switchboard/switchboard/switchboard_application.py
+++ b/Engine/Plugins/VirtualProduction/Switchboard/Source/Switchboard/switchboard/switchboard_application.py
@@ -150,9 +150,6 @@ class MultiUserApplication:
 
             if CONFIG.MUSERVER_CLEAN_HISTORY.get_value():
                 cmdline += " -ConcertClean"
-                
-            if CONFIG.MUSERVER_SLATE_MODE.get_value():
-                cmdline += " -WithSlate"
 
             if len(args) > 0:
                 cmdline += f' {" ".join(args)}'
diff --git a/Engine/Programs/UnrealGameSync/UnrealGameSync.ini b/Engine/Programs/UnrealGameSync/UnrealGameSync.ini
--- a/Engine/Programs/UnrealGameSync/UnrealGameSync.ini
+++ b/Engine/Programs/UnrealGameSync/UnrealGameSync.ini
@@ -35,6 +35,7 @@ DebugGameHasSeparateExecutable=1
 +Step=(UniqueId="586cc0d3-3914-4df9-acb6-2c02cd9d4fc6", Description="Launch Zen Dashboard", StatusText="Running Zen Dashboard", EstimatedDuration="1", Type="Other", FileName="$(BranchDir)\\Engine\\Binaries\\Win64\\ZenDashboard.exe", WorkingDir="", Arguments="", bUseLogWindow="False", bNormalSync="False", bShowAsTool="False", bScheduledSync="False", Link="Zen", Requires="bb48ca5b-5682-4432-824d-c451336a6523" )
 +Step=(UniqueId="91c2a429-c391-49b4-92a5-4e6be71e0f00", Description="Compile SwitchboardListener", StatusText="Compiling SwitchboardListener...", EstimatedDuration="1", Type="Compile", Target="SwitchboardListener", Platform="Win64", Configuration="Development", Arguments="", bNormalSync="False", bScheduledSync="False", bShowAsTool="False")
 +Step=(UniqueId="5036c75b-8df0-4329-82a1-869dd2d48605", Description="Compile UnrealMultiUserServer", StatusText="Compiling UnrealMultiUserServer...", EstimatedDuration="1", Type="Compile", Target="UnrealMultiUserServer", Platform="Win64", Configuration="Development", Arguments="", bNormalSync="False", bScheduledSync="False", bShowAsTool="False")
++Step=(UniqueId="274b89c3-9dc6-4465-a508-40abc4593cc2", Description="Compile UnrealMultiUserSlateServer", StatusText="Compiling UnrealMultiUserSlateServer...", EstimatedDuration="1", Type="Compile", Target="UnrealMultiUserSlateServer", Platform="Win64", Configuration="Development", Arguments="", bNormalSync="False", bScheduledSync="False", bShowAsTool="False")
 
 [Default]
 EditorTarget=UnrealEditor
diff --git a/Engine/Programs/UnrealMultiUserServer/Config/DefaultEngine.ini b/Engine/Programs/UnrealMultiUserServer/Config/DefaultEngine.ini
--- a/Engine/Programs/UnrealMultiUserServer/Config/DefaultEngine.ini
+++ b/Engine/Programs/UnrealMultiUserServer/Config/DefaultEngine.ini
@@ -1,7 +1,6 @@
 [URL]
-GameNameShort=UCS
+GameNameShort=MUS
 
 [Plugins]
 +ProgramEnabledPlugins="UdpMessaging"
 +ProgramEnabledPlugins="ConcertSyncServer"
-+ProgramEnabledPlugins="MultiUserServer"
diff --git a/Engine/Programs/UnrealMultiUserSlateServer/Config/DefaultEngine.ini b/Engine/Programs/UnrealMultiUserSlateServer/Config/DefaultEngine.ini
new file mode 100644
--- /dev/null
+++ b/Engine/Programs/UnrealMultiUserSlateServer/Config/DefaultEngine.ini
@@ -0,0 +1,7 @@
+[URL]
+GameNameShort=MUSS
+
+[Plugins]
++ProgramEnabledPlugins="UdpMessaging"
++ProgramEnabledPlugins="ConcertSyncServer"
++ProgramEnabledPlugins="MultiUserServer"
diff --git a/Engine/Source/Programs/UnrealMultiUserServer/Private/UnrealMultiUserServerRun.cpp b/Engine/Source/Programs/UnrealMultiUserServer/Private/UnrealMultiUserServerRun.cpp
--- a/Engine/Source/Programs/UnrealMultiUserServer/Private/UnrealMultiUserServerRun.cpp
+++ b/Engine/Source/Programs/UnrealMultiUserServer/Private/UnrealMultiUserServerRun.cpp
@@ -4,38 +4,8 @@
 
 #include "ConcertSettings.h"
 #include "ConcertSyncServerLoop.h"
-#include "IMultiUserServerModule.h"
 
-#include "Misc/Parse.h"
-
-namespace UE::UnrealMultiUserServer
-{
-	static void OptionallySetupSlate(int ArgC, TCHAR* ArgV[], FConcertSyncServerLoopInitArgs& ServerLoopInitArgs)
-	{
-		const FString CommandLine = FCommandLine::BuildFromArgV(nullptr, ArgC, ArgV, nullptr);
-
-		bool bUseSlate = false;
-		bUseSlate |= FParse::Bool(*CommandLine, TEXT("-WITHSLATE"), bUseSlate);
-		FParse::Bool(*CommandLine, TEXT("-WITHSLATE="), bUseSlate);
-		
-		if (bUseSlate)
-		{
-			ServerLoopInitArgs.bShowConsole = false;
-			ServerLoopInitArgs.PreInitServerLoop.AddLambda([&ServerLoopInitArgs]()	
-			{
-				TSharedPtr<IPlugin> Plugin = IPluginManager::Get().FindPlugin(TEXT("MultiUserServer"));
-				if (!Plugin || !Plugin->IsEnabled())
-				{
-					UE_LOG(LogSyncServer, Error, TEXT("The 'ConcertServerUI' plugin is disabled."));
-				}
-				else
-				{
-					IMultiUserServerModule::Get().InitSlateForServer(ServerLoopInitArgs);
-				}
-			});
-		}
-	}
-}
+#include "Misc/CommandLine.h"
 
 int32 RunUnrealMultiUserServer(int ArgC, TCHAR* ArgV[])
 {
@@ -59,6 +29,5 @@ int32 RunUnrealMultiUserServer(int ArgC, TCHAR* ArgV[])
 		return ServerConfig;
 	};
 
-	UE::UnrealMultiUserServer::OptionallySetupSlate(ArgC, ArgV, ServerLoopInitArgs);
-	return ConcertSyncServerLoop(ArgC, ArgV, ServerLoopInitArgs);
+	return ConcertSyncServerLoop(*FCommandLine::BuildFromArgV(nullptr, ArgC, ArgV, nullptr), ServerLoopInitArgs);
 }
diff --git a/Engine/Source/Programs/UnrealMultiUserServer/UnrealMultiUserServer.Build.cs b/Engine/Source/Programs/UnrealMultiUserServer/UnrealMultiUserServer.Build.cs
--- a/Engine/Source/Programs/UnrealMultiUserServer/UnrealMultiUserServer.Build.cs
+++ b/Engine/Source/Programs/UnrealMultiUserServer/UnrealMultiUserServer.Build.cs
@@ -6,6 +6,7 @@ public class UnrealMultiUserServer : ModuleRules
 {
 	public UnrealMultiUserServer(ReadOnlyTargetRules Target) : base(Target)
 	{
+		IncludeOrderVersion = EngineIncludeOrderVersion.Latest;
 		PrivateDependencyModuleNames.AddRange(
 			new string[] {
 				"Core",
@@ -21,8 +22,7 @@ public class UnrealMultiUserServer : ModuleRules
 			new string[] {
 				"Launch",
 				"ConcertSyncCore",
-				"ConcertSyncServer",
-				"MultiUserServer",
+				"ConcertSyncServer"
 			}
 		);
 
@@ -34,8 +34,7 @@ public class UnrealMultiUserServer : ModuleRules
 
 		DynamicallyLoadedModuleNames.AddRange(
 			new string[] {
-				"ConcertSyncServer",
-				"MultiUserServer"
+				"ConcertSyncServer"
 			}
 		);
 	}
diff --git a/Engine/Source/Programs/UnrealMultiUserSlateServer/Private/Linux/LinuxUnrealMultiUserSlateServerMain.cpp b/Engine/Source/Programs/UnrealMultiUserSlateServer/Private/Linux/LinuxUnrealMultiUserSlateServerMain.cpp
new file mode 100644
--- /dev/null
+++ b/Engine/Source/Programs/UnrealMultiUserSlateServer/Private/Linux/LinuxUnrealMultiUserSlateServerMain.cpp
@@ -0,0 +1,11 @@
+// Copyright Epic Games, Inc. All Rights Reserved.
+
+#include "CoreMinimal.h"
+#include "UnrealMultiUserSlateServerRun.h"
+#include "LaunchEngineLoop.h"
+#include "UnixCommonStartup.h"
+
+int main(int argc, char *argv[])
+{
+	return CommonUnixMain(argc, argv, &RunUnrealMultiUserServer, [] { } );
+}
diff --git a/Engine/Source/Programs/UnrealMultiUserSlateServer/Private/Mac/MacUnrealMultiUserSlateServerMain.cpp b/Engine/Source/Programs/UnrealMultiUserSlateServer/Private/Mac/MacUnrealMultiUserSlateServerMain.cpp
new file mode 100644
--- /dev/null
+++ b/Engine/Source/Programs/UnrealMultiUserSlateServer/Private/Mac/MacUnrealMultiUserSlateServerMain.cpp
@@ -0,0 +1,124 @@
+// Copyright Epic Games, Inc. All Rights Reserved.
+
+#include "CoreMinimal.h"
+#include "UnrealMultiUserSlateServerRun.h"
+#include "HAL/ExceptionHandling.h"
+#include "LaunchEngineLoop.h"
+#include "Misc/CommandLine.h"
+#include "Mac/CocoaThread.h"
+
+
+static FString GSavedCommandLine;
+
+
+@interface UnrealAppDelegate : NSObject<NSApplicationDelegate, NSFileManagerDelegate>
+{
+}
+
+@end
+
+
+@implementation UnrealAppDelegate
+
+//handler for the quit apple event used by the Dock menu
+- (void)handleQuitEvent:(NSAppleEventDescriptor*)Event withReplyEvent:(NSAppleEventDescriptor*)ReplyEvent
+{
+    [self requestQuit:self];
+}
+
+- (IBAction)requestQuit:(id)Sender
+{
+	RequestEngineExit(TEXT("requestQuit"));
+}
+
+- (void)runGameThread:(id)Arg
+{
+	FPlatformMisc::SetGracefulTerminationHandler();
+	FPlatformMisc::SetCrashHandler(nullptr);
+	
+#if !UE_BUILD_SHIPPING
+	if (FParse::Param(*GSavedCommandLine,TEXT("crashreports")))
+	{
+		GAlwaysReportCrash = true;
+	}
+#endif
+	
+#if UE_BUILD_DEBUG
+	if (!GAlwaysReportCrash)
+#else
+		if (FPlatformMisc::IsDebuggerPresent() && !GAlwaysReportCrash)
+#endif
+		{
+			RunUnrealMultiUserServer(*GSavedCommandLine);
+		}
+		else
+		{
+			GIsGuarded = 1;
+			RunUnrealMultiUserServer(*GSavedCommandLine);
+			GIsGuarded = 0;
+		}
+	
+	[NSApp terminate: self];
+}
+
+- (NSApplicationTerminateReply)applicationShouldTerminate:(NSApplication *)Sender;
+{
+	if(!IsEngineExitRequested() || ([NSThread gameThread] && [NSThread gameThread] != [NSThread mainThread]))
+	{
+		[self requestQuit:self];
+		return NSTerminateLater;
+	}
+	else
+	{
+		return NSTerminateNow;
+	}
+}
+
+- (void) applicationWillTerminate:(NSNotification*)notification
+{
+	FTaskTagScope::SetTagStaticInit();
+}
+
+- (void)applicationDidFinishLaunching:(NSNotification *)Notification
+{
+	//install the custom quit event handler
+    NSAppleEventManager* appleEventManager = [NSAppleEventManager sharedAppleEventManager];
+    [appleEventManager setEventHandler:self andSelector:@selector(handleQuitEvent:withReplyEvent:) forEventClass:kCoreEventClass andEventID:kAEQuitApplication];
+	
+	RunGameThread(self, @selector(runGameThread:));
+}
+
+@end
+
+
+int main( int argc, char *argv[] )
+{
+	FTaskTagScope::SetTagNone();
+
+	for (int32 Option = 1; Option < argc; Option++)
+	{
+		GSavedCommandLine += TEXT(" ");
+		FString Argument(ANSI_TO_TCHAR(argv[Option]));
+		if (Argument.Contains(TEXT(" ")))
+		{
+			if (Argument.Contains(TEXT("=")))
+			{
+				FString ArgName;
+				FString ArgValue;
+				Argument.Split( TEXT("="), &ArgName, &ArgValue );
+				Argument = FString::Printf( TEXT("%s=\"%s\""), *ArgName, *ArgValue );
+			}
+			else
+			{
+				Argument = FString::Printf(TEXT("\"%s\""), *Argument);
+			}
+		}
+		GSavedCommandLine += Argument;
+	}
+
+	SCOPED_AUTORELEASE_POOL;
+	[NSApplication sharedApplication];
+	[NSApp setDelegate:[UnrealAppDelegate new]];
+	[NSApp run];
+	return 0;
+}
diff --git a/Engine/Source/Programs/UnrealMultiUserSlateServer/Private/UnrealMultiUserSlateServerMain.cpp b/Engine/Source/Programs/UnrealMultiUserSlateServer/Private/UnrealMultiUserSlateServerMain.cpp
new file mode 100644
--- /dev/null
+++ b/Engine/Source/Programs/UnrealMultiUserSlateServer/Private/UnrealMultiUserSlateServerMain.cpp
@@ -0,0 +1,6 @@
+﻿// Copyright Epic Games, Inc. All Rights Reserved.
+
+#include "HAL/Platform.h"
+#include "RequiredProgramMainCPPInclude.h"
+
+IMPLEMENT_APPLICATION(UnrealMultiUserSlateServer, "UnrealMultiUserSlateServer");
diff --git a/Engine/Source/Programs/UnrealMultiUserSlateServer/Private/UnrealMultiUserSlateServerRun.cpp b/Engine/Source/Programs/UnrealMultiUserSlateServer/Private/UnrealMultiUserSlateServerRun.cpp
new file mode 100644
--- /dev/null
+++ b/Engine/Source/Programs/UnrealMultiUserSlateServer/Private/UnrealMultiUserSlateServerRun.cpp
@@ -0,0 +1,53 @@
+// Copyright Epic Games, Inc. All Rights Reserved.
+
+#include "UnrealMultiUserSlateServerRun.h"
+
+#include "ConcertSettings.h"
+#include "ConcertSyncServerLoop.h"
+#include "IMultiUserServerModule.h"
+
+namespace UE::UnrealMultiUserServer
+{
+	static void SetupSlate(FConcertSyncServerLoopInitArgs& ServerLoopInitArgs)
+	{
+		ServerLoopInitArgs.bShowConsole = false;
+		ServerLoopInitArgs.PreInitServerLoop.AddLambda([&ServerLoopInitArgs]()
+		{
+			TSharedPtr<IPlugin> Plugin = IPluginManager::Get().FindPlugin(TEXT("MultiUserServer"));
+			if (!Plugin || !Plugin->IsEnabled())
+			{
+				UE_LOG(LogSyncServer, Error, TEXT("The 'MultiUserServer' plugin is disabled."));
+			}
+			else
+			{
+				IMultiUserServerModule::Get().InitSlateForServer(ServerLoopInitArgs);
+			}
+		});
+	}
+}
+
+int32 RunUnrealMultiUserServer(const TCHAR* CommandLine)
+{
+	FString Role(TEXT("MultiUser"));
+	FConcertSyncServerLoopInitArgs ServerLoopInitArgs;
+	ServerLoopInitArgs.SessionFlags = EConcertSyncSessionFlags::Default_MultiUserSession;
+	ServerLoopInitArgs.ServiceRole = Role;
+	ServerLoopInitArgs.ServiceFriendlyName = TEXT("Multi-User Editing Server");
+
+	ServerLoopInitArgs.GetServerConfigFunc = [Role]() -> const UConcertServerConfig*
+	{
+		UConcertServerConfig* ServerConfig = IConcertSyncServerModule::Get().ParseServerSettings(FCommandLine::Get());
+		if (ServerConfig->WorkingDir.IsEmpty())
+		{
+			ServerConfig->WorkingDir = FPaths::ProjectIntermediateDir() / Role;
+		}
+		if (ServerConfig->ArchiveDir.IsEmpty())
+		{
+			ServerConfig->ArchiveDir = FPaths::ProjectSavedDir() / Role;
+		}
+		return ServerConfig;
+	};
+
+	UE::UnrealMultiUserServer::SetupSlate(ServerLoopInitArgs);
+	return ConcertSyncServerLoop(CommandLine, ServerLoopInitArgs);
+}
diff --git a/Engine/Source/Programs/UnrealMultiUserSlateServer/Private/UnrealMultiUserSlateServerRun.h b/Engine/Source/Programs/UnrealMultiUserSlateServer/Private/UnrealMultiUserSlateServerRun.h
new file mode 100644
--- /dev/null
+++ b/Engine/Source/Programs/UnrealMultiUserSlateServer/Private/UnrealMultiUserSlateServerRun.h
@@ -0,0 +1,7 @@
+// Copyright Epic Games, Inc. All Rights Reserved.
+
+#pragma once
+
+#include "HAL/Platform.h"
+
+int32 RunUnrealMultiUserServer(const TCHAR* CommandLine);
diff --git a/Engine/Source/Programs/UnrealMultiUserSlateServer/Private/Windows/WindowsUnrealMultiUserSlateServerMain.cpp b/Engine/Source/Programs/UnrealMultiUserSlateServer/Private/Windows/WindowsUnrealMultiUserSlateServerMain.cpp
new file mode 100644
--- /dev/null
+++ b/Engine/Source/Programs/UnrealMultiUserSlateServer/Private/Windows/WindowsUnrealMultiUserSlateServerMain.cpp
@@ -0,0 +1,64 @@
+// Copyright Epic Games, Inc. All Rights Reserved.
+
+#include "CoreMinimal.h"
+#include "UnrealMultiUserSlateServerRun.h"
+#include "HAL/ExceptionHandling.h"
+#include "LaunchEngineLoop.h"
+#include "Windows/WindowsHWrapper.h"
+#include "Misc/CommandLine.h"
+#include "Misc/OutputDeviceError.h"
+
+/**
+ * The main application entry point for Windows platforms.
+ *
+ * @param hInInstance Handle to the current instance of the application.
+ * @param hPrevInstance Handle to the previous instance of the application (always NULL).
+ * @param lpCmdLine Command line for the application.
+ * @param nShowCmd Specifies how the window is to be shown.
+ * @return Application's exit value.
+ */
+int32 WINAPI WinMain( HINSTANCE hInInstance, HINSTANCE hPrevInstance, char* lpCmdLine, int32 nShowCmd )
+{
+	hInstance = hInInstance;
+
+	const TCHAR* CmdLine = ::GetCommandLineW();
+	CmdLine = FCommandLine::RemoveExeName(CmdLine);
+
+#if !UE_BUILD_SHIPPING
+	if (FParse::Param(CmdLine, TEXT("crashreports")))
+	{
+		GAlwaysReportCrash = true;
+	}
+#endif
+
+	int32 ErrorLevel = 0;
+
+#if UE_BUILD_DEBUG
+	if (!GAlwaysReportCrash)
+#else
+	if (FPlatformMisc::IsDebuggerPresent() && !GAlwaysReportCrash)
+#endif
+	{
+		ErrorLevel = RunUnrealMultiUserServer(CmdLine);
+	}
+	else
+	{
+#if !PLATFORM_SEH_EXCEPTIONS_DISABLED
+		__try
+#endif
+ 		{
+			GIsGuarded = 1;
+			ErrorLevel = RunUnrealMultiUserServer(CmdLine);
+			GIsGuarded = 0;
+		}
+#if !PLATFORM_SEH_EXCEPTIONS_DISABLED
+		__except (ReportCrash(GetExceptionInformation()))
+		{
+			ErrorLevel = 1;
+			GError->HandleError();
+			FPlatformMisc::RequestExit(true);
+		}
+#endif
+	}
+	return ErrorLevel;
+}
diff --git a/Engine/Source/Programs/UnrealMultiUserSlateServer/UnrealMultiUserSlateServer.Build.cs b/Engine/Source/Programs/UnrealMultiUserSlateServer/UnrealMultiUserSlateServer.Build.cs
new file mode 100644
--- /dev/null
+++ b/Engine/Source/Programs/UnrealMultiUserSlateServer/UnrealMultiUserSlateServer.Build.cs
@@ -0,0 +1,52 @@
+// Copyright Epic Games, Inc. All Rights Reserved.
+
+using UnrealBuildTool;
+
+public class UnrealMultiUserSlateServer : ModuleRules
+{
+	public UnrealMultiUserSlateServer(ReadOnlyTargetRules Target) : base(Target)
+	{
+		//ShortName = "MU-SlateServer";
+		PrivateDependencyModuleNames.AddRange(
+			new string[] {
+				"Core",
+				"CoreUObject",
+				"Concert",
+				"TraceInsights",
+				"ApplicationCore",					// for LaunchEngineLoop.cpp dependency
+				"Projects",							// for LaunchEngineLoop.cpp dependency
+			}
+		);
+
+		PrivateIncludePathModuleNames.AddRange(
+			new string[] {
+				"Launch",
+				"ConcertSyncCore",
+				"ConcertSyncServer",
+				"MultiUserServer",
+			}
+		);
+
+		PrivateIncludePaths.AddRange(
+			new string[] {
+				"Runtime/Launch/Private",           // for LaunchEngineLoop.cpp include
+			}
+		);
+
+		DynamicallyLoadedModuleNames.AddRange(
+			new string[] {
+				"ConcertSyncServer",
+				"MultiUserServer"
+			}
+		);
+		
+		if (Target.IsInPlatformGroup(UnrealPlatformGroup.Linux))
+		{
+			PrivateDependencyModuleNames.AddRange(
+				new string[] {
+					"UnixCommonStartup"
+				}
+			);
+		}
+	}
+}
diff --git a/Engine/Source/Programs/UnrealMultiUserSlateServer/UnrealMultiUserSlateServer.Target.cs b/Engine/Source/Programs/UnrealMultiUserSlateServer/UnrealMultiUserSlateServer.Target.cs
new file mode 100644
--- /dev/null
+++ b/Engine/Source/Programs/UnrealMultiUserSlateServer/UnrealMultiUserSlateServer.Target.cs
@@ -0,0 +1,30 @@
+// Copyright Epic Games, Inc. All Rights Reserved.
+
+using UnrealBuildTool;
+using System.Collections.Generic;
+
+public class UnrealMultiUserSlateServerTarget : TargetRules
+{
+	public UnrealMultiUserSlateServerTarget(TargetInfo Target) : base(Target)
+	{
+		Type = TargetType.Program;
+		LinkType = TargetLinkType.Modular;
+		LaunchModuleName = "UnrealMultiUserSlateServer";
+		AdditionalPlugins.Add("UdpMessaging");
+		AdditionalPlugins.Add("ConcertSyncServer");
+		AdditionalPlugins.Add("MultiUserServer");
+
+		// This app compiles against Core/CoreUObject, but not the Engine or Editor, so compile out Engine and Editor references from Core/CoreUObject
+		bCompileAgainstCoreUObject = true;
+		bCompileAgainstEngine = false;
+		bBuildWithEditorOnlyData = false;
+
+		// Enable Developer plugins (like Concert!)
+		bCompileWithPluginSupport = true;
+		bBuildDeveloperTools = true;
+		bIsBuildingConsoleApplication = false;
+
+		GlobalDefinitions.Add("UE_LOG_CONCERT_DEBUG_VERBOSITY_LEVEL=Log");
+		GlobalDefinitions.Add("UE_TRACE_ENABLED=1");
+	}
+}
diff --git a/Engine/Source/Programs/UnrealRecoverySvc/Private/UnrealRecoverySvc.cpp b/Engine/Source/Programs/UnrealRecoverySvc/Private/UnrealRecoverySvc.cpp
--- a/Engine/Source/Programs/UnrealRecoverySvc/Private/UnrealRecoverySvc.cpp
+++ b/Engine/Source/Programs/UnrealRecoverySvc/Private/UnrealRecoverySvc.cpp
@@ -4,6 +4,7 @@
 #include "ConcertSyncServerLoop.h"
 #include "ConcertLocalFileSharingService.h"
 
+#include "Misc/CommandLine.h"
 #include "RequiredProgramMainCPPInclude.h"
 
 IMPLEMENT_APPLICATION(UnrealRecoverySvc, "UnrealRecoverySvc");
@@ -59,5 +60,5 @@ INT32_MAIN_INT32_ARGC_TCHAR_ARGV()
 		return true;
 	});
 
-	return ConcertSyncServerLoop(ArgC, ArgV, ServerLoopInitArgs);
+	return ConcertSyncServerLoop(*FCommandLine::BuildFromArgV(nullptr, ArgC, ArgV, nullptr), ServerLoopInitArgs);
 }
