Unshelved from pending changelist '20261422':

Sequencer: Follow up to ToolMenus for multi user sequencer buttons

diff --git a/Engine/Plugins/Developer/Concert/ConcertApp/MultiUserClient/Source/MultiUserClient/Private/ConcertWorkspaceUI.cpp b/Engine/Plugins/Developer/Concert/ConcertApp/MultiUserClient/Source/MultiUserClient/Private/ConcertWorkspaceUI.cpp
--- a/Engine/Plugins/Developer/Concert/ConcertApp/MultiUserClient/Source/MultiUserClient/Private/ConcertWorkspaceUI.cpp
+++ b/Engine/Plugins/Developer/Concert/ConcertApp/MultiUserClient/Source/MultiUserClient/Private/ConcertWorkspaceUI.cpp
@@ -33,6 +33,7 @@
 
 #include "Session/History/SSessionHistoryWrapper.h"
 #include "ToolMenus.h"
+#include "Toolkits/AssetEditorToolkit.h"
 #include "Widgets/ClientSessionHistoryController.h"
 #include "Widgets/DeclarativeSyntaxSupport.h"
 #include "Widgets/SConcertSandboxPersistWidget.h"
@@ -1002,17 +1003,35 @@ TSharedRef<SWidget> FConcertWorkspaceUI::OnGenerateAssetViewModifiedByOtherToolt
 
 void FConcertWorkspaceUI::OnPreSequencerInit(TSharedRef<ISequencer> InSequencer, TSharedRef<ISequencerObjectChangeListener> InObjectChangeListener, const FSequencerInitParams& InitParams)
 {
-	if (InitParams.ViewParams.ToolbarExtender && InitParams.bEditWithinLevelEditor)
+	if (InitParams.bEditWithinLevelEditor)
 	{
 		TWeakPtr<FConcertWorkspaceUI> Weak = AsShared();
 		TSharedPtr<FUICommandList> CommandList = MakeShared<FUICommandList>();
-		InitParams.ViewParams.ToolbarExtender->AddToolBarExtension("CurveEditor", EExtensionHook::Before, CommandList, FToolBarExtensionDelegate::CreateLambda([Weak](FToolBarBuilder& ToolbarBuilder)
+
+		ISequencerModule& SequencerModule = FModuleManager::Get().LoadModuleChecked<ISequencerModule>("Sequencer");
+		TSharedPtr<FExtensibilityManager> ExtensionManager = SequencerModule.GetToolBarExtensibilityManager();
+
+		if (ToolbarExtender.IsValid())
 		{
-			if (TSharedPtr<FConcertWorkspaceUI> PinThis = Weak.Pin())
+			ExtensionManager->RemoveExtender(ToolbarExtender);
+			ToolbarExtender.Reset();
+		}
+
+		ToolbarExtender = MakeShareable(new FExtender);
+		ToolbarExtender->AddToolBarExtension(
+			"CurveEditor",
+			EExtensionHook::Before,
+			CommandList,
+			FToolBarExtensionDelegate::CreateLambda([Weak](FToolBarBuilder& ToolbarBuilder)
 			{
-				ToolbarBuilder.AddWidget(SNew(SConcertWorkspaceSequencerToolbarExtension, PinThis));
-			}
-		}));
+				if (TSharedPtr<FConcertWorkspaceUI> PinThis = Weak.Pin())
+				{
+					ToolbarBuilder.AddWidget(SNew(SConcertWorkspaceSequencerToolbarExtension, PinThis));
+				}
+			})
+		);
+
+		ExtensionManager->AddExtender(ToolbarExtender);
 	}
 }
 
diff --git a/Engine/Plugins/Developer/Concert/ConcertApp/MultiUserClient/Source/MultiUserClient/Private/ConcertWorkspaceUI.h b/Engine/Plugins/Developer/Concert/ConcertApp/MultiUserClient/Source/MultiUserClient/Private/ConcertWorkspaceUI.h
--- a/Engine/Plugins/Developer/Concert/ConcertApp/MultiUserClient/Source/MultiUserClient/Private/ConcertWorkspaceUI.h
+++ b/Engine/Plugins/Developer/Concert/ConcertApp/MultiUserClient/Source/MultiUserClient/Private/ConcertWorkspaceUI.h
@@ -154,4 +154,7 @@ private:
 
 	/** Asset history layout for asset history window tabs. */
 	TSharedPtr<FTabManager::FLayout> AssetHistoryLayout;
+
+	/** Toolbar extender for sequencer toolbar */
+	TSharedPtr<FExtender> ToolbarExtender;
 };
