Multi-User:

-Update LaunchConcertServer to return FProcHandle

diff --git a/Engine/Plugins/Developer/Concert/ConcertApp/MultiUserClient/Source/MultiUserClient/Private/MultiUserClientModule.cpp b/Engine/Plugins/Developer/Concert/ConcertApp/MultiUserClient/Source/MultiUserClient/Private/MultiUserClientModule.cpp
--- a/Engine/Plugins/Developer/Concert/ConcertApp/MultiUserClient/Source/MultiUserClient/Private/MultiUserClientModule.cpp
+++ b/Engine/Plugins/Developer/Concert/ConcertApp/MultiUserClient/Source/MultiUserClient/Private/MultiUserClientModule.cpp
@@ -1074,7 +1074,7 @@ private:
 	/**
 	 * Launch a Concert collaboration server on the local machine.
 	 */
-	virtual void LaunchConcertServer(TOptional<FServerLaunchOverrides> LaunchOverrides = {}) override
+	virtual TOptional<FProcHandle> LaunchConcertServer(TOptional<FServerLaunchOverrides> LaunchOverrides = {}) override
 	{
 		FAsyncTaskNotificationConfig NotificationConfig;
 		NotificationConfig.bKeepOpenOnFailure = true;
@@ -1100,7 +1100,7 @@ private:
 				LOCTEXT("LaunchUnrealMultiUserServerError_ExecutableMissing", "Could not find the executable. Have you compiled the Unreal Multi-User Server?"), 
 				false
 				);
-			return;
+			return TOptional<FProcHandle>();
 		}
 
 		// Validate we do not have it running locally 
@@ -1112,14 +1112,17 @@ private:
 				LOCTEXT("LaunchUnrealMultiUserServerError_AlreadyRunning", "An Unreal Multi-User Server instance is already running."), 
 				false
 				);
-			return;
+			return TOptional<FProcHandle>();
 		}
 
 		FString CmdLineArgs = GenerateConcertServerCommandLine(LaunchOverrides);
 		FProcHandle ProcHandle = FPlatformProcess::CreateProc(*ServerPath, *CmdLineArgs, true, false, false, nullptr, 0, nullptr, nullptr, nullptr);
 		if (ProcHandle.IsValid())
 		{
-			Notification.SetComplete(LOCTEXT("LaunchedUnrealMultiUserServer", "Launched Unreal Multi-User Server"), FText(), true);
+			Notification.SetComplete(
+				LOCTEXT("LaunchedUnrealMultiUserServer", "Launched Unreal Multi-User Server"), FText(), true);
+
+			return ProcHandle;
 		}
 		else // Very unlikely in practice, but possible in theory.
 		{
@@ -1128,6 +1131,8 @@ private:
 				LOCTEXT("LaunchUnrealMultiUserServerError_InvalidHandle", "Failed to create the Multi-User Server process."),
 				false);
 		}
+		
+		return TOptional<FProcHandle>();
 	}
 
 	void ShutdownConcertServer() override
diff --git a/Engine/Plugins/Developer/Concert/ConcertApp/MultiUserClient/Source/MultiUserClient/Public/IMultiUserClientModule.h b/Engine/Plugins/Developer/Concert/ConcertApp/MultiUserClient/Source/MultiUserClient/Public/IMultiUserClientModule.h
--- a/Engine/Plugins/Developer/Concert/ConcertApp/MultiUserClient/Source/MultiUserClient/Public/IMultiUserClientModule.h
+++ b/Engine/Plugins/Developer/Concert/ConcertApp/MultiUserClient/Source/MultiUserClient/Public/IMultiUserClientModule.h
@@ -3,11 +3,14 @@
 #pragma once
 
 #include "CoreMinimal.h"
+#include "Misc/Optional.h"
 #include "Modules/ModuleInterface.h"
 #include "Modules/ModuleManager.h"
 
 class IConcertSyncClient;
 
+struct FProcHandle;
+
 struct FServerLaunchOverrides
 {
 	/** The name of the server upon launch */
@@ -74,10 +77,10 @@ public:
 	virtual bool DisconnectSession(bool bAlwaysAskConfirmation = false) = 0;
 
 	/**
-	 * Launches a server (if none are running) on the local machine. On success, the server is launched. On
-	 * failure, an asynchronous notification (banner) is displayed to the user.
+	 * Launches a server (if none are running) on the local machine. On success, the server is launched
+	 * and its FProcHandle is returned. On failure, an asynchronous notification (banner) is displayed to the user.
 	 */
-	virtual void LaunchConcertServer(TOptional<FServerLaunchOverrides> LaunchOverrides = {}) = 0;
+	virtual TOptional<FProcHandle> LaunchConcertServer(TOptional<FServerLaunchOverrides> LaunchOverrides = {}) = 0;
 
 	/**
 	 * Shuts down all local servers running. Do nothing if no servers are running.
