MultiUserClient: reorder filtering predicates in package dirty callback to avoid unnecessary class loading

A recent change in World Partition added actor pinning during the save process to
try to ensure that those actors do not get unloaded after saving. This involves
dirtying the WorldFolders object that lives in the /Engine/Transient package and
ends up triggering another package dirty event that invokes
FConcertWorkspaceUI::OnMarkPackageDirty(). The dirtying of that package
should have been ignored for the purposes of Multi-User, but the call to
ShouldIgnorePackageDirtyEvent() to determine that came after the call
to HasLiveTransactionSupport(). The latter ends up reading class paths off
of the IncludeObjectClassFilters property of UConcertSyncConfig and then
attempts to load those classes using TryLoadClass(). This eventually invoked
a call to StaticFindObjectFast(), but since we are in the middle of a package
save at this point, we ended up triggering a fatal log message.

To address this, we simply perform the check to determine whether a package
dirty event should be ignored before we check whether there is live transaction
support for the package.

diff --git a/Engine/Plugins/Developer/Concert/ConcertApp/MultiUserClient/Source/MultiUserClient/Private/ConcertWorkspaceUI.cpp b/Engine/Plugins/Developer/Concert/ConcertApp/MultiUserClient/Source/MultiUserClient/Private/ConcertWorkspaceUI.cpp
--- a/Engine/Plugins/Developer/Concert/ConcertApp/MultiUserClient/Source/MultiUserClient/Private/ConcertWorkspaceUI.cpp
+++ b/Engine/Plugins/Developer/Concert/ConcertApp/MultiUserClient/Source/MultiUserClient/Private/ConcertWorkspaceUI.cpp
@@ -866,8 +866,8 @@ void FConcertWorkspaceUI::OnMarkPackageDirty(UPackage* InPackage, bool /*bWasDir
 {
 	TSharedPtr<IConcertClientWorkspace> ClientWorkspacePin = ClientWorkspace.Pin();
 	if (ClientWorkspacePin.IsValid()
-		&& !ClientWorkspacePin->HasLiveTransactionSupport(InPackage)
-		&& !ClientWorkspacePin->ShouldIgnorePackageDirtyEvent(InPackage))
+		&& !ClientWorkspacePin->ShouldIgnorePackageDirtyEvent(InPackage)
+		&& !ClientWorkspacePin->HasLiveTransactionSupport(InPackage))
 	{
 		FGuid ResourceLockId = ClientWorkspacePin->GetResourceLockId(InPackage->GetFName());
 
