Fix bad formatting of session list in "are you sure?" delete dialog

diff --git a/Engine/Plugins/Developer/Concert/ConcertApp/MultiUserClient/Source/MultiUserClient/Private/Widgets/SConcertBrowser.cpp b/Engine/Plugins/Developer/Concert/ConcertApp/MultiUserClient/Source/MultiUserClient/Private/Widgets/SConcertBrowser.cpp
--- a/Engine/Plugins/Developer/Concert/ConcertApp/MultiUserClient/Source/MultiUserClient/Private/Widgets/SConcertBrowser.cpp
+++ b/Engine/Plugins/Developer/Concert/ConcertApp/MultiUserClient/Source/MultiUserClient/Private/Widgets/SConcertBrowser.cpp
@@ -1724,7 +1724,7 @@ bool SConcertClientSessionBrowser::ConfirmDeleteSessionWithDialog(const TArray<T
 		LOCTEXT("DeleteSessionConfirmationMessage", "Do you really want to delete {0} {0}|plural(one=session,other=sessions) from {1} {1}|plural(one=server,other=servers)?\n\nThe {0}|plural(one=session,other=sessions) {2} will be deleted from the {1}|plural(one=server,other=servers) {3}."),
 		SessionItems.Num(),
 		UniqueServers.Num(),
-		FText::FromString(FString::Join(SessionNames, TEXT(", "))),
+		FText::FromString(FString::JoinBy(SessionNames, TEXT(", "), [](const FString& Name) { return FString("\n\t- ") + Name; }) + FString("\n")),
 		FText::FromString(FString::Join(UniqueServers, TEXT(", ")))
 		);
 	const FText ConfirmationTitle = LOCTEXT("DeleteSessionConfirmationTitle", "Delete Session Confirmation");
