Do not allow users to create session from the new server UI.

#rnx

diff --git a/Engine/Plugins/Developer/Concert/ConcertApp/MultiUserClient/Source/MultiUserClient/Private/Widgets/SConcertBrowser.cpp b/Engine/Plugins/Developer/Concert/ConcertApp/MultiUserClient/Source/MultiUserClient/Private/Widgets/SConcertBrowser.cpp
--- a/Engine/Plugins/Developer/Concert/ConcertApp/MultiUserClient/Source/MultiUserClient/Private/Widgets/SConcertBrowser.cpp
+++ b/Engine/Plugins/Developer/Concert/ConcertApp/MultiUserClient/Source/MultiUserClient/Private/Widgets/SConcertBrowser.cpp
@@ -120,21 +120,21 @@ public:
 	
 	//~ Begin IConcertSessionBrowserController Interface
 	virtual TArray<FConcertServerInfo> GetServers() const override { return Servers; }
-	
+
 	virtual TArray<FActiveSessionInfo> GetActiveSessions() const override
 	{
 		TArray<FActiveSessionInfo> Result;
 		Algo::Transform(ActiveSessions, Result, [](TSharedPtr<FClientActiveSessionInfo> Info) { return *Info; });
 		return Result;
 	}
-	
+
 	virtual TArray<FArchivedSessionInfo> GetArchivedSessions() const override
 	{
 		TArray<FArchivedSessionInfo> Result;
 		Algo::Transform(ArchivedSessions, Result, [](TSharedPtr<FClientArchivedSessionInfo> Info) { return *Info; });
 		return Result;
 	}
-	
+
 	virtual TOptional<FConcertSessionInfo> GetActiveSessionInfo(const FGuid& AdminEndpoint, const FGuid& SessionId) const override;
 	virtual TOptional<FConcertSessionInfo> GetArchivedSessionInfo(const FGuid& AdminEndpoint, const FGuid& SessionId) const override;
 
@@ -148,6 +148,7 @@ public:
 	virtual void DeleteSessions(const FGuid& ServerAdminEndpointId, const TArray<FGuid>& SessionIds) override;
 	virtual bool CanDeleteActiveSession(const FGuid& ServerAdminEndpointId, const FGuid& SessionId) const override;
 	virtual bool CanDeleteArchivedSession(const FGuid& ServerAdminEndpointId, const FGuid& SessionId) const override;
+	virtual bool CanCreateSessions() const override { return true; }
 	//~ End IConcertSessionBrowserController Interface
 
 private:
diff --git a/Engine/Plugins/Developer/Concert/ConcertApp/MultiUserServer/Source/MultiUserServer/Private/Widgets/Browser/ConcertServerSessionBrowserController.h b/Engine/Plugins/Developer/Concert/ConcertApp/MultiUserServer/Source/MultiUserServer/Private/Widgets/Browser/ConcertServerSessionBrowserController.h
--- a/Engine/Plugins/Developer/Concert/ConcertApp/MultiUserServer/Source/MultiUserServer/Private/Widgets/Browser/ConcertServerSessionBrowserController.h
+++ b/Engine/Plugins/Developer/Concert/ConcertApp/MultiUserServer/Source/MultiUserServer/Private/Widgets/Browser/ConcertServerSessionBrowserController.h
@@ -55,6 +55,7 @@ namespace UE::MultiUserServer
 		virtual bool CanRenameArchivedSession(const FGuid& ServerAdminEndpointId, const FGuid& SessionId) const override { return true; }
 		virtual bool CanDeleteActiveSession(const FGuid& ServerAdminEndpointId, const FGuid& SessionId) const override { return true; }
 		virtual bool CanDeleteArchivedSession(const FGuid& ServerAdminEndpointId, const FGuid& SessionId) const override { return true; }
+		virtual bool CanCreateSessions() const override { return false; }
 		//~ Begin IConcertSessionBrowserController Interface
 
 	private:
diff --git a/Engine/Plugins/Developer/Concert/ConcertUI/ConcertSharedSlate/Source/ConcertSharedSlate/Private/Session/Browser/SConcertSessionBrowser.cpp b/Engine/Plugins/Developer/Concert/ConcertUI/ConcertSharedSlate/Source/ConcertSharedSlate/Private/Session/Browser/SConcertSessionBrowser.cpp
--- a/Engine/Plugins/Developer/Concert/ConcertUI/ConcertSharedSlate/Source/ConcertSharedSlate/Private/Session/Browser/SConcertSessionBrowser.cpp
+++ b/Engine/Plugins/Developer/Concert/ConcertUI/ConcertSharedSlate/Source/ConcertSharedSlate/Private/Session/Browser/SConcertSessionBrowser.cpp
@@ -1052,7 +1052,7 @@ void SConcertSessionBrowser::OnSessionSelectionChanged(TSharedPtr<FConcertSessio
 
 bool SConcertSessionBrowser::IsNewButtonEnabledInternal() const
 {
-	return GetController()->GetServers().Num() > 0;
+	return GetController()->CanCreateSessions() && GetController()->GetServers().Num() > 0;
 }
 
 bool SConcertSessionBrowser::IsRestoreButtonEnabledInternal() const
diff --git a/Engine/Plugins/Developer/Concert/ConcertUI/ConcertSharedSlate/Source/ConcertSharedSlate/Public/Session/Browser/IConcertSessionBrowserController.h b/Engine/Plugins/Developer/Concert/ConcertUI/ConcertSharedSlate/Source/ConcertSharedSlate/Public/Session/Browser/IConcertSessionBrowserController.h
--- a/Engine/Plugins/Developer/Concert/ConcertUI/ConcertSharedSlate/Source/ConcertSharedSlate/Public/Session/Browser/IConcertSessionBrowserController.h
+++ b/Engine/Plugins/Developer/Concert/ConcertUI/ConcertSharedSlate/Source/ConcertSharedSlate/Public/Session/Browser/IConcertSessionBrowserController.h
@@ -64,6 +64,7 @@ public:
 	virtual void DeleteSessions(const FGuid& ServerAdminEndpointId, const TArray<FGuid>& SessionIds) = 0;
 	virtual bool CanDeleteActiveSession(const FGuid& ServerAdminEndpointId, const FGuid& SessionId) const = 0;
 	virtual bool CanDeleteArchivedSession(const FGuid& ServerAdminEndpointId, const FGuid& SessionId) const = 0;
-	
+	virtual bool CanCreateSessions() const = 0;
+
 	virtual ~IConcertSessionBrowserController() = default;
 };
