Add a combo box to set the Multi-user send / receive state.   This only provides the UI mechanism.  The low level implementation is still TBC.

diff --git a/Engine/Plugins/Developer/Concert/ConcertApp/MultiUserClient/Source/MultiUserClient/Private/Widgets/SActiveSession.cpp b/Engine/Plugins/Developer/Concert/ConcertApp/MultiUserClient/Source/MultiUserClient/Private/Widgets/SActiveSession.cpp
--- a/Engine/Plugins/Developer/Concert/ConcertApp/MultiUserClient/Source/MultiUserClient/Private/Widgets/SActiveSession.cpp
+++ b/Engine/Plugins/Developer/Concert/ConcertApp/MultiUserClient/Source/MultiUserClient/Private/Widgets/SActiveSession.cpp
@@ -7,6 +7,7 @@
 #include "ConcertFrontendUtils.h"
 #include "ConcertMessageData.h"
 #include "ClientSessionHistoryController.h"
+#include "ConcertMessages.h"
 #include "IConcertClientPresenceManager.h"
 #include "IConcertClient.h"
 #include "IConcertSyncClient.h"
@@ -15,6 +16,7 @@
 #include "Algo/Transform.h"
 #include "EditorFontGlyphs.h"
 #include "FileHelpers.h"
+#include "ISettingsModule.h"
 #include "Misc/PackageName.h"
 #include "Styling/AppStyle.h"
 #include "Session/History/SSessionHistory.h"
@@ -327,12 +329,34 @@ private:
 	TWeakPtr<FConcertSessionClientInfo> SessionClientInfo;
 };
 
+TArray< TSharedPtr< SActiveSession::FSendReceiveComboItem > > ConstructSendReceiveComboList()
+{
+	TArray< TSharedPtr< SActiveSession::FSendReceiveComboItem > > ComboList;
+	ComboList.Add( MakeShared<SActiveSession::FSendReceiveComboItem>(
+			LOCTEXT("DefaultSendReceiveState", "Default"),
+			LOCTEXT("DefaultSendReceiveStateTooltip", "Full send/receive mode for multi-user events."),
+			EConcertSendReceiveState::Default)
+		);
+
+	ComboList.Add( MakeShared<SActiveSession::FSendReceiveComboItem>(
+			LOCTEXT("SendState", "Send only"),
+			LOCTEXT("SendStateTooltip", "Transactions received from clients will be suspended; however local transactions will be sent."),
+			EConcertSendReceiveState::SendOnly)
+		);
+
+	ComboList.Add( MakeShared<SActiveSession::FSendReceiveComboItem>(
+			LOCTEXT("ReceiveState", "Receive only"),
+			LOCTEXT("ReceiveStateTooltip", "Local changes will be queued for transmission but not sent. Updates from clients will be received."),
+			EConcertSendReceiveState::ReceiveOnly)
+		);
+	return ComboList;
+}
 
 void SActiveSession::Construct(const FArguments& InArgs, TSharedPtr<IConcertSyncClient> InConcertSyncClient)
 {
 	WeakConcertSyncClient = InConcertSyncClient;
 	SessionHistoryController = MakeShared<FClientSessionHistoryController>(InConcertSyncClient.ToSharedRef());
-	
+
 	if (InConcertSyncClient.IsValid())
 	{
 		IConcertClientRef ConcertClient = InConcertSyncClient->GetConcertClient();
@@ -347,6 +371,8 @@ void SActiveSession::Construct(const FArguments& InArgs, TSharedPtr<IConcertSync
 		}
 	}
 
+	SendReceiveComboList = ConstructSendReceiveComboList();
+
 	TSharedRef<SHorizontalBox> StatusBar = 
 		SNew(SHorizontalBox)
 
@@ -376,30 +402,47 @@ void SActiveSession::Construct(const FArguments& InArgs, TSharedPtr<IConcertSync
 				.Font(FAppStyle::Get().GetFontStyle("BoldFont"))
 				.Text(this, &SActiveSession::GetConnectionStatusText)
 			]
+		]
+		+SHorizontalBox::Slot()
+		.AutoWidth()
+		.VAlign(VAlign_Center)
+		.Padding(FMargin(1.0f))
+		[
+			SNew(SBox)
+			.MinDesiredWidth(130.0f)
+			[
+				SAssignNew(SendReceiveComboBox, SComboBox< TSharedPtr<FSendReceiveComboItem> >)
+				.OptionsSource(&SendReceiveComboList)
+				.InitiallySelectedItem(SendReceiveComboList[0])
+				.ContentPadding(FMargin(4.0f,1.0f))
+				.OnGenerateWidget(this, &SActiveSession::GenerateSendReceiveComboItem)
+				.OnSelectionChanged(this, &SActiveSession::HandleSendReceiveChanged)
+				[
+					SNew(STextBlock)
+					.Text(this, &SActiveSession::GetRequestedSendReceiveComboText)
+				]
+			]
+		]
+		// The "Settings" icons.
+		+SHorizontalBox::Slot()
+		.VAlign(VAlign_Fill)
+		.AutoWidth()
+		[
+			SNew(SButton)
+			.ButtonStyle(FAppStyle::Get(), TEXT("SimpleButton"))
+			.OnClicked_Lambda([](){FModuleManager::GetModulePtr<ISettingsModule>("Settings")->ShowViewer("Project", "Plugins", "Concert"); return FReply::Handled(); })
+			.VAlign(VAlign_Center)
+			.HAlign(HAlign_Fill)
+			[
+				SNew(SImage)
+				.Image(FAppStyle::GetBrush("Icons.Settings"))
+			]
 		];
 
 	// Append the buttons to the status bar
 	{
 		TArray<FConcertActionDefinition> ButtonDefs;
 
-		// Resume Session
-		FConcertActionDefinition& ResumeSessionDef = ButtonDefs.AddDefaulted_GetRef();
-		ResumeSessionDef.Type = EConcertActionType::Normal;
-		ResumeSessionDef.IsVisible = MakeAttributeSP(this, &SActiveSession::IsStatusBarResumeSessionVisible);
-		ResumeSessionDef.Text = FEditorFontGlyphs::Play_Circle;
-		ResumeSessionDef.ToolTipText = LOCTEXT("ResumeCurrentSessionToolTip", "Resume receiving updates from the current session");
-		ResumeSessionDef.OnExecute.BindLambda([this]() { OnClickResumeSession(); });
-		ResumeSessionDef.IconStyle = TEXT("Concert.ResumeSession");
-
-		// Suspend Session
-		FConcertActionDefinition& SuspendSessionDef = ButtonDefs.AddDefaulted_GetRef();
-		SuspendSessionDef.Type = EConcertActionType::Normal;
-		SuspendSessionDef.IsVisible = MakeAttributeSP(this, &SActiveSession::IsStatusBarSuspendSessionVisible);
-		SuspendSessionDef.Text = FEditorFontGlyphs::Pause_Circle;
-		SuspendSessionDef.ToolTipText = LOCTEXT("SuspendCurrentSessionToolTip", "Suspend receiving updates from the current session");
-		SuspendSessionDef.OnExecute.BindLambda([this]() { OnClickSuspendSession(); });
-		SuspendSessionDef.IconStyle = TEXT("Concert.PauseSession");
-
 		// Leave Session
 		FConcertActionDefinition& LeaveSessionDef = ButtonDefs.AddDefaulted_GetRef();
 		LeaveSessionDef.Type = EConcertActionType::Normal;
@@ -519,6 +562,34 @@ void SActiveSession::Construct(const FArguments& InArgs, TSharedPtr<IConcertSync
 	UpdateSessionClientListView();
 }
 
+FText SActiveSession::GetRequestedSendReceiveComboText() const
+{
+	if (SendReceiveComboBox->GetSelectedItem().IsValid())
+	{
+		return SendReceiveComboBox->GetSelectedItem()->Name;
+	}
+	else
+	{
+		return LOCTEXT("ActiveSessionDefaultSendReceive", "Default");
+	}
+}
+
+TSharedRef<SWidget> SActiveSession::GenerateSendReceiveComboItem(TSharedPtr<SActiveSession::FSendReceiveComboItem> InItem)
+{
+	return SNew(STextBlock)
+		.Text(InItem->Name)
+		.ToolTipText(InItem->ToolTip);
+}
+
+void SActiveSession::HandleSendReceiveChanged(TSharedPtr<SActiveSession::FSendReceiveComboItem> Item, ESelectInfo::Type SelectInfo)
+{
+	check(Item.IsValid());
+	if (TSharedPtr<IConcertClientSession> Session = WeakSessionPtr.Pin())
+	{
+		Session->SetSendReceiveState(Item->State);
+	}
+}
+
 TSharedRef<ITableRow> SActiveSession::HandleGenerateRow(TSharedPtr<FConcertSessionClientInfo> InClientInfo, const TSharedRef<STableViewBase>& OwnerTable) const
 {
 	// Generate a row for the client corresponding to InClientInfo.
@@ -697,23 +768,24 @@ void SActiveSession::UpdateSessionClientListView(const FConcertSessionClientInfo
 const FButtonStyle& SActiveSession::GetConnectionIconStyle() const
 {
 	EConcertActionType ButtonStyle = EConcertActionType::Danger;
-	
+
 	TSharedPtr<IConcertClientSession> ClientSession = WeakSessionPtr.Pin();
 	if (ClientSession.IsValid())
 	{
 		if (ClientSession->GetConnectionStatus() == EConcertConnectionStatus::Connected)
 		{
-			if (ClientSession->IsSuspended())
+			const bool bIsDefault = ClientSession->GetSendReceiveState() == EConcertSendReceiveState::Default;
+			if (bIsDefault)
 			{
-				ButtonStyle = EConcertActionType::Warning;
+				ButtonStyle = EConcertActionType::Success;
 			}
 			else
 			{
-				ButtonStyle = EConcertActionType::Success;
+				ButtonStyle = EConcertActionType::Warning;
 			}
 		}
 	}
-	
+
 	return FAppStyle::Get().GetWidgetStyle<FButtonStyle>(ConcertClientFrontendUtils::ButtonStyleNames[(int32)ButtonStyle]);
 }
 
@@ -739,39 +811,23 @@ FText SActiveSession::GetConnectionStatusText() const
 	{
 		const FText SessionDisplayName = FText::FromString(ClientSessionPtr->GetSessionInfo().SessionName);
 
-		if (ClientSessionPtr->IsSuspended())
-		{
-			StatusText = FText::Format(LOCTEXT("StatusSuspendedFmt", "Suspended: {0}"), SessionDisplayName);
-		}
-		else
+		switch (ClientSessionPtr->GetSendReceiveState())
 		{
+		case EConcertSendReceiveState::SendOnly:
+			StatusText = FText::Format(LOCTEXT("StatusSendOnlyFmt", "Send Only: {0}"), SessionDisplayName);
+			break;
+		case EConcertSendReceiveState::ReceiveOnly:
+			StatusText = FText::Format(LOCTEXT("StatusReceiveOnlyFmt", "Receive Only: {0}"), SessionDisplayName);
+			break;
+		case EConcertSendReceiveState::Default:
 			StatusText = FText::Format(LOCTEXT("StatusConnectedFmt", "Connected: {0}"), SessionDisplayName);
-		}
+			break;
+		};
 	}
 
 	return StatusText;
 }
 
-bool SActiveSession::IsStatusBarSuspendSessionVisible() const
-{
-	if (TSharedPtr<IConcertClientSession> ClientSession = WeakSessionPtr.Pin())
-	{
-		return ClientSession->GetConnectionStatus() == EConcertConnectionStatus::Connected && !ClientSession->IsSuspended();
-	}
-	
-	return false;
-}
-
-bool SActiveSession::IsStatusBarResumeSessionVisible() const
-{
-	if (TSharedPtr<IConcertClientSession> ClientSession = WeakSessionPtr.Pin())
-	{
-		return ClientSession->GetConnectionStatus() == EConcertConnectionStatus::Connected && ClientSession->IsSuspended();
-	}
-
-	return false;
-}
-
 bool SActiveSession::IsStatusBarLeaveSessionVisible() const
 {
 	if (TSharedPtr<IConcertClientSession> ClientSession = WeakSessionPtr.Pin())
@@ -782,28 +838,6 @@ bool SActiveSession::IsStatusBarLeaveSessionVisible() const
 	return false;
 }
 
-FReply SActiveSession::OnClickSuspendSession()
-{
-	TSharedPtr<IConcertClientSession> ClientSession = WeakSessionPtr.Pin();
-	if (ClientSession.IsValid())
-	{
-		ClientSession->Suspend();
-	}
-
-	return FReply::Handled();
-}
-
-FReply SActiveSession::OnClickResumeSession()
-{
-	TSharedPtr<IConcertClientSession> ClientSession = WeakSessionPtr.Pin();
-	if (ClientSession.IsValid())
-	{
-		ClientSession->Resume();
-	}
-
-	return FReply::Handled();
-}
-
 FReply SActiveSession::OnClickLeaveSession()
 {
 	IMultiUserClientModule::Get().DisconnectSession();
diff --git a/Engine/Plugins/Developer/Concert/ConcertApp/MultiUserClient/Source/MultiUserClient/Private/Widgets/SActiveSession.h b/Engine/Plugins/Developer/Concert/ConcertApp/MultiUserClient/Source/MultiUserClient/Private/Widgets/SActiveSession.h
--- a/Engine/Plugins/Developer/Concert/ConcertApp/MultiUserClient/Source/MultiUserClient/Private/Widgets/SActiveSession.h
+++ b/Engine/Plugins/Developer/Concert/ConcertApp/MultiUserClient/Source/MultiUserClient/Private/Widgets/SActiveSession.h
@@ -7,6 +7,7 @@
 #include "Widgets/DeclarativeSyntaxSupport.h"
 #include "Widgets/SCompoundWidget.h"
 #include "Widgets/Views/SListView.h"
+#include "Widgets/Input/SComboBox.h"
 #include "ConcertMessages.h"
 
 class IConcertClientSession;
@@ -23,6 +24,16 @@ class SExpandableArea;
 class SActiveSession : public SCompoundWidget
 {
 public:
+	/** Struct to store the current send / receive state. */
+	struct FSendReceiveComboItem
+	{
+		FSendReceiveComboItem(FText InName, FText InToolTip, EConcertSendReceiveState InState) :
+			Name(MoveTemp(InName)), ToolTip(MoveTemp(InToolTip)), State(InState) {};
+
+		FText Name;
+		FText ToolTip;
+		EConcertSendReceiveState State;
+	};
 
 	SLATE_BEGIN_ARGS(SActiveSession) { }
 	SLATE_END_ARGS();
@@ -59,19 +70,15 @@ private:
 
 	/** Find a client with its endpoint id */
 	TSharedPtr<FConcertSessionClientInfo> FindAvailableClient(const FGuid& InClientEndpointId) const;
-	
+
 	/** Handling for the status icon and text */
 	const FButtonStyle& GetConnectionIconStyle() const;
 	FSlateColor GetConnectionIconColor() const;
 	FSlateFontInfo GetConnectionIconFontInfo() const;
 	FText GetConnectionStatusText() const;
 
-	/** Handling for the suspend, resume and leave session buttons */
-	bool IsStatusBarSuspendSessionVisible() const;
-	bool IsStatusBarResumeSessionVisible() const;
+	/** Handling for leave session button */
 	bool IsStatusBarLeaveSessionVisible() const;
-	FReply OnClickSuspendSession();
-	FReply OnClickResumeSession();
 	FReply OnClickLeaveSession();
 
 	/** Handles how much space the 'Clients' area uses with respect to its expansion state. */
@@ -83,7 +90,16 @@ private:
 	void OnHistoryAreaExpansionChanged(bool bExpanded) { bHistoryAreaExpanded = bExpanded; }
 
 private:
-	
+
+	/** Get the text object for the send/receive combo box. */
+	FText GetRequestedSendReceiveComboText() const;
+
+	/** Generate the SWidget for the given item. */
+	TSharedRef<SWidget> GenerateSendReceiveComboItem(TSharedPtr<FSendReceiveComboItem> InItem);
+
+	/** Handle a send receiver state change from the combo box. */
+	void HandleSendReceiveChanged(TSharedPtr<FSendReceiveComboItem> Item, ESelectInfo::Type SelectInfo);
+
 	/** Pointer on the client sync. */
 	TWeakPtr<IConcertSyncClient> WeakConcertSyncClient;
 
@@ -108,6 +124,12 @@ private:
 	/** The 'History' expandable area. */
 	TSharedPtr<SExpandableArea> HistoryArea;
 
+	/** Holds a pointer to the preset combo box widget. */
+	TSharedPtr< SComboBox< TSharedPtr<FSendReceiveComboItem> > > SendReceiveComboBox;
+
+	/** Available states for the SendReceiveComboBox */
+	TArray< TSharedPtr< FSendReceiveComboItem > > SendReceiveComboList;
+
 	/** Keeps the status of 'Clients' area expansion. */
 	bool bClientAreaExpanded = true;
 
diff --git a/Engine/Plugins/Developer/Concert/ConcertMain/Source/Concert/Private/ConcertClient.cpp b/Engine/Plugins/Developer/Concert/ConcertMain/Source/Concert/Private/ConcertClient.cpp
--- a/Engine/Plugins/Developer/Concert/ConcertMain/Source/Concert/Private/ConcertClient.cpp
+++ b/Engine/Plugins/Developer/Concert/ConcertMain/Source/Concert/Private/ConcertClient.cpp
@@ -3,6 +3,7 @@
 #include "ConcertClient.h"
 
 #include "ConcertClientSession.h"
+#include "ConcertMessages.h"
 #include "ConcertUtil.h"
 #include "ConcertLogger.h"
 #include "ConcertLogGlobal.h"
@@ -1077,27 +1078,23 @@ void FConcertClient::DisconnectSession()
 	PendingConnection.Reset();
 }
 
-void FConcertClient::ResumeSession()
+EConcertSendReceiveState FConcertClient::GetSendReceiveState() const
 {
 	if (ClientSession.IsValid())
 	{
-		ClientSession->Resume();
+		ClientSession->GetSendReceiveState();
 	}
+	return EConcertSendReceiveState::Default;
 }
 
-void FConcertClient::SuspendSession()
+void FConcertClient::SetSendReceiveState(EConcertSendReceiveState InSendReceiveState)
 {
 	if (ClientSession.IsValid())
 	{
-		ClientSession->Suspend();
+		ClientSession->SetSendReceiveState(InSendReceiveState);
 	}
 }
 
-bool FConcertClient::IsSessionSuspended() const
-{
-	return ClientSession.IsValid() && ClientSession->IsSuspended();
-}
-
 bool FConcertClient::IsOwnerOf(const FConcertSessionInfo& InSessionInfo) const
 {
 	return ClientInfo.UserName == InSessionInfo.OwnerUserName && ClientInfo.DeviceName == InSessionInfo.OwnerDeviceName;
diff --git a/Engine/Plugins/Developer/Concert/ConcertMain/Source/Concert/Private/ConcertClient.h b/Engine/Plugins/Developer/Concert/ConcertMain/Source/Concert/Private/ConcertClient.h
--- a/Engine/Plugins/Developer/Concert/ConcertMain/Source/Concert/Private/ConcertClient.h
+++ b/Engine/Plugins/Developer/Concert/ConcertMain/Source/Concert/Private/ConcertClient.h
@@ -86,9 +86,10 @@ public:
 	virtual TFuture<EConcertResponseCode> DeleteSession(const FGuid& ServerAdminEndpointId, const FGuid& SessionId) override;
 	virtual TFuture<FConcertAdmin_BatchDeleteSessionResponse> BatchDeleteSessions(const FGuid& ServerAdminEndpointId, const FConcertBatchDeleteSessionsArgs& BatchDeletionArgs) override;
 	virtual void DisconnectSession() override;
-	virtual void ResumeSession() override;
-	virtual void SuspendSession() override;
-	virtual bool IsSessionSuspended() const override;
+
+	virtual EConcertSendReceiveState GetSendReceiveState() const override;
+	virtual void SetSendReceiveState(EConcertSendReceiveState InSendReceiveState) override;
+
 	virtual bool IsOwnerOf(const FConcertSessionInfo& InSessionInfo) const override;
 	virtual TSharedPtr<IConcertClientSession> GetCurrentSession() const override;
 
diff --git a/Engine/Plugins/Developer/Concert/ConcertMain/Source/Concert/Private/ConcertClientSession.cpp b/Engine/Plugins/Developer/Concert/ConcertMain/Source/Concert/Private/ConcertClientSession.cpp
--- a/Engine/Plugins/Developer/Concert/ConcertMain/Source/Concert/Private/ConcertClientSession.cpp
+++ b/Engine/Plugins/Developer/Concert/ConcertMain/Source/Concert/Private/ConcertClientSession.cpp
@@ -2,6 +2,7 @@
 
 #include "ConcertClientSession.h"
 #include "ConcertLogGlobal.h"
+#include "ConcertMessages.h"
 #include "IConcertEndpoint.h"
 #include "Scratchpad/ConcertScratchpad.h"
 
@@ -17,7 +18,6 @@ FConcertClientSession::FConcertClientSession(const FConcertSessionInfo& InSessio
 	, ClientInfo(InClientInfo)
 	, ConnectionStatus(EConcertConnectionStatus::Disconnected)
 	, ClientSessionEndpoint(MoveTemp(InClientSessionEndpoint))
-	, SuspendedCount(0)
 	, LastConnectionTick(0)
 	, SessionTickFrequency(0, 0, InSettings.SessionTickFrequencySeconds)
 	, SessionDirectory(InSessionDirectory)
@@ -120,24 +120,39 @@ void FConcertClientSession::Disconnect()
 	}
 }
 
-void FConcertClientSession::Resume()
+EConcertSendReceiveState FConcertClientSession::GetSendReceiveState() const
 {
-	check(IsSuspended());
-	--SuspendedCount;
-
-	UE_LOG(LogConcert, Display, TEXT("Resumed Concert session '%s' (Id: %s, Owner: %s)."), *SessionInfo.SessionName, *SessionInfo.SessionId.ToString(), *SessionInfo.OwnerUserName);
+	return SendReceiveState;
 }
 
-void FConcertClientSession::Suspend()
+void LogSendReceiveStateChange(
+	FConcertSessionInfo SessionInfo,
+	EConcertSendReceiveState InState)
 {
-	++SuspendedCount;
-
-	UE_LOG(LogConcert, Display, TEXT("Suspended Concert session '%s' (Id: %s, Owner: %s)."), *SessionInfo.SessionName, *SessionInfo.SessionId.ToString(), *SessionInfo.OwnerUserName);
+	FString State = [InState]()
+	{
+		if (InState == EConcertSendReceiveState::Default)
+		{
+			return FString("Resumed");
+		}
+		else if (InState == EConcertSendReceiveState::SendOnly)
+		{
+			return FString("SendOnly");
+		}
+		return FString("ReceiveOnly");
+	}();
+	UE_LOG(LogConcert, Display,
+		   TEXT("%s Concert session '%s' (Id: %s, Owner: %s)."),
+		   *State,
+		   *SessionInfo.SessionName,
+		   *SessionInfo.SessionId.ToString(),
+		   *SessionInfo.OwnerUserName);
 }
 
-bool FConcertClientSession::IsSuspended() const
+void FConcertClientSession::SetSendReceiveState(EConcertSendReceiveState InSendReceiveState)
 {
-	return ConnectionStatus == EConcertConnectionStatus::Connected && SuspendedCount > 0;
+	SendReceiveState = InSendReceiveState;
+	LogSendReceiveStateChange(SessionInfo, SendReceiveState);
 }
 
 FOnConcertClientSessionTick& FConcertClientSession::OnTick()
diff --git a/Engine/Plugins/Developer/Concert/ConcertMain/Source/Concert/Private/ConcertClientSession.h b/Engine/Plugins/Developer/Concert/ConcertMain/Source/Concert/Private/ConcertClientSession.h
--- a/Engine/Plugins/Developer/Concert/ConcertMain/Source/Concert/Private/ConcertClientSession.h
+++ b/Engine/Plugins/Developer/Concert/ConcertMain/Source/Concert/Private/ConcertClientSession.h
@@ -84,9 +84,10 @@ public:
 
 	virtual void Connect() override;
 	virtual void Disconnect() override;
-	virtual void Resume() override;
-	virtual void Suspend() override;
-	virtual bool IsSuspended() const override;
+
+	virtual EConcertSendReceiveState GetSendReceiveState() const override;
+	virtual void SetSendReceiveState(EConcertSendReceiveState InSendReceiveState) override;
+
 	virtual FOnConcertClientSessionTick& OnTick() override;
 	virtual FOnConcertClientSessionConnectionChanged& OnConnectionChanged() override;
 	virtual FOnConcertClientSessionClientChanged& OnSessionClientChanged() override;
@@ -115,7 +116,7 @@ protected:
 	}
 
 	virtual void InternalSendCustomEvent(const UScriptStruct* EventType, const void* EventData, const TArray<FGuid>& DestinationEndpointIds, EConcertMessageFlags Flags) override;
-	
+
 	virtual void InternalRegisterCustomRequestHandler(const FName& RequestMessageType, const TSharedRef<IConcertSessionCustomRequestHandler>& Handler) override
 	{
 		CommonRegisterCustomRequestHandler(RequestMessageType, Handler);
@@ -134,7 +135,7 @@ private:
 
 	/**  */
 	void HandleJoinSessionResultEvent(const FConcertMessageContext& Context);
-	
+
 	/**  */
 	void HandleClientListUpdatedEvent(const FConcertMessageContext& Context);
 
@@ -168,8 +169,8 @@ private:
 	/** This session endpoint where message are sent and received from. */
 	IConcertLocalEndpointPtr ClientSessionEndpoint;
 
-	/** Count of the number of times this session has been suspended */
-	uint8 SuspendedCount;
+	/** State of the connection */
+	EConcertSendReceiveState SendReceiveState = EConcertSendReceiveState::Default;
 
 	/** Ticker for the session */
 	FTSTicker::FDelegateHandle SessionTick;
diff --git a/Engine/Plugins/Developer/Concert/ConcertMain/Source/Concert/Public/ConcertMessages.h b/Engine/Plugins/Developer/Concert/ConcertMain/Source/Concert/Public/ConcertMessages.h
--- a/Engine/Plugins/Developer/Concert/ConcertMain/Source/Concert/Public/ConcertMessages.h
+++ b/Engine/Plugins/Developer/Concert/ConcertMain/Source/Concert/Public/ConcertMessages.h
@@ -50,6 +50,19 @@ enum class EConcertClientStatus : uint8
 	Updated,
 };
 
+UENUM()
+enum class EConcertSendReceiveState : uint8
+{
+	/** Full send / receive */
+	Default,
+
+	/** Sender only. Local transaction events will be sent to server but receiving data from clients has been suspended */
+	SendOnly,
+
+	/** Receive only. Updates will be received from clients but local transaction events will be deferred until state has been set to default or send. */
+	ReceiveOnly
+};
+
 /** Response codes for a session custom request */
 UENUM()
 enum class EConcertSessionResponseCode : uint8
diff --git a/Engine/Plugins/Developer/Concert/ConcertMain/Source/Concert/Public/IConcertClient.h b/Engine/Plugins/Developer/Concert/ConcertMain/Source/Concert/Public/IConcertClient.h
--- a/Engine/Plugins/Developer/Concert/ConcertMain/Source/Concert/Public/IConcertClient.h
+++ b/Engine/Plugins/Developer/Concert/ConcertMain/Source/Concert/Public/IConcertClient.h
@@ -357,20 +357,29 @@ public:
 	 */
 	virtual void DisconnectSession() = 0;
 
-	/**
-	 * Resume live-updates for the current session (must be paired with a call to SuspendSession).
-	 */
-	virtual void ResumeSession() = 0;
-
-	/**
-	 * Suspend live-updates for the current session.
-	 */
-	virtual void SuspendSession() = 0;
-
-	/**
-	 * Does the current session have live-updates suspended?
-	 */
-	virtual bool IsSessionSuspended() const = 0;
+	UE_DEPRECATED(5.2,"Use SetSendReceiveState to change live update status.")
+	virtual void ResumeSession()
+	{
+		SetSendReceiveState(EConcertSendReceiveState::Default);
+	}
+
+	UE_DEPRECATED(5.2,"Use SetSendReceiveState to change live update status.")
+	virtual void SuspendSession()
+	{
+		SetSendReceiveState(EConcertSendReceiveState::SendOnly);
+	}
+
+	UE_DEPRECATED(5.2,"Use GetSendReceiveState to query the live update status.")
+	virtual bool IsSessionSuspended() const
+	{
+		return GetSendReceiveState() == EConcertSendReceiveState::SendOnly;
+	}
+
+	/** Get the send/receive state for this session. */
+	virtual EConcertSendReceiveState GetSendReceiveState() const = 0;
+
+	/** Set the send/receive state for this session. */
+	virtual void SetSendReceiveState(EConcertSendReceiveState InSendReceiveState) = 0;
 
 	/**
 	 * Does the client think it is the owner of the session?
diff --git a/Engine/Plugins/Developer/Concert/ConcertMain/Source/Concert/Public/IConcertSession.h b/Engine/Plugins/Developer/Concert/ConcertMain/Source/Concert/Public/IConcertSession.h
--- a/Engine/Plugins/Developer/Concert/ConcertMain/Source/Concert/Public/IConcertSession.h
+++ b/Engine/Plugins/Developer/Concert/ConcertMain/Source/Concert/Public/IConcertSession.h
@@ -289,13 +289,31 @@ public:
 	virtual void Disconnect() = 0;
 
 	/** Resume live-updates for this session (must be paired with a call to Suspend) */
-	virtual void Resume() = 0;
+	UE_DEPRECATED(5.2,"Use SetSendReceiveState to change live update status.")
+	virtual void Resume()
+	{
+		SetSendReceiveState(EConcertSendReceiveState::Default);
+	};
 
 	/** Suspend live-updates for this session */
-	virtual void Suspend() = 0;
+	UE_DEPRECATED(5.2,"Use SetSendReceiveState to change live update status.")
+	virtual void Suspend()
+	{
+		SetSendReceiveState(EConcertSendReceiveState::SendOnly);
+	}
 
 	/** Does this session currently have live-updates suspended? */
-	virtual bool IsSuspended() const = 0;
+	UE_DEPRECATED(5.2,"Use GetSendReceiveState to query the live update status.")
+	virtual bool IsSuspended() const
+	{
+		return GetSendReceiveState() == EConcertSendReceiveState::SendOnly;
+	}
+
+	/** Get the send/receive state for this session. */
+	virtual EConcertSendReceiveState GetSendReceiveState() const = 0;
+
+	/** Set the send/receive state for this session. */
+	virtual void SetSendReceiveState(EConcertSendReceiveState InSendReceiveState) = 0;
 
 	/** Callback when a connected client session gets ticked */
 	virtual FOnConcertClientSessionTick& OnTick() = 0;
diff --git a/Engine/Plugins/Developer/Concert/ConcertSync/ConcertSyncClient/Source/ConcertSyncClient/Private/ConcertClientPackageManager.cpp b/Engine/Plugins/Developer/Concert/ConcertSync/ConcertSyncClient/Source/ConcertSyncClient/Private/ConcertClientPackageManager.cpp
--- a/Engine/Plugins/Developer/Concert/ConcertSync/ConcertSyncClient/Source/ConcertSyncClient/Private/ConcertClientPackageManager.cpp
+++ b/Engine/Plugins/Developer/Concert/ConcertSync/ConcertSyncClient/Source/ConcertSyncClient/Private/ConcertClientPackageManager.cpp
@@ -631,7 +631,8 @@ bool FConcertClientPackageManager::IsReloadingPackage(FName PackageName) const
 
 bool FConcertClientPackageManager::CanHotReloadOrPurge() const
 {
-	return ConcertSyncClientUtil::CanPerformBlockingAction() && !LiveSession->GetSession().IsSuspended();
+	const bool bIsSuspended = LiveSession->GetSession().GetSendReceiveState() == EConcertSendReceiveState::SendOnly;
+	return ConcertSyncClientUtil::CanPerformBlockingAction() && !bIsSuspended;
 }
 
 void FConcertClientPackageManager::HotReloadPendingPackages()
diff --git a/Engine/Plugins/Developer/Concert/ConcertSync/ConcertSyncClient/Source/ConcertSyncClient/Private/ConcertClientTransactionManager.cpp b/Engine/Plugins/Developer/Concert/ConcertSync/ConcertSyncClient/Source/ConcertSyncClient/Private/ConcertClientTransactionManager.cpp
--- a/Engine/Plugins/Developer/Concert/ConcertSync/ConcertSyncClient/Source/ConcertSyncClient/Private/ConcertClientTransactionManager.cpp
+++ b/Engine/Plugins/Developer/Concert/ConcertSync/ConcertSyncClient/Source/ConcertSyncClient/Private/ConcertClientTransactionManager.cpp
@@ -266,7 +266,8 @@ void FConcertClientTransactionManager::HandleTransactionRejectedEvent(const FCon
 
 bool FConcertClientTransactionManager::CanProcessTransactionEvent() const
 {
-	return TransactionBridge->CanApplyRemoteTransaction() && !LiveSession->GetSession().IsSuspended();
+	const bool bIsSuspended = LiveSession->GetSession().GetSendReceiveState() == EConcertSendReceiveState::SendOnly;
+	return TransactionBridge->CanApplyRemoteTransaction() && !bIsSuspended;
 }
 
 void FConcertClientTransactionManager::ProcessTransactionEvent(const FPendingTransactionToProcessContext& InContext, const FStructOnScope& InEvent)
diff --git a/Engine/Plugins/Developer/Concert/ConcertSync/ConcertSyncClient/Source/ConcertSyncClient/Test/ConcertClientLiveTransactionAuthorsTests.cpp b/Engine/Plugins/Developer/Concert/ConcertSync/ConcertSyncClient/Source/ConcertSyncClient/Test/ConcertClientLiveTransactionAuthorsTests.cpp
--- a/Engine/Plugins/Developer/Concert/ConcertSync/ConcertSyncClient/Source/ConcertSyncClient/Test/ConcertClientLiveTransactionAuthorsTests.cpp
+++ b/Engine/Plugins/Developer/Concert/ConcertSync/ConcertSyncClient/Source/ConcertSyncClient/Test/ConcertClientLiveTransactionAuthorsTests.cpp
@@ -1,5 +1,6 @@
 // Copyright Epic Games, Inc. All Rights Reserved.
 
+#include "ConcertMessages.h"
 #include "CoreMinimal.h"
 #include "Misc/AutomationTest.h"
 #include "Misc/Paths.h"
@@ -59,9 +60,17 @@ public:
 	virtual void UpdateLocalClientInfo(const FConcertClientInfoUpdate&) override     { return NotMocked<void>(); }
 	virtual void Connect() override                                                  { return NotMocked<void>(); }
 	virtual void Disconnect() override                                               { return NotMocked<void>(); }
-	virtual void Resume() override                                                   { return NotMocked<void>(); }
-	virtual void Suspend() override                                                  { return NotMocked<void>(); }
-	virtual bool IsSuspended() const override                                        { return NotMocked<bool>(false); }
+
+	virtual EConcertSendReceiveState GetSendReceiveState() const override
+	{
+		return NotMocked<EConcertSendReceiveState>(EConcertSendReceiveState::Default);
+	}
+
+	virtual void SetSendReceiveState(EConcertSendReceiveState InSendReceiveState) override
+	{
+		return NotMocked<void>();
+	}
+
 	virtual FOnConcertClientSessionTick& OnTick() override                           { return NotMocked<FOnConcertClientSessionTick&>(Tick); }
 	virtual FOnConcertClientSessionConnectionChanged& OnConnectionChanged() override { return NotMocked<FOnConcertClientSessionConnectionChanged&>(ConnectionChanged); }
 	virtual FOnConcertClientSessionClientChanged& OnSessionClientChanged() override  { return NotMocked<FOnConcertClientSessionClientChanged&>(ClientChanged); }
diff --git a/Engine/Plugins/Developer/Concert/ConcertSync/ConcertSyncTest/Source/ConcertSyncTest/Private/ConcertDataStoreTests.cpp b/Engine/Plugins/Developer/Concert/ConcertSync/ConcertSyncTest/Source/ConcertSyncTest/Private/ConcertDataStoreTests.cpp
--- a/Engine/Plugins/Developer/Concert/ConcertSync/ConcertSyncTest/Source/ConcertSyncTest/Private/ConcertDataStoreTests.cpp
+++ b/Engine/Plugins/Developer/Concert/ConcertSync/ConcertSyncTest/Source/ConcertSyncTest/Private/ConcertDataStoreTests.cpp
@@ -156,9 +156,17 @@ public:
 	virtual void UpdateLocalClientInfo(const FConcertClientInfoUpdate&) override     { return NotMocked<void>(); }
 	virtual void Connect() override                                                  { return NotMocked<void>(); }
 	virtual void Disconnect() override                                               { return NotMocked<void>(); }
-	virtual void Resume() override                                                   { return NotMocked<void>(); }
-	virtual void Suspend() override                                                  { return NotMocked<void>(); }
-	virtual bool IsSuspended() const override                                        { return NotMocked<bool>(false); }
+
+	virtual EConcertSendReceiveState GetSendReceiveState() const override
+	{
+		return NotMocked<EConcertSendReceiveState>(EConcertSendReceiveState::Default);
+	}
+
+	virtual void SetSendReceiveState(EConcertSendReceiveState InSendReceiveState) override
+	{
+		return NotMocked<void>();
+	}
+
 	virtual FOnConcertClientSessionTick& OnTick() override                           { return NotMocked<FOnConcertClientSessionTick&>(Tick); }
 	virtual FOnConcertClientSessionConnectionChanged& OnConnectionChanged() override { return NotMocked<FOnConcertClientSessionConnectionChanged&>(ConnectionChanged); }
 	virtual FOnConcertClientSessionClientChanged& OnSessionClientChanged() override  { return NotMocked<FOnConcertClientSessionClientChanged&>(ClientChanged); }
