Add a pop-up to defer reload until the operation deems appropriate.  This is to prevent VCam users getting disrupted in a session.

diff --git a/Engine/Plugins/Developer/Concert/ConcertApp/MultiUserClient/Source/MultiUserClient/Private/Widgets/SActiveSession.cpp b/Engine/Plugins/Developer/Concert/ConcertApp/MultiUserClient/Source/MultiUserClient/Private/Widgets/SActiveSession.cpp
--- a/Engine/Plugins/Developer/Concert/ConcertApp/MultiUserClient/Source/MultiUserClient/Private/Widgets/SActiveSession.cpp
+++ b/Engine/Plugins/Developer/Concert/ConcertApp/MultiUserClient/Source/MultiUserClient/Private/Widgets/SActiveSession.cpp
@@ -8,9 +8,13 @@
 #include "ConcertMessageData.h"
 #include "ClientSessionHistoryController.h"
 #include "ConcertMessages.h"
+#include "Dialog/SMessageDialog.h"
+#include "IConcertClientPackageBridge.h"
 #include "IConcertClientPresenceManager.h"
 #include "IConcertClient.h"
+#include "IConcertClientWorkspace.h"
 #include "IConcertSyncClient.h"
+#include "ConcertClientSettings.h"
 #include "IMultiUserClientModule.h"
 
 #include "Algo/Transform.h"
@@ -20,6 +24,7 @@
 #include "Misc/AsyncTaskNotification.h"
 #include "Misc/PackageName.h"
 #include "Styling/AppStyle.h"
+#include "Dialog/SCustomDialog.h"
 #include "Session/History/SSessionHistory.h"
 #include "Session/History/SSessionHistoryWrapper.h"
 #include "Widgets/Notifications/SNotificationList.h"
@@ -30,6 +35,7 @@
 #include "Widgets/Views/STableViewBase.h"
 #include "Widgets/Views/STableRow.h"
 #include "Framework/Notifications/NotificationManager.h"
+#include "HAL/PlatformFileManager.h"
 
 
 #define LOCTEXT_NAMESPACE "SActiveSession"
@@ -395,6 +401,11 @@ void SActiveSession::Construct(const FArguments& InArgs, TSharedPtr<IConcertSync
 			ClientSession->OnSessionClientChanged().AddSP(this, &SActiveSession::HandleSessionClientChanged);
 			InConcertSyncClient->GetTransactionBridge()->OnConflictResolutionForPendingSend().AddSP(this, &SActiveSession::OnSendConflict);
 		}
+
+		TSharedPtr<IConcertClientWorkspace> Workspace = InConcertSyncClient->GetWorkspace();
+		Workspace->OnActivityAddedOrUpdated().AddSP(this, &SActiveSession::ActivityUpdated);
+		Workspace->AddWorkspaceCanProcessPackagesDelegate(
+			TEXT("SActiveSession"), FCanProcessPendingPackages::CreateSP(this, &SActiveSession::CanProcessPendingPackages));
 	}
 
 	SendReceiveComboList = ConstructSendReceiveComboList();
@@ -918,4 +929,90 @@ TSharedPtr<FConcertSessionClientInfo> SActiveSession::FindAvailableClient(const
 	return FoundClientPtr ? *FoundClientPtr : nullptr;
 }
 
+bool SActiveSession::CanProcessPendingPackages() const
+{
+	return bCanHotReload;
+}
+
+void SActiveSession::Tick(const FGeometry& AllottedGeometry, const double InCurrentTime, const float InDeltaTime)
+{
+	if (bCanHotReload && CanReloadDialog)
+	{
+		CanReloadDialog = {};
+	}
+}
+
+namespace UE::MultiUserClient::Private
+{
+static bool IsMapAsset(FName PackageName)
+{
+	FString FullPath;
+	if (FPackageName::TryConvertLongPackageNameToFilename(PackageName.ToString(), FullPath, FPackageName::GetMapPackageExtension()))
+	{
+		if (FPlatformFileManager::Get().GetPlatformFile().FileExists(*FullPath))
+		{
+			return true;
+		}
+	}
+	return false;
+}
+
+static TSharedPtr<SMessageDialog> MakeReloadDialog(FName PackageName)
+{
+	static const FText DialogTitle = LOCTEXT("ActiveSessionPendingHotReloadTitle", "A level hot reload is pending.");
+
+	static const FText ReloadButton = LOCTEXT("ActiveSessionPendingHotReloadButton", "Reload");
+
+	const FText ReloadMessage =
+		FText::Format(
+			LOCTEXT("ActiveSessionPendingHotReload","A hot reload on {0} is pending. Please dismiss this notification to continue with reload. All transactions will be queued until this occurs."),
+			 FText::FromName(PackageName)
+			);
+
+	return SNew(SMessageDialog)
+		.Title(DialogTitle)
+		.Message(ReloadMessage)
+		.Buttons( { SCustomDialog::FButton(ReloadButton) } );
+}
+
+}
+void SActiveSession::ActivityUpdated(const FConcertClientInfo& InClientInfo, const FConcertSyncActivity& InActivity, const FStructOnScope& ActivitySummary)
+{
+	if (InActivity.EventType == EConcertSyncActivityEventType::Package)
+	{
+		TSharedPtr<IConcertSyncClient> SyncClientPin = WeakConcertSyncClient.Pin();
+		if (!WeakSessionPtr.IsValid() || !SyncClientPin || CanReloadDialog)
+		{
+			return;
+		}
+
+		if (InActivity.EndpointId == WeakSessionPtr.Pin()->GetSessionClientEndpointId())
+		{
+			return;
+		}
+
+		check(ActivitySummary.GetStruct()->IsChildOf(FConcertSyncPackageActivitySummary::StaticStruct()));
+		const FConcertSyncPackageActivitySummary* ActivitySummaryPtr =
+			(FConcertSyncPackageActivitySummary*)ActivitySummary.GetStructMemory();
+
+		// User changed its settings.
+		const UConcertClientConfig* ConcertClientConfig = SyncClientPin->GetConcertClient()->GetConfiguration();
+		if (ConcertClientConfig->bShouldPromptForHotReloadOnLevel)
+		{
+			const bool bIsAMap = UE::MultiUserClient::Private::IsMapAsset(ActivitySummaryPtr->PackageName);
+			if (bIsAMap)
+			{
+				bCanHotReload = false;
+				CanReloadDialog = UE::MultiUserClient::Private::MakeReloadDialog(ActivitySummaryPtr->PackageName);
+				CanReloadDialog->GetOnWindowClosedEvent().AddLambda([this](const TSharedRef<SWindow>& ClosedWindow)
+					{
+						bCanHotReload = true;
+					});
+
+				CanReloadDialog->Show();
+			}
+		}
+	}
+}
+
 #undef LOCTEXT_NAMESPACE /* SActiveSession */
diff --git a/Engine/Plugins/Developer/Concert/ConcertApp/MultiUserClient/Source/MultiUserClient/Private/Widgets/SActiveSession.h b/Engine/Plugins/Developer/Concert/ConcertApp/MultiUserClient/Source/MultiUserClient/Private/Widgets/SActiveSession.h
--- a/Engine/Plugins/Developer/Concert/ConcertApp/MultiUserClient/Source/MultiUserClient/Private/Widgets/SActiveSession.h
+++ b/Engine/Plugins/Developer/Concert/ConcertApp/MultiUserClient/Source/MultiUserClient/Private/Widgets/SActiveSession.h
@@ -12,11 +12,16 @@
 
 class IConcertClientSession;
 class IConcertSyncClient;
+
+class FAsyncTaskNotification;
 struct FConcertConflictDescriptionBase;
 struct FConcertSessionClientInfo;
+
+class SCustomDialog;
 class SDockTab;
 class SExpandableArea;
 
+
 /**
  * Displays the multi-users active session clients and activity, enables the client
  * to leave the session, open the same level as another client or teleport to another
@@ -46,6 +51,8 @@ public:
 	*/
 	void Construct(const FArguments& InArgs, TSharedPtr<IConcertSyncClient> InConcertSyncClient);
 
+	virtual void Tick(const FGeometry& AllottedGeometry, const double InCurrentTime, const float InDeltaTime) override;
+
 private:
 
 	/** Generate a new client row */
@@ -93,6 +100,15 @@ private:
 	/** Delegate handler when a conflict between inbound transaction and those stored pending outbound transactions. */
 	void OnSendConflict(const FConcertConflictDescriptionBase& ConflictMsg);
 
+	/** Delegate handler on inbound transactions that focus specifically on packages. */
+	void OnPackageChangeActivity();
+
+	/** Delegate handler to indicate if we can hot reload an inbound package. */
+	bool CanProcessPendingPackages() const;
+
+	/** Delegate handler invoked when an activity has been received by multi-user. */
+	void ActivityUpdated(const FConcertClientInfo& InClientInfo, const FConcertSyncActivity& InActivity, const FStructOnScope& /*unused*/);
+
 private:
 
 	/** Get the text object for the send/receive combo box. */
@@ -137,6 +153,12 @@ private:
 	/** Available states for the SendReceiveComboBox */
 	TArray< TSharedPtr< FSendReceiveComboItem > > SendReceiveComboList;
 
+	/** Notification handler for hot reload. */
+	TSharedPtr<SCustomDialog> CanReloadDialog;
+
+	/** Flag to indicate if it is OK to hotreload the packages. */
+	bool bCanHotReload = true;
+
 	/** Keeps the status of 'Clients' area expansion. */
 	bool bClientAreaExpanded = true;
 
diff --git a/Engine/Plugins/Developer/Concert/ConcertMain/Source/ConcertClient/Public/ConcertClientSettings.h b/Engine/Plugins/Developer/Concert/ConcertMain/Source/ConcertClient/Public/ConcertClientSettings.h
--- a/Engine/Plugins/Developer/Concert/ConcertMain/Source/ConcertClient/Public/ConcertClientSettings.h
+++ b/Engine/Plugins/Developer/Concert/ConcertMain/Source/ConcertClient/Public/ConcertClientSettings.h
@@ -222,6 +222,14 @@ public:
 	UPROPERTY(config, EditAnywhere, Category = "Client Settings")
 	FGameplayTagContainer SendOnlyAssignment;
 
+	/**
+	 * Hot reload of a level happens automatically when level / sublevel is changed. This property allows users
+	 * to receive a prompt before hot reload occurs so that it does not interrupt the user in the current level.
+	 * Can be specified on the editor command with `-CONCERTSHOULDPROMPTFORHOTRELOAD`
+	 */
+	UPROPERTY(config, EditAnywhere, Category="Client Settings")
+	bool bShouldPromptForHotReloadOnLevel = false;
+
 	/** Client & client session settings */
 	UPROPERTY(config, EditAnywhere, Category="Client Settings", meta=(ShowOnlyInnerProperties))
 	FConcertClientSettings ClientSettings;
diff --git a/Engine/Plugins/Developer/Concert/ConcertSync/ConcertSyncClient/Source/ConcertSyncClient/Private/ConcertClientPackageManager.cpp b/Engine/Plugins/Developer/Concert/ConcertSync/ConcertSyncClient/Source/ConcertSyncClient/Private/ConcertClientPackageManager.cpp
--- a/Engine/Plugins/Developer/Concert/ConcertSync/ConcertSyncClient/Source/ConcertSyncClient/Private/ConcertClientPackageManager.cpp
+++ b/Engine/Plugins/Developer/Concert/ConcertSync/ConcertSyncClient/Source/ConcertSyncClient/Private/ConcertClientPackageManager.cpp
@@ -631,7 +631,7 @@ bool FConcertClientPackageManager::IsReloadingPackage(FName PackageName) const
 
 bool FConcertClientPackageManager::CanHotReloadOrPurge() const
 {
-	const bool bIsSuspended = LiveSession->GetSession().GetSendReceiveState() == EConcertSendReceiveState::SendOnly;
+	const bool bIsSuspended = LiveSession->GetSession().GetSendReceiveState() != EConcertSendReceiveState::Default;
 	return ConcertSyncClientUtil::CanPerformBlockingAction() && !bIsSuspended;
 }
 
diff --git a/Engine/Plugins/Developer/Concert/ConcertSync/ConcertSyncClient/Source/ConcertSyncClient/Private/ConcertSyncClientModule.cpp b/Engine/Plugins/Developer/Concert/ConcertSync/ConcertSyncClient/Source/ConcertSyncClient/Private/ConcertSyncClientModule.cpp
--- a/Engine/Plugins/Developer/Concert/ConcertSync/ConcertSyncClient/Source/ConcertSyncClient/Private/ConcertSyncClientModule.cpp
+++ b/Engine/Plugins/Developer/Concert/ConcertSync/ConcertSyncClient/Source/ConcertSyncClient/Private/ConcertSyncClientModule.cpp
@@ -73,6 +73,9 @@ public:
 			ClientConfig->bAutoConnect |= FParse::Param(CommandLine, TEXT("CONCERTAUTOCONNECT"));
 			FParse::Bool(CommandLine, TEXT("-CONCERTAUTOCONNECT="), ClientConfig->bAutoConnect);
 
+			ClientConfig->bShouldPromptForHotReloadOnLevel |= FParse::Param(CommandLine, TEXT("CONCERTSHOULDPROMPTFORHOTRELOAD"));
+			FParse::Bool(CommandLine, TEXT("-CONCERTSHOULDPROMPTFORHOTRELOAD="), ClientConfig->bShouldPromptForHotReloadOnLevel);
+
 			ClientConfig->bRetryAutoConnectOnError |= FParse::Param(CommandLine, TEXT("CONCERTRETRYAUTOCONNECTONERROR"));
 			FParse::Bool(CommandLine, TEXT("-CONCERTRETRYAUTOCONNECTONERROR="), ClientConfig->bRetryAutoConnectOnError);
 
diff --git a/Engine/Plugins/Experimental/VirtualCamera/Source/VCamCore/Private/VCamMultiUser.cpp b/Engine/Plugins/Experimental/VirtualCamera/Source/VCamCore/Private/VCamMultiUser.cpp
--- a/Engine/Plugins/Experimental/VirtualCamera/Source/VCamCore/Private/VCamMultiUser.cpp
+++ b/Engine/Plugins/Experimental/VirtualCamera/Source/VCamCore/Private/VCamMultiUser.cpp
@@ -4,6 +4,7 @@
 
 #include "CineCameraComponent.h"
 #include "GameFramework/Actor.h"
+#include "UObject/Package.h"
 
 FMultiUserVCamCameraFocusData::FMultiUserVCamCameraFocusData(const UCineCameraComponent* CineCamera)
 	: ManualFocusDistance(CineCamera->CurrentFocusDistance)
@@ -41,6 +42,9 @@ FMultiUserVCamCameraData::FMultiUserVCamCameraData(const AActor* InOwner, const
 
 void FMultiUserVCamCameraData::ApplyTo(AActor* InOwner, UCineCameraComponent* InCineCameraComponent) const
 {
+	InOwner->Modify();
+	InCineCameraComponent->Modify();
+
 	InOwner->SetActorLocationAndRotation(CameraActorLocation, CameraActorRotation);
 	InCineCameraComponent->SetRelativeLocationAndRotation(CameraComponentLocation, CameraComponentRotation);
 	InCineCameraComponent->CurrentAperture = CurrentAperture;
