Multi-User asset actions submenu is missing

diff --git a/Engine/Plugins/Developer/Concert/ConcertApp/MultiUserClient/Source/MultiUserClient/MultiUserClient.Build.cs b/Engine/Plugins/Developer/Concert/ConcertApp/MultiUserClient/Source/MultiUserClient/MultiUserClient.Build.cs
--- a/Engine/Plugins/Developer/Concert/ConcertApp/MultiUserClient/Source/MultiUserClient/MultiUserClient.Build.cs
+++ b/Engine/Plugins/Developer/Concert/ConcertApp/MultiUserClient/Source/MultiUserClient/MultiUserClient.Build.cs
@@ -25,6 +25,7 @@ namespace UnrealBuildTool.Rules
 					"ConcertSyncClient",
 					"ConcertSyncCore",
 					"ConcertTransport",
+					"ContentBrowser",
 					"DesktopPlatform",
 					"EditorStyle",
 					"InputCore",
diff --git a/Engine/Plugins/Developer/Concert/ConcertApp/MultiUserClient/Source/MultiUserClient/Private/ConcertWorkspaceUI.cpp b/Engine/Plugins/Developer/Concert/ConcertApp/MultiUserClient/Source/MultiUserClient/Private/ConcertWorkspaceUI.cpp
--- a/Engine/Plugins/Developer/Concert/ConcertApp/MultiUserClient/Source/MultiUserClient/Private/ConcertWorkspaceUI.cpp
+++ b/Engine/Plugins/Developer/Concert/ConcertApp/MultiUserClient/Source/MultiUserClient/Private/ConcertWorkspaceUI.cpp
@@ -9,6 +9,7 @@
 #include "IConcertClientSequencerManager.h"
 #include "ConcertFrontendStyle.h"
 #include "ConcertLogGlobal.h"
+#include "ContentBrowserMenuContexts.h"
 
 #include "Algo/Transform.h"
 #include "Modules/ModuleManager.h"
@@ -498,6 +499,15 @@ private:
 // FConcertWorkspaceUI implementation.
 //------------------------------------------------------------------------------
 
+FName FConcertWorkspaceUI::ConcertSourceControlMenuOwnerName = TEXT("ConcertSourceControlMenu");
+FName FConcertWorkspaceUI::ConcertPersistSessionChangesMenuEntryName = TEXT("ConcertPersistSessionChanges");
+
+FName FConcertWorkspaceUI::AssetActionSectionNameToExtend = TEXT("ContentBrowser.AssetContextMenu");
+FName FConcertWorkspaceUI::ConcertAssetActionsMenuOwnerName = TEXT("ConcertAssetActionsMenu");
+FName FConcertWorkspaceUI::ConcertAssetActionsSectionName = TEXT("ConcertAssetActionsSection");
+FName FConcertWorkspaceUI::ConcertAssetActionsDynamicEntryName = TEXT("MultiUserActions");
+FName FConcertWorkspaceUI::ConcertAssetActionsSubMenuName = TEXT("MultiUserActionsSubMenu");
+
 FConcertWorkspaceUI::FConcertWorkspaceUI()
 {
 	// Extend ContentBrowser Asset Icon
@@ -568,34 +578,8 @@ void FConcertWorkspaceUI::InstallWorkspaceExtensions(TWeakPtr<IConcertClientWork
 	ClientWorkspace = InClientWorkspace;
 	SyncClient = InSyncClient;
 
-	// Extend ContentBrowser for session
-	if (FContentBrowserModule* ContentBrowserModule = FModuleManager::Get().GetModulePtr<FContentBrowserModule>(TEXT("ContentBrowser")))
-	{
-		// Asset Context Menu Extension
-		ContentBrowserAssetExtenderDelegateHandle = ContentBrowserModule->GetAllAssetViewContextMenuExtenders()
-			.Add_GetRef(FContentBrowserMenuExtender_SelectedAssets::CreateSP(this, &FConcertWorkspaceUI::OnExtendContentBrowserAssetSelectionMenu)).GetHandle();
-	}
-
-	FToolMenuOwnerScoped SourceControlMenuOwner("ConcertSourceControlMenu");
-
-	// Setup Concert Source Control Extension
-	UToolMenu* SourceControlMenu = UToolMenus::Get()->ExtendMenu("StatusBar.ToolBar.SourceControl");
-	FToolMenuSection& Section = SourceControlMenu->FindOrAddSection("SourceControlMenu");
-
-	TWeakPtr<FConcertWorkspaceUI> Weak = AsShared();
-	Section.AddMenuEntry(
-		"ConcertPersistSessionChanges",
-		LOCTEXT("ConcertWVPersist", "Persist Session Changes..."),
-		LOCTEXT("ConcertWVPersistTooltip", "Persist the session changes and prepare the files for revision control submission."),
-		FSlateIcon(FConcertFrontendStyle::GetStyleSetName(), "Concert.Persist"),
-		FUIAction(FExecuteAction::CreateLambda([Weak]()
-			{
-				if (TSharedPtr<FConcertWorkspaceUI> PinThis = Weak.Pin())
-				{
-					PinThis->PromptPersistSessionChanges(); // Required to adapt the function signature.
-				}
-			}))
-		);
+	ExtendToolbarWithSourceControlMenu();
+	ExtendAssetContextMenuOptions();
 
 	// Register for the "MarkPackageDirty" callback to catch packages that have been modified so we can acquire lock or warn
 	UPackage::PackageMarkedDirtyEvent.AddRaw(this, &FConcertWorkspaceUI::OnMarkPackageDirty);
@@ -611,7 +595,8 @@ void FConcertWorkspaceUI::UninstallWorspaceExtensions()
 		ContentBrowserAssetExtenderDelegateHandle.Reset();
 	}
 
-	UToolMenus::Get()->UnregisterOwnerByName("ConcertSourceControlMenu");
+	UToolMenus::Get()->UnregisterOwnerByName(ConcertSourceControlMenuOwnerName);
+	UToolMenus::Get()->UnregisterOwnerByName(ConcertAssetActionsMenuOwnerName);
 
 	// Remove package dirty hook
 	UPackage::PackageMarkedDirtyEvent.RemoveAll(this);
@@ -928,36 +913,57 @@ void FConcertWorkspaceUI::OnMarkPackageDirty(UPackage* InPackage, bool /*bWasDir
 	}
 }
 
-TSharedRef<FExtender> FConcertWorkspaceUI::OnExtendContentBrowserAssetSelectionMenu(const TArray<FAssetData>& SelectedAssets)
+void FConcertWorkspaceUI::ExtendToolbarWithSourceControlMenu()
 {
-	// Menu extender for Content Browser context menu when an asset is selected
-	TSharedRef<FExtender> Extender = MakeShared<FExtender>();
-	if (SelectedAssets.Num() > 0)
+	const FToolMenuOwnerScoped SourceControlMenuOwner(ConcertSourceControlMenuOwnerName);
+	// Setup Concert Source Control Extension
+	UToolMenu* const SourceControlMenu = UToolMenus::Get()->ExtendMenu("StatusBar.ToolBar.SourceControl");
+	FToolMenuSection& Section = SourceControlMenu->FindOrAddSection("SourceControlMenu");
+
+	Section.AddMenuEntry(
+		ConcertPersistSessionChangesMenuEntryName,
+		LOCTEXT("ConcertWVPersist", "Persist Session Changes..."),
+		LOCTEXT("ConcertWVPersistTooltip", "Persist the session changes and prepare the files for revision control submission."),
+		FSlateIcon(FConcertFrontendStyle::GetStyleSetName(), "Concert.Persist"),
+		FUIAction(FExecuteAction::CreateLambda([Weak = TWeakPtr<FConcertWorkspaceUI>(AsShared())]()
+		{
+			if (TSharedPtr<FConcertWorkspaceUI> PinThis = Weak.Pin())
+			{
+				PinThis->PromptPersistSessionChanges(); // Required to adapt the function signature.
+			}
+		}))
+	);
+}
+
+void FConcertWorkspaceUI::ExtendAssetContextMenuOptions()
+{
+	const FToolMenuOwnerScoped SourceControlMenuOwner(ConcertAssetActionsMenuOwnerName);
+	UToolMenu* const Menu = UToolMenus::Get()->ExtendMenu(AssetActionSectionNameToExtend);
+	FToolMenuSection& Section = Menu->AddSection(ConcertAssetActionsSectionName, LOCTEXT("MultiUserSectionLabel", "Multi-User"));
+	Section.AddDynamicEntry(ConcertAssetActionsDynamicEntryName, FNewToolMenuSectionDelegate::CreateLambda([WeakThis = TWeakPtr<FConcertWorkspaceUI>(AsShared())](FToolMenuSection& InSection)
 	{
+		const TSharedPtr<FConcertWorkspaceUI> PinThis = WeakThis.Pin();
+		UContentBrowserAssetContextMenuContext* Context = InSection.FindContext<UContentBrowserAssetContextMenuContext>();
+		if (!Context || !Context->bCanBeModified || Context->SelectedAssets.Num() == 0 || !ensure(PinThis))
+		{
+			return;
+		}
+
 		TArray<FName> TransformedAssets;
-		Algo::Transform(SelectedAssets, TransformedAssets, [](const FAssetData& AssetData)
+		Algo::Transform(Context->SelectedAssets, TransformedAssets, [](const FAssetData& AssetData)
 		{
 			return AssetData.PackageName;
 		});
 
-		TWeakPtr<FConcertWorkspaceUI> Weak = AsShared();
-		Extender->AddMenuExtension("AssetContextSourceControl", EExtensionHook::After, nullptr, FMenuExtensionDelegate::CreateLambda(
-			[Weak, AssetObjectPaths = MoveTemp(TransformedAssets)](FMenuBuilder& MenuBuilder) mutable
-			{
-				if (TSharedPtr<FConcertWorkspaceUI> PinThis = Weak.Pin())
-				{
-					MenuBuilder.AddMenuSeparator();
-					MenuBuilder.AddSubMenu(
-						LOCTEXT("Concert_ContextMenu", "Multi-User"),
-						FText(),
-						FNewMenuDelegate::CreateSP(PinThis.Get(), &FConcertWorkspaceUI::GenerateConcertAssetContextMenu, MoveTemp(AssetObjectPaths)),
-						false,
-						FSlateIcon(FConcertFrontendStyle::GetStyleSetName(), "Concert.MultiUser")
-					);
-				}
-			}));
-	}
-	return Extender;
+		InSection.AddSubMenu(
+			ConcertAssetActionsSubMenuName,
+			LOCTEXT("Concert_ContextMenu", "Multi-User Actions"),
+			FText::GetEmpty(),
+			FNewMenuDelegate::CreateSP(PinThis.Get(), &FConcertWorkspaceUI::GenerateConcertAssetContextMenu, MoveTemp(TransformedAssets)),
+			false,
+			FSlateIcon(FConcertFrontendStyle::GetStyleSetName(), "Concert.MultiUser")
+			);
+	}));
 }
 
 void FConcertWorkspaceUI::GenerateConcertAssetContextMenu(FMenuBuilder& MenuBuilder, TArray<FName> AssetObjectPaths)
diff --git a/Engine/Plugins/Developer/Concert/ConcertApp/MultiUserClient/Source/MultiUserClient/Private/ConcertWorkspaceUI.h b/Engine/Plugins/Developer/Concert/ConcertApp/MultiUserClient/Source/MultiUserClient/Private/ConcertWorkspaceUI.h
--- a/Engine/Plugins/Developer/Concert/ConcertApp/MultiUserClient/Source/MultiUserClient/Private/ConcertWorkspaceUI.h
+++ b/Engine/Plugins/Developer/Concert/ConcertApp/MultiUserClient/Source/MultiUserClient/Private/ConcertWorkspaceUI.h
@@ -25,6 +25,22 @@ struct FSequencerInitParams;
  */
 class FConcertWorkspaceUI : public TSharedFromThis<FConcertWorkspaceUI>
 {
+	/** Name of menu extensions going into source control (bottom-right) */
+	static FName ConcertSourceControlMenuOwnerName;
+	/** Name for menu entry in source control that will persist changes */
+	static FName ConcertPersistSessionChangesMenuEntryName;
+
+	/** Name of UE section in asset context */
+	static FName AssetActionSectionNameToExtend;
+	/** Name of menu extensions going to asset context menu */
+	static FName ConcertAssetActionsMenuOwnerName;
+	/** Name of the new section in the asset action menu*/
+	static FName ConcertAssetActionsSectionName;
+	/** Name of entry added to asset context menu */
+	static FName ConcertAssetActionsDynamicEntryName;
+	/** Name of submenu added by ConcertAssetActionsDynamicEntryName */
+	static FName ConcertAssetActionsSubMenuName;
+
 public:
 	FConcertWorkspaceUI();
 	~FConcertWorkspaceUI();
@@ -100,11 +116,12 @@ private:
 	bool IsAssetModifiedByOtherClients(const FName& AssetName, int32* OutOtherClientsWithModifNum = nullptr, TArray<FConcertClientInfo>* OutOtherClientsWithModifInfo = nullptr, int32 OtherClientsWithModifMaxFetchNum = 0) const;
 
 	/** Delegate called when a package is marked dirty. */
-	void OnMarkPackageDirty(class UPackage* InPackage, bool bDirty);
-
-	/** Delegate to extend the content browser asset context menu. */
-	TSharedRef<FExtender> OnExtendContentBrowserAssetSelectionMenu(const TArray<FAssetData>& SelectedAssets);
+	void OnMarkPackageDirty(UPackage* InPackage, bool bDirty);
 
+	/** Extends the toolbar with MU source control options */
+	void ExtendToolbarWithSourceControlMenu();
+	/** Extends the content browser asset context menu with MU options. */
+	void ExtendAssetContextMenuOptions();
 	/** Called to generate concert asset context menu. */
 	void GenerateConcertAssetContextMenu(FMenuBuilder& MenuBuilder, TArray<FName> AssetObjectPaths);
 
