Expose Multi-user session info to blueprints.

diff --git a/Engine/Plugins/Developer/Concert/ConcertApp/MultiUserClient/Source/MultiUserClientLibrary/Private/MultiUserClientStatics.cpp b/Engine/Plugins/Developer/Concert/ConcertApp/MultiUserClient/Source/MultiUserClientLibrary/Private/MultiUserClientStatics.cpp
--- a/Engine/Plugins/Developer/Concert/ConcertApp/MultiUserClient/Source/MultiUserClientLibrary/Private/MultiUserClientStatics.cpp
+++ b/Engine/Plugins/Developer/Concert/ConcertApp/MultiUserClient/Source/MultiUserClientLibrary/Private/MultiUserClientStatics.cpp
@@ -4,6 +4,8 @@
 #include "Logging/LogMacros.h"
 
 #if WITH_CONCERT
+#include "Algo/Find.h"
+
 #include "IMultiUserClientModule.h"
 #include "IConcertSyncClient.h"
 #include "IConcertClient.h"
@@ -57,6 +59,17 @@ EMultiUserConnectionStatus ConvertConnectionStatus(EConcertConnectionStatus Conn
 	return static_cast<EMultiUserConnectionStatus>(ConnectionStatus);
 }
 
+FMultiUserSessionInfo ConvertSessionInfo(const FConcertSessionInfo& InSessionInfo, const FConcertServerInfo& InServerInfo)
+{
+	FMultiUserSessionInfo ReturnVal;
+	ReturnVal.SessionName = InSessionInfo.SessionName;
+	ReturnVal.ServerName = InServerInfo.ServerName;
+	ReturnVal.EndpointName = InServerInfo.InstanceInfo.InstanceName;
+	ReturnVal.ServerEndpointId = InSessionInfo.ServerInstanceId;
+	ReturnVal.bValid = true;
+	return ReturnVal;
+}
+
 } // namespace MultiUserClientUtil
 #endif
 
@@ -252,6 +265,34 @@ FMultiUserClientInfo UMultiUserClientStatics::GetLocalMultiUserClientInfo()
 	return ClientInfo;
 }
 
+FMultiUserSessionInfo UMultiUserClientStatics::GetMultiUserSessionInfo()
+{
+#if WITH_CONCERT
+	if (IMultiUserClientModule::IsAvailable())
+	{
+		if (TSharedPtr<IConcertSyncClient> ConcertSyncClient = IMultiUserClientModule::Get().GetClient())
+		{
+			IConcertClientRef ConcertClient = ConcertSyncClient->GetConcertClient();
+			TSharedPtr<IConcertClientSession> ClientSession = ConcertClient->GetCurrentSession();
+			if (ClientSession)
+			{
+				const FConcertSessionInfo& SessionInfo = ClientSession->GetSessionInfo();
+				TArray<FConcertServerInfo> Servers = ConcertClient->GetKnownServers();
+				FConcertServerInfo* ServerInfo = Algo::FindByPredicate(Servers, [&SessionInfo](const FConcertServerInfo& InServerInfo)
+				{
+					return InServerInfo.InstanceInfo.InstanceId == SessionInfo.ServerInstanceId;
+				});
+				if (ServerInfo)
+				{
+					return MultiUserClientUtil::ConvertSessionInfo(SessionInfo, *ServerInfo);
+				}
+			}
+		}
+	}
+#endif
+	return {};
+}
+
 bool UMultiUserClientStatics::GetMultiUserClientInfoByName(const FString& ClientName, FMultiUserClientInfo& ClientInfo)
 {
 #if WITH_CONCERT
diff --git a/Engine/Plugins/Developer/Concert/ConcertApp/MultiUserClient/Source/MultiUserClientLibrary/Public/MultiUserClientStatics.h b/Engine/Plugins/Developer/Concert/ConcertApp/MultiUserClient/Source/MultiUserClientLibrary/Public/MultiUserClientStatics.h
--- a/Engine/Plugins/Developer/Concert/ConcertApp/MultiUserClient/Source/MultiUserClientLibrary/Public/MultiUserClientStatics.h
+++ b/Engine/Plugins/Developer/Concert/ConcertApp/MultiUserClient/Source/MultiUserClientLibrary/Public/MultiUserClientStatics.h
@@ -35,6 +35,36 @@ struct FMultiUserClientInfo
 	TArray<FName> Tags;
 };
 
+/**
+ * BP copy of FConcertSessionInfo
+ * Holds info on a connected session.
+ */
+USTRUCT(BlueprintType)
+struct FMultiUserSessionInfo
+{
+	GENERATED_BODY()
+
+	/** Holds the server endpoint id for this session. */
+	UPROPERTY(VisibleAnywhere, BlueprintReadOnly, Category = "Session Info")
+	FGuid ServerEndpointId;
+
+	/** Holds the session name current connected. */
+	UPROPERTY(VisibleAnywhere, BlueprintReadOnly, Category = "Session Info")
+	FString SessionName;
+
+	/** Holds the server name for the connected session. */
+	UPROPERTY(VisibleAnywhere, BlueprintReadOnly, Category = "Session Info")
+	FString ServerName;
+
+	/** Holds the resolved endpoint name (HostPC) for the connected session. */
+	UPROPERTY(VisibleAnywhere, BlueprintReadOnly, Category = "Session Info")
+	FString EndpointName;
+
+	/** Indicates if the current session info is valid. If we are connected to a session this is true, otherwise false. */
+	UPROPERTY(VisibleAnywhere, BlueprintReadOnly, Category = "Session Info")
+	bool bValid = false;
+};
+
 UCLASS()
 class UMultiUserClientSyncDatabase : public UObject
 {
@@ -174,6 +204,10 @@ public:
 	UFUNCTION(BlueprintCallable, Category = "Multi-User Client", meta=(DevelopmentOnly, DisplayName = "Get Local Multi-User Client Info"))
 	static FMultiUserClientInfo GetLocalMultiUserClientInfo();
 
+	/** Get the current SessionInfo. Works when not connected to a session, but returns an empty session info. */
+	UFUNCTION(BlueprintCallable, Category = "Multi-User Client", meta=(DevelopmentOnly, DisplayName = "Get Multi-User Session Info"))
+	static FMultiUserSessionInfo GetMultiUserSessionInfo();
+
 	/** Get the ClientInfo for any Multi-User participant by name. The local user is found even when not connected to a session. Returns false is no client was found. */
 	UFUNCTION(BlueprintCallable, Category = "Multi-User Client", meta=(DevelopmentOnly, DisplayName = "Get Multi-User Client Info by Name"))
 	static bool GetMultiUserClientInfoByName(const FString& ClientName, FMultiUserClientInfo& ClientInfo);
