Check check for dirty packages on a per frame basis to avoid performance hit during take record.

diff --git a/Engine/Plugins/Developer/Concert/ConcertApp/MultiUserClient/Source/MultiUserClient/Private/ConcertWorkspaceUI.cpp b/Engine/Plugins/Developer/Concert/ConcertApp/MultiUserClient/Source/MultiUserClient/Private/ConcertWorkspaceUI.cpp
--- a/Engine/Plugins/Developer/Concert/ConcertApp/MultiUserClient/Source/MultiUserClient/Private/ConcertWorkspaceUI.cpp
+++ b/Engine/Plugins/Developer/Concert/ConcertApp/MultiUserClient/Source/MultiUserClient/Private/ConcertWorkspaceUI.cpp
@@ -583,6 +583,7 @@ void FConcertWorkspaceUI::InstallWorkspaceExtensions(TWeakPtr<IConcertClientWork
 
 	// Register for the "MarkPackageDirty" callback to catch packages that have been modified so we can acquire lock or warn
 	UPackage::PackageMarkedDirtyEvent.AddRaw(this, &FConcertWorkspaceUI::OnMarkPackageDirty);
+	FCoreDelegates::OnEndFrame.AddRaw(this, &FConcertWorkspaceUI::OnEndFrame);
 }
 
 void FConcertWorkspaceUI::UninstallWorspaceExtensions()
@@ -600,6 +601,7 @@ void FConcertWorkspaceUI::UninstallWorspaceExtensions()
 
 	// Remove package dirty hook
 	UPackage::PackageMarkedDirtyEvent.RemoveAll(this);
+	FCoreDelegates::OnEndFrame.RemoveAll(this);
 
 	ClientWorkspace.Reset();
 }
@@ -874,8 +876,23 @@ bool FConcertWorkspaceUI::IsAssetModifiedByOtherClients(const FName& AssetName,
 	return false;
 }
 
+void FConcertWorkspaceUI::OnEndFrame()
+{
+	PerFramePackageDirtyList.Reset();
+}
+
 void FConcertWorkspaceUI::OnMarkPackageDirty(UPackage* InPackage, bool /*bWasDirty*/)
 {
+	if (PerFramePackageDirtyList.Contains(InPackage->GetFName()))
+	{
+		// In a take recorder situation, a record can include a large number of dirty events (especially with
+		// a large number of properties).  This means that this delegate would get invoked at least once
+		// per property per frame.  We reduce this overhead by storing the checked package and prevent
+		// the relatively expensive call below.
+		//
+		return;
+	}
+
 	TSharedPtr<IConcertClientWorkspace> ClientWorkspacePin = ClientWorkspace.Pin();
 	if (ClientWorkspacePin.IsValid()
 		&& !ClientWorkspacePin->ShouldIgnorePackageDirtyEvent(InPackage)
@@ -911,6 +928,7 @@ void FConcertWorkspaceUI::OnMarkPackageDirty(UPackage* InPackage, bool /*bWasDir
 			}
 		}
 	}
+	PerFramePackageDirtyList.Add(InPackage->GetFName());
 }
 
 void FConcertWorkspaceUI::ExtendToolbarWithSourceControlMenu()
diff --git a/Engine/Plugins/Developer/Concert/ConcertApp/MultiUserClient/Source/MultiUserClient/Private/ConcertWorkspaceUI.h b/Engine/Plugins/Developer/Concert/ConcertApp/MultiUserClient/Source/MultiUserClient/Private/ConcertWorkspaceUI.h
--- a/Engine/Plugins/Developer/Concert/ConcertApp/MultiUserClient/Source/MultiUserClient/Private/ConcertWorkspaceUI.h
+++ b/Engine/Plugins/Developer/Concert/ConcertApp/MultiUserClient/Source/MultiUserClient/Private/ConcertWorkspaceUI.h
@@ -67,6 +67,9 @@ private:
 	friend class SConcertWorkspaceModifiedByOtherTooltip;
 	friend class SConcertWorkspaceSequencerToolbarExtension;
 
+	/** On end frame callback. */
+	void OnEndFrame();
+
 	/** Check out and optionally submit files to source control. */
 	bool SubmitChangelist(const FConcertPersistCommand& PersistCommand, FText& OperationMessage);
 
@@ -174,4 +177,7 @@ private:
 
 	/** Toolbar extender for sequencer toolbar */
 	TSharedPtr<FExtender> ToolbarExtender;
+
+	/** Set to track dirty packages per frame. */
+	TSet<FName> PerFramePackageDirtyList;
 };
