Editor for multi user replication asset

diff --git a/Engine/Plugins/Developer/Concert/ConcertApp/MultiUserClient/Config/BaseMultiUserClient.ini b/Engine/Plugins/Developer/Concert/ConcertApp/MultiUserClient/Config/BaseMultiUserClient.ini
new file mode 100644
--- /dev/null
+++ b/Engine/Plugins/Developer/Concert/ConcertApp/MultiUserClient/Config/BaseMultiUserClient.ini
@@ -0,0 +1,2 @@
+[/Script/MultiUserReplicationEditor.MultiUserReplicationSettings]
++DefaultPropertySelection=(("/Script/Engine.SceneComponent", (DefaultSelectedProperties=("RelativeLocation","RelativeLocation.X","RelativeLocation.Y","RelativeLocation.Z","RelativeRotation","RelativeRotation.X","RelativeRotation.Y","RelativeRotation.Z","RelativeScale","RelativeScale.X","RelativeScale.Y","RelativeScale.Z"))))
\ No newline at end of file
diff --git a/Engine/Plugins/Developer/Concert/ConcertApp/MultiUserClient/MultiUserClient.uplugin b/Engine/Plugins/Developer/Concert/ConcertApp/MultiUserClient/MultiUserClient.uplugin
--- a/Engine/Plugins/Developer/Concert/ConcertApp/MultiUserClient/MultiUserClient.uplugin
+++ b/Engine/Plugins/Developer/Concert/ConcertApp/MultiUserClient/MultiUserClient.uplugin
@@ -24,6 +24,16 @@
 			"Name": "MultiUserClientLibrary",
 			"Type": "Runtime",
 			"LoadingPhase": "Default"
+		},
+		{
+			"Name": "MultiUserReplication",
+			"Type": "Runtime",
+			"LoadingPhase": "Default"
+		},
+		{
+			"Name": "MultiUserReplicationEditor",
+			"Type": "Editor",
+			"LoadingPhase": "PostDefault"
 		}
 	],
 	"Plugins": [
diff --git a/Engine/Plugins/Developer/Concert/ConcertApp/MultiUserClient/Source/MultiUserReplication/MultiUserReplication.Build.cs b/Engine/Plugins/Developer/Concert/ConcertApp/MultiUserClient/Source/MultiUserReplication/MultiUserReplication.Build.cs
new file mode 100644
--- /dev/null
+++ b/Engine/Plugins/Developer/Concert/ConcertApp/MultiUserClient/Source/MultiUserReplication/MultiUserReplication.Build.cs
@@ -0,0 +1,28 @@
+// Copyright Epic Games, Inc. All Rights Reserved.
+
+namespace UnrealBuildTool.Rules
+{
+	public class MultiUserReplication : ModuleRules
+	{
+		public MultiUserReplication(ReadOnlyTargetRules Target) : base(Target)
+		{
+			PublicDependencyModuleNames.AddRange(
+				new string[]
+				{
+					"Core",
+					"CoreUObject",
+					"ConcertSyncCore", 
+					"Engine",
+				}
+			);
+			
+			PrivateDependencyModuleNames.AddRange(
+				new string[]
+				{
+					"Concert",
+					"ConcertClient"
+				}
+			);
+		}
+	}
+}
diff --git a/Engine/Plugins/Developer/Concert/ConcertApp/MultiUserClient/Source/MultiUserReplication/Private/MultiUserPropertyReplicationSelection.cpp b/Engine/Plugins/Developer/Concert/ConcertApp/MultiUserClient/Source/MultiUserReplication/Private/MultiUserPropertyReplicationSelection.cpp
new file mode 100644
--- /dev/null
+++ b/Engine/Plugins/Developer/Concert/ConcertApp/MultiUserClient/Source/MultiUserReplication/Private/MultiUserPropertyReplicationSelection.cpp
@@ -0,0 +1,3 @@
+// Copyright Epic Games, Inc. All Rights Reserved.
+
+#include "MultiUserPropertyReplicationSelection.h"
\ No newline at end of file
diff --git a/Engine/Plugins/Developer/Concert/ConcertApp/MultiUserClient/Source/MultiUserReplication/Private/MultiUserReplicationModule.cpp b/Engine/Plugins/Developer/Concert/ConcertApp/MultiUserClient/Source/MultiUserReplication/Private/MultiUserReplicationModule.cpp
new file mode 100644
--- /dev/null
+++ b/Engine/Plugins/Developer/Concert/ConcertApp/MultiUserClient/Source/MultiUserReplication/Private/MultiUserReplicationModule.cpp
@@ -0,0 +1,20 @@
+// Copyright Epic Games, Inc. All Rights Reserved.
+
+#include "MultiUserReplicationModule.h"
+
+#include "Modules/ModuleManager.h"
+
+namespace UE::MultiUserReplication::Private
+{
+	void FMultiUserReplicationModule::StartupModule()
+	{
+		
+	}
+
+	void FMultiUserReplicationModule::ShutdownModule()
+	{
+		
+	}
+};
+
+IMPLEMENT_MODULE(UE::MultiUserReplication::Private::FMultiUserReplicationModule, MultiUserReplication);
diff --git a/Engine/Plugins/Developer/Concert/ConcertApp/MultiUserClient/Source/MultiUserReplication/Private/MultiUserReplicationModule.h b/Engine/Plugins/Developer/Concert/ConcertApp/MultiUserClient/Source/MultiUserReplication/Private/MultiUserReplicationModule.h
new file mode 100644
--- /dev/null
+++ b/Engine/Plugins/Developer/Concert/ConcertApp/MultiUserClient/Source/MultiUserReplication/Private/MultiUserReplicationModule.h
@@ -0,0 +1,13 @@
+// Copyright Epic Games, Inc. All Rights Reserved.
+
+#include "Modules/ModuleManager.h"
+
+namespace UE::MultiUserReplication::Private
+{
+	class FMultiUserReplicationModule : public IModuleInterface
+	{
+	public:
+		virtual void StartupModule() override;
+		virtual void ShutdownModule() override;
+	};
+}
diff --git a/Engine/Plugins/Developer/Concert/ConcertApp/MultiUserClient/Source/MultiUserReplication/Private/MultiUserReplicationStreamAsset.cpp b/Engine/Plugins/Developer/Concert/ConcertApp/MultiUserClient/Source/MultiUserReplication/Private/MultiUserReplicationStreamAsset.cpp
new file mode 100644
--- /dev/null
+++ b/Engine/Plugins/Developer/Concert/ConcertApp/MultiUserClient/Source/MultiUserReplication/Private/MultiUserReplicationStreamAsset.cpp
@@ -0,0 +1,15 @@
+// Copyright Epic Games, Inc. All Rights Reserved.
+
+#include "MultiUserReplicationStreamAsset.h"
+
+UMultiUserReplicationStreamAsset::UMultiUserReplicationStreamAsset()
+{
+	ReplicationList = CreateDefaultSubobject<UMultiUserPropertyReplicationSelection>(TEXT("ReplicationList"));
+	ReplicationList->SetFlags(RF_Transactional);
+}
+
+FReplicationStreamDescription UMultiUserReplicationStreamAsset::GenerateDescription() const
+{
+	// TODO:
+	return {};
+}
diff --git a/Engine/Plugins/Developer/Concert/ConcertApp/MultiUserClient/Source/MultiUserReplication/Public/MultiUserPropertyReplicationSelection.h b/Engine/Plugins/Developer/Concert/ConcertApp/MultiUserClient/Source/MultiUserReplication/Public/MultiUserPropertyReplicationSelection.h
new file mode 100644
--- /dev/null
+++ b/Engine/Plugins/Developer/Concert/ConcertApp/MultiUserClient/Source/MultiUserReplication/Public/MultiUserPropertyReplicationSelection.h
@@ -0,0 +1,23 @@
+// Copyright Epic Games, Inc. All Rights Reserved.
+
+#pragma once
+
+#include "UObject/Object.h"
+#include "Replication/Data/ObjectReplicationMap.h"
+#include "MultiUserPropertyReplicationSelection.generated.h"
+
+/**
+ * Holds the objects and properties to be replicated.
+ *
+ * Separate from UMultiUserReplicationStreamAsset so editor UI code editing it can be easily exist independently from
+ * UMultiUserReplicationStreamAsset.
+ */
+UCLASS()
+class MULTIUSERREPLICATION_API UMultiUserPropertyReplicationSelection : public UObject
+{
+	GENERATED_BODY()
+public:
+
+	UPROPERTY()
+	FObjectReplicationMap ReplicationMap;
+};
diff --git a/Engine/Plugins/Developer/Concert/ConcertApp/MultiUserClient/Source/MultiUserReplication/Public/MultiUserReplicationClientProfileAsset.h b/Engine/Plugins/Developer/Concert/ConcertApp/MultiUserClient/Source/MultiUserReplication/Public/MultiUserReplicationClientProfileAsset.h
new file mode 100644
--- /dev/null
+++ b/Engine/Plugins/Developer/Concert/ConcertApp/MultiUserClient/Source/MultiUserReplication/Public/MultiUserReplicationClientProfileAsset.h
@@ -0,0 +1,24 @@
+// Copyright Epic Games, Inc. All Rights Reserved.
+
+#pragma once
+
+#include "UObject/Object.h"
+#include "Replication/Data/ReplicationClientDescription.h"
+#include "MultiUserReplicationStreamAsset.h"
+#include "MultiUserReplicationClientProfileAsset.generated.h"
+
+/** Asset for users to describe a client to the server. */
+UCLASS()
+class MULTIUSERREPLICATION_API UMultiUserReplicationClientProfileAsset : public UObject
+{
+	GENERATED_BODY()
+public:
+
+	/** Describes this client to the server */
+	UPROPERTY(EditAnywhere, Category = "Replication")
+	FReplicationClientDescription ClientDescription;
+
+	/** The streams this client will send to the server. It takes authority over the listed properties. */
+	UPROPERTY(EditAnywhere, Category = "Replication")
+	TArray<TObjectPtr<UMultiUserReplicationStreamAsset>> Streams;
+};
\ No newline at end of file
diff --git a/Engine/Plugins/Developer/Concert/ConcertApp/MultiUserClient/Source/MultiUserReplication/Public/MultiUserReplicationStreamAsset.h b/Engine/Plugins/Developer/Concert/ConcertApp/MultiUserClient/Source/MultiUserReplication/Public/MultiUserReplicationStreamAsset.h
new file mode 100644
--- /dev/null
+++ b/Engine/Plugins/Developer/Concert/ConcertApp/MultiUserClient/Source/MultiUserReplication/Public/MultiUserReplicationStreamAsset.h
@@ -0,0 +1,25 @@
+// Copyright Epic Games, Inc. All Rights Reserved.
+
+#pragma once
+
+#include "UObject/Object.h"
+#include "Replication/Data/ReplicationStreamDescription.h"
+#include "MultiUserPropertyReplicationSelection.h"
+#include "MultiUserReplicationStreamAsset.generated.h"
+
+
+/** Asset describing a stream of replicated data the client sends to the server.  */
+UCLASS()
+class MULTIUSERREPLICATION_API UMultiUserReplicationStreamAsset : public UObject
+{
+	GENERATED_BODY()
+public:
+
+	UPROPERTY(Instanced)
+	TObjectPtr<UMultiUserPropertyReplicationSelection> ReplicationList;
+
+	UMultiUserReplicationStreamAsset();
+
+	/** Generates a stream description based on this asset for sending to MU server. */
+	FReplicationStreamDescription GenerateDescription() const;
+};
\ No newline at end of file
diff --git a/Engine/Plugins/Developer/Concert/ConcertApp/MultiUserClient/Source/MultiUserReplicationEditor/MultiUserReplicationEditor.Build.cs b/Engine/Plugins/Developer/Concert/ConcertApp/MultiUserClient/Source/MultiUserReplicationEditor/MultiUserReplicationEditor.Build.cs
new file mode 100644
--- /dev/null
+++ b/Engine/Plugins/Developer/Concert/ConcertApp/MultiUserClient/Source/MultiUserReplicationEditor/MultiUserReplicationEditor.Build.cs
@@ -0,0 +1,45 @@
+// Copyright Epic Games, Inc. All Rights Reserved.
+
+namespace UnrealBuildTool.Rules
+{
+	public class MultiUserReplicationEditor : ModuleRules
+	{
+		public MultiUserReplicationEditor(ReadOnlyTargetRules Target) : base(Target)
+		{
+			PublicDependencyModuleNames.AddRange(
+				new string[]
+				{
+					// Engine
+					"Core",
+					"CoreUObject",
+					"Engine",
+					
+					// Concert
+					"ConcertSyncCore",
+					"MultiUserReplication"
+				}
+			);
+			
+			PrivateDependencyModuleNames.AddRange(
+				new string[]
+				{
+					// Engine
+					"AssetDefinition",
+					"AssetRegistry",
+					"InputCore",
+					"Projects",
+					"Slate",
+					"SlateCore",
+					"ToolWidgets",
+					"UnrealEd", 
+					
+					// Concert
+					"Concert",
+					"ConcertClient",
+					"ConcertSharedSlate",
+					"MultiUserClient"
+				}
+			);
+		}
+	}
+}
diff --git a/Engine/Plugins/Developer/Concert/ConcertApp/MultiUserClient/Source/MultiUserReplicationEditor/Private/Assets/ReplicationClientProfileAssetDefinition.cpp b/Engine/Plugins/Developer/Concert/ConcertApp/MultiUserClient/Source/MultiUserReplicationEditor/Private/Assets/ReplicationClientProfileAssetDefinition.cpp
new file mode 100644
--- /dev/null
+++ b/Engine/Plugins/Developer/Concert/ConcertApp/MultiUserClient/Source/MultiUserReplicationEditor/Private/Assets/ReplicationClientProfileAssetDefinition.cpp
@@ -0,0 +1,31 @@
+// Copyright Epic Games, Inc. All Rights Reserved.
+
+#include "ReplicationClientProfileAssetDefinition.h"
+
+#include "IMultiUserReplicationEditorModule.h"
+#include "MultiUserReplicationClientProfileAsset.h"
+
+#define LOCTEXT_NAMESPACE "UReplicationClientProfileAssetDefinition"
+
+FText UReplicationClientProfileAssetDefinition::GetAssetDisplayName() const
+{
+	return LOCTEXT("AssetName", "Replication Client Profile");
+}
+
+FLinearColor UReplicationClientProfileAssetDefinition::GetAssetColor() const
+{
+	return FLinearColor(1.f, 0.8f, 0.2f);
+}
+
+TSoftClassPtr<UObject> UReplicationClientProfileAssetDefinition::GetAssetClass() const
+{
+	return UMultiUserReplicationClientProfileAsset::StaticClass();
+}
+
+TConstArrayView<FAssetCategoryPath> UReplicationClientProfileAssetDefinition::GetAssetCategories() const
+{
+	static const TArray<FAssetCategoryPath> Categories = { UE::MultiUserReplicationEditor::IMultiUserReplicationEditorModule::Get().GetMultiUserReplicationCategory() };
+	return Categories;
+}
+
+#undef LOCTEXT_NAMESPACE
diff --git a/Engine/Plugins/Developer/Concert/ConcertApp/MultiUserClient/Source/MultiUserReplicationEditor/Private/Assets/ReplicationClientProfileAssetDefinition.h b/Engine/Plugins/Developer/Concert/ConcertApp/MultiUserClient/Source/MultiUserReplicationEditor/Private/Assets/ReplicationClientProfileAssetDefinition.h
new file mode 100644
--- /dev/null
+++ b/Engine/Plugins/Developer/Concert/ConcertApp/MultiUserClient/Source/MultiUserReplicationEditor/Private/Assets/ReplicationClientProfileAssetDefinition.h
@@ -0,0 +1,23 @@
+// Copyright Epic Games, Inc. All Rights Reserved.
+
+#pragma once
+
+#include "AssetDefinitionDefault.h"
+#include "ReplicationClientProfileAssetDefinition.generated.h"
+
+/**
+ * 
+ */
+UCLASS()
+class MULTIUSERREPLICATIONEDITOR_API UReplicationClientProfileAssetDefinition : public UAssetDefinitionDefault
+{
+	GENERATED_BODY()
+public:
+
+	//~ Begin UAssetDefinition Interface
+	virtual FText GetAssetDisplayName() const override; 
+	virtual FLinearColor GetAssetColor() const override;
+	virtual TSoftClassPtr<UObject> GetAssetClass() const override; 
+	virtual TConstArrayView<FAssetCategoryPath> GetAssetCategories() const override;
+	//~ End UAssetDefinition Interface
+};
diff --git a/Engine/Plugins/Developer/Concert/ConcertApp/MultiUserClient/Source/MultiUserReplicationEditor/Private/Assets/ReplicationClientProfileAssetFactory.cpp b/Engine/Plugins/Developer/Concert/ConcertApp/MultiUserClient/Source/MultiUserReplicationEditor/Private/Assets/ReplicationClientProfileAssetFactory.cpp
new file mode 100644
--- /dev/null
+++ b/Engine/Plugins/Developer/Concert/ConcertApp/MultiUserClient/Source/MultiUserReplicationEditor/Private/Assets/ReplicationClientProfileAssetFactory.cpp
@@ -0,0 +1,39 @@
+// Copyright Epic Games, Inc. All Rights Reserved.
+
+#include "ReplicationClientProfileAssetFactory.h"
+
+#include "CVarMultiUserReplicationEditor.h"
+#include "MultiUserReplicationClientProfileAsset.h"
+
+#define LOCTEXT_NAMESPACE "UReplicationClientProfileAssetFactory"
+
+UReplicationClientProfileAssetFactory::UReplicationClientProfileAssetFactory(const FObjectInitializer& ObjectInitializer)
+ : Super(ObjectInitializer)
+{
+	bCreateNew = true;
+	bEditAfterNew = true;
+	SupportedClass = UMultiUserReplicationClientProfileAsset::StaticClass();
+}
+
+bool UReplicationClientProfileAssetFactory::CanCreateNew() const
+{
+	return UE::MultiUserReplicationEditor::ConsoleVariables::CVarEnableReplication.GetValueOnAnyThread();
+}
+
+FText UReplicationClientProfileAssetFactory::GetDisplayName() const
+{
+	return LOCTEXT("DisplayName", "Client Profile");
+}
+
+bool UReplicationClientProfileAssetFactory::ConfigureProperties()
+{
+	return true;
+}
+
+UObject* UReplicationClientProfileAssetFactory::FactoryCreateNew(UClass* Class, UObject* InParent, FName Name, EObjectFlags Flags, UObject* Context, FFeedbackContext* Warn)
+{
+	UMultiUserReplicationClientProfileAsset* Asset = NewObject<UMultiUserReplicationClientProfileAsset>(InParent, Name, Flags);
+	return Asset;
+}
+
+#undef LOCTEXT_NAMESPACE
\ No newline at end of file
diff --git a/Engine/Plugins/Developer/Concert/ConcertApp/MultiUserClient/Source/MultiUserReplicationEditor/Private/Assets/ReplicationClientProfileAssetFactory.h b/Engine/Plugins/Developer/Concert/ConcertApp/MultiUserClient/Source/MultiUserReplicationEditor/Private/Assets/ReplicationClientProfileAssetFactory.h
new file mode 100644
--- /dev/null
+++ b/Engine/Plugins/Developer/Concert/ConcertApp/MultiUserClient/Source/MultiUserReplicationEditor/Private/Assets/ReplicationClientProfileAssetFactory.h
@@ -0,0 +1,23 @@
+// Copyright Epic Games, Inc. All Rights Reserved.
+
+#pragma once
+
+#include "UObject/ObjectMacros.h"
+#include "Factories/Factory.h"
+#include "ReplicationClientProfileAssetFactory.generated.h"
+
+UCLASS(hidecategories=Object, MinimalAPI)
+class UReplicationClientProfileAssetFactory : public UFactory
+{
+	GENERATED_UCLASS_BODY()
+
+	//~ Begin UFactory Interface
+	virtual bool CanCreateNew() const override;
+	virtual FText GetDisplayName() const override;
+	virtual bool ConfigureProperties() override;
+	virtual UObject* FactoryCreateNew(UClass* Class,UObject* InParent,FName Name,EObjectFlags Flags,UObject* Context,FFeedbackContext* Warn) override;
+	//~ Begin UFactory Interface
+};
+
+
+
diff --git a/Engine/Plugins/Developer/Concert/ConcertApp/MultiUserClient/Source/MultiUserReplicationEditor/Private/Assets/ReplicationStreamAssetDefinition.cpp b/Engine/Plugins/Developer/Concert/ConcertApp/MultiUserClient/Source/MultiUserReplicationEditor/Private/Assets/ReplicationStreamAssetDefinition.cpp
new file mode 100644
--- /dev/null
+++ b/Engine/Plugins/Developer/Concert/ConcertApp/MultiUserClient/Source/MultiUserReplicationEditor/Private/Assets/ReplicationStreamAssetDefinition.cpp
@@ -0,0 +1,55 @@
+// Copyright Epic Games, Inc. All Rights Reserved.
+
+#include "ReplicationStreamAssetDefinition.h"
+
+#include "StreamEditor/ReplicationStreamAssetEditor.h"
+#include "IMultiUserReplicationEditorModule.h"
+#include "MultiUserReplicationStreamAsset.h"
+
+#define LOCTEXT_NAMESPACE "UReplicationStreamAssetDefinition"
+
+FText UReplicationStreamAssetDefinition::GetAssetDisplayName() const
+{
+	return LOCTEXT("AssetName", "Replication Stream");
+}
+
+FLinearColor UReplicationStreamAssetDefinition::GetAssetColor() const
+{
+	return FLinearColor(4.f, 0.4f, 0.65f);
+}
+
+TSoftClassPtr<UObject> UReplicationStreamAssetDefinition::GetAssetClass() const
+{
+	return UMultiUserReplicationStreamAsset::StaticClass();
+}
+
+TConstArrayView<FAssetCategoryPath> UReplicationStreamAssetDefinition::GetAssetCategories() const
+{
+	static const TArray<FAssetCategoryPath> Categories = { UE::MultiUserReplicationEditor::IMultiUserReplicationEditorModule::Get().GetMultiUserReplicationCategory() };
+	return Categories;
+}
+
+EAssetCommandResult UReplicationStreamAssetDefinition::OpenAssets(const FAssetOpenArgs& OpenArgs) const
+{
+	UAssetEditorSubsystem* AssetEditorSubsystem = GEditor->GetEditorSubsystem<UAssetEditorSubsystem>();
+	if (!AssetEditorSubsystem)
+	{
+		return EAssetCommandResult::Handled;
+	}
+
+	for (UMultiUserReplicationStreamAsset* ClientProfile : OpenArgs.LoadObjects<UMultiUserReplicationStreamAsset>())
+	{
+		UReplicationStreamAssetEditor* AssetEditor = NewObject<UReplicationStreamAssetEditor>(AssetEditorSubsystem, NAME_None, RF_Transient);
+		if (!AssetEditor)
+		{
+			continue;
+		}
+
+		AssetEditor->SetObjectToEdit(ClientProfile);
+		AssetEditor->Initialize();
+	}
+
+	return EAssetCommandResult::Handled;
+}
+
+#undef LOCTEXT_NAMESPACE
diff --git a/Engine/Plugins/Developer/Concert/ConcertApp/MultiUserClient/Source/MultiUserReplicationEditor/Private/Assets/ReplicationStreamAssetDefinition.h b/Engine/Plugins/Developer/Concert/ConcertApp/MultiUserClient/Source/MultiUserReplicationEditor/Private/Assets/ReplicationStreamAssetDefinition.h
new file mode 100644
--- /dev/null
+++ b/Engine/Plugins/Developer/Concert/ConcertApp/MultiUserClient/Source/MultiUserReplicationEditor/Private/Assets/ReplicationStreamAssetDefinition.h
@@ -0,0 +1,24 @@
+// Copyright Epic Games, Inc. All Rights Reserved.
+
+#pragma once
+
+#include "AssetDefinitionDefault.h"
+#include "ReplicationStreamAssetDefinition.generated.h"
+
+/**
+ * 
+ */
+UCLASS()
+class MULTIUSERREPLICATIONEDITOR_API UReplicationStreamAssetDefinition : public UAssetDefinitionDefault
+{
+	GENERATED_BODY()
+public:
+	
+	//~ Begin UAssetDefinition Interface
+	virtual FText GetAssetDisplayName() const override; 
+	virtual FLinearColor GetAssetColor() const override;
+	virtual TSoftClassPtr<UObject> GetAssetClass() const override; 
+	virtual TConstArrayView<FAssetCategoryPath> GetAssetCategories() const override;
+	EAssetCommandResult OpenAssets(const FAssetOpenArgs& OpenArgs) const override;
+	//~ End UAssetDefinition Interface
+};
diff --git a/Engine/Plugins/Developer/Concert/ConcertApp/MultiUserClient/Source/MultiUserReplicationEditor/Private/Assets/ReplicationStreamAssetFactory.cpp b/Engine/Plugins/Developer/Concert/ConcertApp/MultiUserClient/Source/MultiUserReplicationEditor/Private/Assets/ReplicationStreamAssetFactory.cpp
new file mode 100644
--- /dev/null
+++ b/Engine/Plugins/Developer/Concert/ConcertApp/MultiUserClient/Source/MultiUserReplicationEditor/Private/Assets/ReplicationStreamAssetFactory.cpp
@@ -0,0 +1,39 @@
+// Copyright Epic Games, Inc. All Rights Reserved.
+
+#include "ReplicationStreamAssetFactory.h"
+
+#include "CVarMultiUserReplicationEditor.h"
+#include "MultiUserReplicationStreamAsset.h"
+
+#define LOCTEXT_NAMESPACE "ReplicationStreamAssetFactory"
+
+UReplicationStreamAssetFactory::UReplicationStreamAssetFactory(const FObjectInitializer& ObjectInitializer)
+ : Super(ObjectInitializer)
+{
+	bCreateNew = true;
+	bEditAfterNew = true;
+	SupportedClass = UMultiUserReplicationStreamAsset::StaticClass();
+}
+
+bool UReplicationStreamAssetFactory::CanCreateNew() const
+{
+	return UE::MultiUserReplicationEditor::ConsoleVariables::CVarEnableReplication.GetValueOnAnyThread();;
+}
+
+FText UReplicationStreamAssetFactory::GetDisplayName() const
+{
+	return LOCTEXT("DisplayName", "Replication Stream");
+}
+
+bool UReplicationStreamAssetFactory::ConfigureProperties()
+{
+	return true;
+}
+
+UObject* UReplicationStreamAssetFactory::FactoryCreateNew(UClass* Class, UObject* InParent, FName Name, EObjectFlags Flags, UObject* Context, FFeedbackContext* Warn)
+{
+	UMultiUserReplicationStreamAsset* Asset = NewObject<UMultiUserReplicationStreamAsset>(InParent, Name, Flags);
+	return Asset;
+}
+
+#undef LOCTEXT_NAMESPACE
\ No newline at end of file
diff --git a/Engine/Plugins/Developer/Concert/ConcertApp/MultiUserClient/Source/MultiUserReplicationEditor/Private/Assets/ReplicationStreamAssetFactory.h b/Engine/Plugins/Developer/Concert/ConcertApp/MultiUserClient/Source/MultiUserReplicationEditor/Private/Assets/ReplicationStreamAssetFactory.h
new file mode 100644
--- /dev/null
+++ b/Engine/Plugins/Developer/Concert/ConcertApp/MultiUserClient/Source/MultiUserReplicationEditor/Private/Assets/ReplicationStreamAssetFactory.h
@@ -0,0 +1,23 @@
+// Copyright Epic Games, Inc. All Rights Reserved.
+
+#pragma once
+
+#include "UObject/ObjectMacros.h"
+#include "Factories/Factory.h"
+#include "ReplicationStreamAssetFactory.generated.h"
+
+UCLASS(hidecategories=Object, MinimalAPI)
+class UReplicationStreamAssetFactory : public UFactory
+{
+	GENERATED_UCLASS_BODY()
+
+	//~ Begin UFactory Interface
+	virtual bool CanCreateNew() const override;
+	virtual FText GetDisplayName() const override;
+	virtual bool ConfigureProperties() override;
+	virtual UObject* FactoryCreateNew(UClass* Class,UObject* InParent,FName Name,EObjectFlags Flags,UObject* Context,FFeedbackContext* Warn) override;
+	//~ Begin UFactory Interface
+};
+
+
+
diff --git a/Engine/Plugins/Developer/Concert/ConcertApp/MultiUserClient/Source/MultiUserReplicationEditor/Private/CVarMultiUserReplicationEditor.cpp b/Engine/Plugins/Developer/Concert/ConcertApp/MultiUserClient/Source/MultiUserReplicationEditor/Private/CVarMultiUserReplicationEditor.cpp
new file mode 100644
--- /dev/null
+++ b/Engine/Plugins/Developer/Concert/ConcertApp/MultiUserClient/Source/MultiUserReplicationEditor/Private/CVarMultiUserReplicationEditor.cpp
@@ -0,0 +1,13 @@
+// Copyright Epic Games, Inc. All Rights Reserved.
+
+#include "CVarMultiUserReplicationEditor.h"
+
+namespace UE::MultiUserReplicationEditor::ConsoleVariables
+{
+	TAutoConsoleVariable<bool> CVarEnableReplication(
+		TEXT("Concert.EnableReplication"),
+		false,
+		TEXT("Controls whether the replication feature is enabled."),
+		ECVF_Default
+		);
+}
\ No newline at end of file
diff --git a/Engine/Plugins/Developer/Concert/ConcertApp/MultiUserClient/Source/MultiUserReplicationEditor/Private/CVarMultiUserReplicationEditor.h b/Engine/Plugins/Developer/Concert/ConcertApp/MultiUserClient/Source/MultiUserReplicationEditor/Private/CVarMultiUserReplicationEditor.h
new file mode 100644
--- /dev/null
+++ b/Engine/Plugins/Developer/Concert/ConcertApp/MultiUserClient/Source/MultiUserReplicationEditor/Private/CVarMultiUserReplicationEditor.h
@@ -0,0 +1,11 @@
+// Copyright Epic Games, Inc. All Rights Reserved.
+
+#pragma once
+
+#include "HAL/IConsoleManager.h"
+
+namespace UE::MultiUserReplicationEditor::ConsoleVariables
+{
+	/** This feature is highly experimental in 5.3. */
+	extern TAutoConsoleVariable<bool> CVarEnableReplication;
+}
diff --git a/Engine/Plugins/Developer/Concert/ConcertApp/MultiUserClient/Source/MultiUserReplicationEditor/Private/LogMultiUserReplicationEditor.cpp b/Engine/Plugins/Developer/Concert/ConcertApp/MultiUserClient/Source/MultiUserReplicationEditor/Private/LogMultiUserReplicationEditor.cpp
new file mode 100644
--- /dev/null
+++ b/Engine/Plugins/Developer/Concert/ConcertApp/MultiUserClient/Source/MultiUserReplicationEditor/Private/LogMultiUserReplicationEditor.cpp
@@ -0,0 +1,5 @@
+// Copyright Epic Games, Inc. All Rights Reserved.
+
+#include "LogMultiUserReplicationEditor.h"
+
+DEFINE_LOG_CATEGORY(LogMultiUserReplicationEditor);
\ No newline at end of file
diff --git a/Engine/Plugins/Developer/Concert/ConcertApp/MultiUserClient/Source/MultiUserReplicationEditor/Private/LogMultiUserReplicationEditor.h b/Engine/Plugins/Developer/Concert/ConcertApp/MultiUserClient/Source/MultiUserReplicationEditor/Private/LogMultiUserReplicationEditor.h
new file mode 100644
--- /dev/null
+++ b/Engine/Plugins/Developer/Concert/ConcertApp/MultiUserClient/Source/MultiUserReplicationEditor/Private/LogMultiUserReplicationEditor.h
@@ -0,0 +1,7 @@
+// Copyright Epic Games, Inc. All Rights Reserved.
+
+#pragma once
+
+#include "Logging/LogMacros.h"
+
+DECLARE_LOG_CATEGORY_EXTERN(LogMultiUserReplicationEditor, All, All);
diff --git a/Engine/Plugins/Developer/Concert/ConcertApp/MultiUserClient/Source/MultiUserReplicationEditor/Private/MultiUserReplicationEditorModule.cpp b/Engine/Plugins/Developer/Concert/ConcertApp/MultiUserClient/Source/MultiUserReplicationEditor/Private/MultiUserReplicationEditorModule.cpp
new file mode 100644
--- /dev/null
+++ b/Engine/Plugins/Developer/Concert/ConcertApp/MultiUserClient/Source/MultiUserReplicationEditor/Private/MultiUserReplicationEditorModule.cpp
@@ -0,0 +1,68 @@
+// Copyright Epic Games, Inc. All Rights Reserved.
+
+#include "MultiUserReplicationEditorModule.h"
+
+#include "ConcertSyncSessionFlags.h"
+#include "CVarMultiUserReplicationEditor.h"
+#include "IMultiUserClientModule.h"
+
+#include "ISettingsModule.h"
+#include "MultiUserReplicationEditorStyle.h"
+#include "Settings/MultiUserReplicationSettings.h"
+#include "Modules/ModuleManager.h"
+
+#define LOCTEXT_NAMESPACE "FMultiUserReplicationEditorModule"
+
+namespace UE::MultiUserReplicationEditor
+{
+	FMultiUserReplicationEditorModule::FMultiUserReplicationEditorModule()
+		: MultiUserReplicationCategory(LOCTEXT("MultiUserReplicationCategory", "Multi-User"))
+	{}
+
+	void FMultiUserReplicationEditorModule::StartupModule()
+	{
+		FMultiUserReplicationEditorStyle::Initialize();
+
+		UpdateSettingsRegistrationBasedOnCVar();
+		ConsoleVariables::CVarEnableReplication->OnChangedDelegate().AddLambda([this](IConsoleVariable*)
+		{
+			UpdateSettingsRegistrationBasedOnCVar();
+		});
+	}
+
+	void FMultiUserReplicationEditorModule::ShutdownModule()
+	{
+		FMultiUserReplicationEditorStyle::Shutdown();
+	}
+
+	void FMultiUserReplicationEditorModule::UpdateSettingsRegistrationBasedOnCVar()
+	{
+		if (ConsoleVariables::CVarEnableReplication.GetValueOnAnyThread())
+		{
+			RegisterSettings();
+		}
+		else
+		{
+			UnregisterSettings();
+		}
+	}
+
+	void FMultiUserReplicationEditorModule::RegisterSettings()
+	{
+		ISettingsModule* SettingsModule = FModuleManager::GetModulePtr<ISettingsModule>("Settings");
+		SettingsModule->RegisterSettings("Project", "Plugins", "Multi-User Replication",
+			LOCTEXT("MultiUserReplicationSettingsName", "Multi-User Replication"),
+			LOCTEXT("MultiUserReplicationSettingsDescription", "Configure the Multi-User Replication settings."),
+			UMultiUserReplicationSettings::Get()
+		);
+	}
+
+	void FMultiUserReplicationEditorModule::UnregisterSettings()
+	{
+		ISettingsModule* SettingsModule = FModuleManager::GetModulePtr<ISettingsModule>("Settings");
+		SettingsModule->UnregisterSettings("Project", "Plugins", "Multi-User Replication");
+	}
+};
+
+IMPLEMENT_MODULE(UE::MultiUserReplicationEditor::FMultiUserReplicationEditorModule, MultiUserReplicationEditor);
+#undef LOCTEXT_NAMESPACE
\ No newline at end of file
diff --git a/Engine/Plugins/Developer/Concert/ConcertApp/MultiUserClient/Source/MultiUserReplicationEditor/Private/MultiUserReplicationEditorModule.h b/Engine/Plugins/Developer/Concert/ConcertApp/MultiUserClient/Source/MultiUserReplicationEditor/Private/MultiUserReplicationEditorModule.h
new file mode 100644
--- /dev/null
+++ b/Engine/Plugins/Developer/Concert/ConcertApp/MultiUserClient/Source/MultiUserReplicationEditor/Private/MultiUserReplicationEditorModule.h
@@ -0,0 +1,31 @@
+// Copyright Epic Games, Inc. All Rights Reserved.
+
+#include "IMultiUserReplicationEditorModule.h"
+#include "Misc/AssetCategoryPath.h"
+
+namespace UE::MultiUserReplicationEditor
+{
+	class FMultiUserReplicationEditorModule : public IMultiUserReplicationEditorModule
+	{
+	public:
+
+		FMultiUserReplicationEditorModule();
+
+		//~ Begin IModuleInterface Interface
+		virtual void StartupModule() override;
+		virtual void ShutdownModule() override;
+		//~ End IModuleInterface Interface
+
+		//~ Begin IMultiUserReplicationEditor Interface
+		virtual FAssetCategoryPath GetMultiUserReplicationCategory() const override { return MultiUserReplicationCategory; }
+		//~ End IMultiUserReplicationEditor Interface
+
+	private:
+
+		FAssetCategoryPath MultiUserReplicationCategory;
+
+		void UpdateSettingsRegistrationBasedOnCVar();
+		void RegisterSettings();
+		void UnregisterSettings();
+	};
+}
diff --git a/Engine/Plugins/Developer/Concert/ConcertApp/MultiUserClient/Source/MultiUserReplicationEditor/Private/MultiUserReplicationEditorStyle.cpp b/Engine/Plugins/Developer/Concert/ConcertApp/MultiUserClient/Source/MultiUserReplicationEditor/Private/MultiUserReplicationEditorStyle.cpp
new file mode 100644
--- /dev/null
+++ b/Engine/Plugins/Developer/Concert/ConcertApp/MultiUserClient/Source/MultiUserReplicationEditor/Private/MultiUserReplicationEditorStyle.cpp
@@ -0,0 +1,59 @@
+// Copyright Epic Games, Inc. All Rights Reserved.
+
+#include "MultiUserReplicationEditorStyle.h"
+
+#include "Interfaces/IPluginManager.h"
+#include "Styling/SlateStyleRegistry.h"
+
+namespace UE::MultiUserReplicationEditor
+{
+	TSharedPtr<FSlateStyleSet> FMultiUserReplicationEditorStyle::StyleInstance = nullptr;
+
+	void FMultiUserReplicationEditorStyle::Initialize()
+	{
+		if (!StyleInstance.IsValid())
+		{
+			StyleInstance = Create();
+			FSlateStyleRegistry::RegisterSlateStyle(*StyleInstance);
+		}
+	}
+
+	void FMultiUserReplicationEditorStyle::Shutdown()
+	{
+		FSlateStyleRegistry::UnRegisterSlateStyle(*StyleInstance);
+		ensure(StyleInstance.IsUnique());
+		StyleInstance.Reset();
+	}
+
+	const ISlateStyle& FMultiUserReplicationEditorStyle::Get()
+	{
+		return *StyleInstance;
+	}
+
+	FName FMultiUserReplicationEditorStyle::GetStyleSetName()
+	{
+		static FName StyleSetName(TEXT("MultiUserReplicationEditor"));
+		return StyleSetName;
+	}
+	
+	TSharedRef< FSlateStyleSet > FMultiUserReplicationEditorStyle::Create()
+	{
+		TSharedRef<FSlateStyleSet> Style = MakeShared<FSlateStyleSet>("MultiUserReplicationEditor");
+
+		TSharedPtr<IPlugin> Plugin = IPluginManager::Get().FindPlugin(TEXT("MultiUserClient"));
+		if (ensure(Plugin.IsValid()))
+		{
+			Style->SetContentRoot(FPaths::Combine(Plugin->GetBaseDir(), TEXT("Resources")));
+		}
+
+		Style->Set("ReplicationTreeView.RowHeight", 20.f);
+		Style->Set("ReplicationTreeView.IconColumnWidth", 20.f);
+		
+		return Style;
+	}
+}
+
+#undef IMAGE_BRUSH
+#undef BOX_BRUSH
+#undef BORDER_BRUSH
+#undef IMAGE_BRUSH_SVG
diff --git a/Engine/Plugins/Developer/Concert/ConcertApp/MultiUserClient/Source/MultiUserReplicationEditor/Private/MultiUserReplicationEditorStyle.h b/Engine/Plugins/Developer/Concert/ConcertApp/MultiUserClient/Source/MultiUserReplicationEditor/Private/MultiUserReplicationEditorStyle.h
new file mode 100644
--- /dev/null
+++ b/Engine/Plugins/Developer/Concert/ConcertApp/MultiUserClient/Source/MultiUserReplicationEditor/Private/MultiUserReplicationEditorStyle.h
@@ -0,0 +1,29 @@
+// Copyright Epic Games, Inc. All Rights Reserved.
+
+#pragma once
+
+#include "Styling/ISlateStyle.h"
+#include "Styling/SlateStyle.h"
+#include "Templates/SharedPointer.h"
+#include "UObject/NameTypes.h"
+
+namespace UE::MultiUserReplicationEditor
+{
+	class FMultiUserReplicationEditorStyle
+	{
+	public:
+	
+		static void Initialize();
+		static void Shutdown();
+
+		
+		static const ISlateStyle& Get();
+		static FName GetStyleSetName();
+
+	private:
+	
+		static TSharedPtr<FSlateStyleSet> StyleInstance;
+		static TSharedRef<FSlateStyleSet> Create();
+	};
+}
+
diff --git a/Engine/Plugins/Developer/Concert/ConcertApp/MultiUserClient/Source/MultiUserReplicationEditor/Private/Settings/MultiUserReplicationSettings.cpp b/Engine/Plugins/Developer/Concert/ConcertApp/MultiUserClient/Source/MultiUserReplicationEditor/Private/Settings/MultiUserReplicationSettings.cpp
new file mode 100644
--- /dev/null
+++ b/Engine/Plugins/Developer/Concert/ConcertApp/MultiUserClient/Source/MultiUserReplicationEditor/Private/Settings/MultiUserReplicationSettings.cpp
@@ -0,0 +1,74 @@
+// Copyright Epic Games, Inc. All Rights Reserved.
+
+#include "Settings/MultiUserReplicationSettings.h"
+
+#include "LogMultiUserReplicationEditor.h"
+#include "Replication/PropertyChainUtils.h"
+#include "Settings/DefaultPropertySelection.h"
+
+#include "Algo/IndexOf.h"
+
+namespace UE::MultiUserReplicationEditor::Private
+{
+	static void ApplyDefaultPropertySelection(FReplicatedObjectInfo& Info, const FDefaultPropertySelection& Selection, UStruct& Class)
+	{
+		// Preparse the the FStrings into paths
+		TArray<TArray<FName>> Paths;
+		Paths.Reserve(Selection.DefaultSelectedProperties.Num());
+		for (const FString& StringPath : Selection.DefaultSelectedProperties)
+		{
+			TArray<FString> Parts;
+			StringPath.ParseIntoArray(Parts, TEXT("."));
+			TArray<FName> NameArray;
+			Algo::Transform(Parts, NameArray, [](const FString& String){ return *String; });
+			Paths.Emplace(MoveTemp(NameArray));
+		}
+
+		// This walks through the entire property hierarchy once
+		using namespace ConcertSyncCore::PropertyChain;
+		BulkConstructConcertChainsFromPaths(Class, Paths.Num(), [&Info, &Paths](const FArchiveSerializedPropertyChain& Chain, const FProperty& LeafProperty)
+		{
+			const int32 IndexOfMatches = Algo::IndexOfByPredicate(Paths, [&Chain, &LeafProperty](const TArray<FName>& Path){ return DoPathAndChainsMatch(Path, Chain, LeafProperty); });
+			const bool bMatches = IndexOfMatches != INDEX_NONE;
+			if (bMatches)
+			{
+				Info.ReplicatedProperties.ReplicatedProperties.Emplace(&Chain, LeafProperty);
+				Paths.RemoveAtSwap(IndexOfMatches);
+			}
+			return true;
+		});
+
+		// Warn the user of unmatched properties so they can de-clutter their settings of useless entries
+		for (const TArray<FName>& Path : Paths)
+		{
+			UE_LOG(LogMultiUserReplicationEditor, Warning,
+				TEXT("Path %s was unmatched. Please fix your settings under MultiUserReplicationEditor > DefaultPropertySelection > %s"),
+				*FString::JoinBy(Path, TEXT("."), [](const FName& Name){ return Name.ToString(); }),
+				*Class.GetName()
+				);
+		}
+	}
+}
+
+void UMultiUserReplicationSettings::AddDefaultPropertiesFromSettings(FReplicatedObjectInfo& Info, UClass& Class)
+{
+	// Find the most specialized class properties
+	UClass* Current = &Class;
+	FDefaultPropertySelection* DefaultProperties = nullptr;
+	for (; Current && !DefaultProperties; Current = Current->GetSuperClass())
+	{
+		DefaultProperties = DefaultPropertySelection.Find(Current);
+		if (!DefaultProperties)
+		{
+			continue;
+		}
+			
+		UE::MultiUserReplicationEditor::Private::ApplyDefaultPropertySelection(Info, *DefaultProperties, *Current);
+		// Recurse super structs
+		if (UClass* Parent = Current->GetSuperClass()
+			; Parent && DefaultProperties->bInheritFromBase)
+		{
+			AddDefaultPropertiesFromSettings(Info, *Parent);
+		}
+	}
+}
diff --git a/Engine/Plugins/Developer/Concert/ConcertApp/MultiUserClient/Source/MultiUserReplicationEditor/Private/StreamEditor/Model/DisplayUtils.cpp b/Engine/Plugins/Developer/Concert/ConcertApp/MultiUserClient/Source/MultiUserReplicationEditor/Private/StreamEditor/Model/DisplayUtils.cpp
new file mode 100644
--- /dev/null
+++ b/Engine/Plugins/Developer/Concert/ConcertApp/MultiUserClient/Source/MultiUserReplicationEditor/Private/StreamEditor/Model/DisplayUtils.cpp
@@ -0,0 +1,90 @@
+// Copyright Epic Games, Inc. All Rights Reserved.
+
+#include "DisplayUtils.h"
+
+#include "LogMultiUserReplicationEditor.h"
+
+#include "AssetRegistry/AssetData.h"
+#include "AssetRegistry/AssetRegistryModule.h"
+#include "Blueprint/BlueprintSupport.h"
+#include "IObjectToPropertiesModel.h"
+#include "GameFramework/Actor.h"
+#include "Styling/SlateIconFinder.h"
+
+namespace UE::MultiUserReplicationEditor::DisplayUtils
+{
+	FText GetObjectDisplayText(const FSoftObjectPath& Object)
+	{
+		// Important! The object may not be loaded, yet. This could be if the asset is using a level that was not opened. 
+		UObject* LoadedObject = Object.TryLoad();
+		return LoadedObject
+			? FText::FromString(GetObjectDisplayString(*LoadedObject))
+			: FText::GetEmpty();
+	}
+
+	FString GetObjectDisplayString(UObject& Object)
+	{
+		if (const AActor* AsActor = Cast<AActor>(&Object))
+		{
+			return AsActor->GetActorLabel();
+		}
+		return Object.GetName();
+	}
+
+	FText GetObjectTypeText(const IObjectToPropertiesModel& Model, const FSoftObjectPath& Object)
+	{
+		if (const FSoftClassPath ClassPath = Model.GetObjectClass(Object); ClassPath.IsValid())
+		{
+			// For C++ classes or loaded Blueprints
+			if (UClass* Class = ClassPath.ResolveClass())
+			{
+				return Class->GetDisplayNameText();
+			}
+
+			// Assuming it is a Blueprint, try to get its display name from the registry tags
+			const FAssetRegistryModule& AssetRegistryModule = FModuleManager::LoadModuleChecked<FAssetRegistryModule>(TEXT("AssetRegistry"));
+			const IAssetRegistry& AssetRegistry = AssetRegistryModule.Get();
+			
+			const FAssetData BlueprintData = AssetRegistry.GetAssetByObjectPath(ClassPath);
+			if (BlueprintData.IsValid())
+			{
+				FString DisplayName;
+				return BlueprintData.GetTagValue(FBlueprintTags::BlueprintDisplayName, DisplayName)
+					? FText::FromString(DisplayName)
+					: FText::GetEmpty();
+			}
+
+			// TODO: Check redirectors
+			// It was not a Blueprint - the class was most likely removed
+			UE_LOG(LogMultiUserReplicationEditor, Error, TEXT("Failed to get display name for unknown class \"%s\""), *ClassPath.ToString());
+		}
+		return FText::GetEmpty();
+	}
+
+	FSlateIcon GetObjectIcon(const IObjectToPropertiesModel& Model, const FSoftObjectPath& Object)
+	{
+		if (const FSoftClassPath ClassPath = Model.GetObjectClass(Object); ClassPath.IsValid())
+		{
+			UClass* Class = ClassPath.ResolveClass();
+			return Class
+				? FSlateIconFinder::FindIconForClass(Class)
+				: FSlateIcon();
+		}
+		return FSlateIcon();
+	}
+
+	FSlateIcon GetObjectIcon(UObject& Object)
+	{
+		return FSlateIconFinder::FindIconForClass(Object.GetClass());
+	}
+	
+	FText GetPropertyDisplayText(const FConcertPropertyChain& Property)
+	{
+		return FText::FromString(GetPropertyDisplayString(Property));
+    }
+	
+	FString GetPropertyDisplayString(const FConcertPropertyChain& Property)
+	{
+		return Property.ToString(FConcertPropertyChain::EToStringMethod::LeafProperty);
+	}
+}
diff --git a/Engine/Plugins/Developer/Concert/ConcertApp/MultiUserClient/Source/MultiUserReplicationEditor/Private/StreamEditor/Model/DisplayUtils.h b/Engine/Plugins/Developer/Concert/ConcertApp/MultiUserClient/Source/MultiUserReplicationEditor/Private/StreamEditor/Model/DisplayUtils.h
new file mode 100644
--- /dev/null
+++ b/Engine/Plugins/Developer/Concert/ConcertApp/MultiUserClient/Source/MultiUserReplicationEditor/Private/StreamEditor/Model/DisplayUtils.h
@@ -0,0 +1,37 @@
+// Copyright Epic Games, Inc. All Rights Reserved.
+
+#pragma once
+
+class FString;
+class FText;
+class UObject;
+
+struct FConcertPropertyChain;
+struct FSlateIcon;
+struct FSoftObjectPath;
+
+namespace UE::MultiUserReplicationEditor
+{
+	class IObjectToPropertiesModel;
+}
+
+namespace UE::MultiUserReplicationEditor::DisplayUtils
+{
+	/** @return The text to use for displaying this object's name */
+	FText GetObjectDisplayText(const FSoftObjectPath& Object);
+	/** @return More lightweight version of GetReplicatedObjectDisplayName which does not construct any FText. */
+	FString GetObjectDisplayString(UObject& Object);
+	
+	/** @return The text to use for displaying this object's type */
+	FText GetObjectTypeText(const IObjectToPropertiesModel& Model, const FSoftObjectPath& Object);
+	
+	/** @return The icon to use for this object */
+	FSlateIcon GetObjectIcon(const IObjectToPropertiesModel& Model, const FSoftObjectPath& Object);
+	/** @return The icon to use for this object */
+	FSlateIcon GetObjectIcon(UObject& Object);
+
+	/** @return The text to use for displaying this property's name. */
+	FText GetPropertyDisplayText(const FConcertPropertyChain& Property);
+	/** @return More lightweight version of GetPropertyDisplayString which does not construct any FText. */
+	FString GetPropertyDisplayString(const FConcertPropertyChain& Property);
+}
diff --git a/Engine/Plugins/Developer/Concert/ConcertApp/MultiUserClient/Source/MultiUserReplicationEditor/Private/StreamEditor/Model/IEditableObjectToPropertiesModel.h b/Engine/Plugins/Developer/Concert/ConcertApp/MultiUserClient/Source/MultiUserReplicationEditor/Private/StreamEditor/Model/IEditableObjectToPropertiesModel.h
new file mode 100644
--- /dev/null
+++ b/Engine/Plugins/Developer/Concert/ConcertApp/MultiUserClient/Source/MultiUserReplicationEditor/Private/StreamEditor/Model/IEditableObjectToPropertiesModel.h
@@ -0,0 +1,49 @@
+// Copyright Epic Games, Inc. All Rights Reserved.
+
+#pragma once
+
+#include "Containers/Map.h"
+#include "Containers/Set.h"
+#include "Delegates/DelegateCombinations.h"
+#include "IObjectToPropertiesModel.h"
+
+class UObject;
+struct FConcertPropertyChain;
+
+namespace UE::MultiUserReplicationEditor
+{
+	enum class EReplicatedObjectChangeReason : uint8
+	{
+		ChangedDirectly,
+		Transacted
+	};
+	
+	/**
+	 * Abstracts the concept of mapping objects to properties. This allows writing.
+	 * 
+	 * Models may not always be writable. When editing a UAsset, it will be writable. However, if we join a multi-user
+	 * session we do not want to edit the objects nor properties - only read.
+	 */
+	class IEditableObjectToPropertiesModel : public IObjectToPropertiesModel
+	{
+	public:
+
+		/** Adds the objects to the mapping */
+		virtual void AddObjects(TArrayView<UObject*> Objects) = 0;
+		/** Removes the objects from the mapping */
+		virtual void RemoveObjects(TArrayView<FSoftObjectPath> Objects) = 0;
+
+		/** Adds these properties to the object's list of selected properties. */
+		virtual void AddProperties(const FSoftObjectPath& Object, TArrayView<FConcertPropertyChain> Properties) = 0;
+		/** Removes these properties from the object's list of selected properties. */
+		virtual void RemoveProperties(const FSoftObjectPath& Object, TArrayView<FConcertPropertyChain> Properties) = 0;
+
+		/** Called when the object list changes. AddedObjects and RemovedObjects are empty if and only if ChangeReason == EReplicatedObjectChangeReason::Transacted. */
+		DECLARE_MULTICAST_DELEGATE_ThreeParams(FOnObjectsChanged, TArrayView<UObject*> AddedObjects, TArrayView<FSoftObjectPath> RemovedObjects, EReplicatedObjectChangeReason ChangeReason);
+		virtual FOnObjectsChanged& OnObjectsChanged() = 0;
+
+		/** Called when the property list of some object changes. If the properties are removed because the object is removed outright, FOnObjectsChanged is called instead.  */
+		DECLARE_MULTICAST_DELEGATE(FOnPropertiesChanged);
+		virtual FOnPropertiesChanged& OnPropertiesChanged() = 0;
+	};
+}
\ No newline at end of file
diff --git a/Engine/Plugins/Developer/Concert/ConcertApp/MultiUserClient/Source/MultiUserReplicationEditor/Private/StreamEditor/Model/IObjectToPropertiesModel.h b/Engine/Plugins/Developer/Concert/ConcertApp/MultiUserClient/Source/MultiUserReplicationEditor/Private/StreamEditor/Model/IObjectToPropertiesModel.h
new file mode 100644
--- /dev/null
+++ b/Engine/Plugins/Developer/Concert/ConcertApp/MultiUserClient/Source/MultiUserReplicationEditor/Private/StreamEditor/Model/IObjectToPropertiesModel.h
@@ -0,0 +1,95 @@
+// Copyright Epic Games, Inc. All Rights Reserved.
+
+#pragma once
+
+#include "Misc/EBreakBehavior.h"
+#include "Replication/Data/ConcertPropertySelection.h"
+#include "Templates/Function.h"
+#include "Templates/SharedPointer.h"
+
+namespace UE::MultiUserReplicationEditor
+{
+	/**
+	 * Abstracts the concept of mapping objects to properties. This is a read-only interface.
+	 * @see IEditableObjectToPropertiesModel 
+	 */
+	class IObjectToPropertiesModel : public TSharedFromThis<IObjectToPropertiesModel>
+	{
+	public:
+
+		/** @return The number of replicated objects */
+		virtual uint32 GetNumReplicatedObjects() const = 0;
+		/** @return The number of properties assigned to this object */
+		virtual uint32 GetNumProperties(const FSoftObjectPath& Object) const = 0;
+
+		/** @return Gets the class of the replication object */
+		virtual FSoftClassPath GetObjectClass(const FSoftObjectPath& Object) const = 0;
+
+		/** @return Whether the objects are selected */
+		virtual bool ContainsObjects(const TSet<FSoftObjectPath>& Objects) const = 0;
+		/** @return Whether the properties are selected */
+		virtual bool ContainsProperties(const FSoftObjectPath& Object, const TSet<FConcertPropertyChain>& Properties) const = 0;
+		
+		/**
+		 * Iterates all objects in the mapping.
+		 * @return Whether there were any mapped objects
+		 */
+		virtual bool ForEachReplicatedObject(TFunctionRef<EBreakBehavior(const FSoftObjectPath& Object)> Delegate) const = 0;
+
+		/**
+		 * Iterates all the properties of the given object in the mapping.
+		 * @param Object The object of which to get the properties
+		 * @param Delegate The callback
+		 */
+		virtual bool ForEachProperty(const FSoftObjectPath& Object, TFunctionRef<EBreakBehavior(const FConcertPropertyChain& Property)> Delegate) const = 0;
+		
+		/**
+		 * Iterates all the direct child properties of a given property in the mapping.
+		 *
+		 * @param Object The object of which to get the properties
+		 * @param Delegate The callback
+		 * @param OptionalParent Defines of which properties to get the child properties. If the property chain is empty, gets the object's root properties.
+		 *
+		 * @return Whether any properties were mapped
+		 */
+		bool ForEachDirectChildProperty(const FSoftObjectPath& Object, TFunctionRef<EBreakBehavior(const FConcertPropertyChain& Property)> Delegate, const FConcertPropertyChain& OptionalParent = {}) const
+		{
+			bool bFoundAtLeastOne = false;
+			ForEachProperty(Object, [&Delegate, &OptionalParent, &bFoundAtLeastOne](const FConcertPropertyChain& Property) mutable
+			{
+				const bool bCanCall = Property == OptionalParent;
+				bFoundAtLeastOne |= bCanCall;
+				return bCanCall
+					? Delegate(Property)
+					: EBreakBehavior::Continue;
+			});
+			return bFoundAtLeastOne;
+		}
+
+		/** Util for getting properties as array */
+		TSet<FConcertPropertyChain> GetAllProperties(const FSoftObjectPath& Object) const
+		{
+			TSet<FConcertPropertyChain> Result;
+			ForEachProperty(Object, [&Result](const FConcertPropertyChain& Property)
+			{
+				Result.Add(Property);
+				return EBreakBehavior::Continue;
+			});
+			return Result;
+		}
+		
+		/** Util for getting child properties as array */
+		TSet<FName> GetChildProperties(const FSoftObjectPath& Object, const FConcertPropertyChain& OptionalParent = {}) const
+		{
+			TSet<FName> Result;
+			ForEachDirectChildProperty(Object, [&Result](const FConcertPropertyChain& Property)
+			{
+				Result.Add(Property.GetLeafProperty());
+				return EBreakBehavior::Continue;
+			}, OptionalParent);
+			return Result;
+		}
+		
+		virtual ~IObjectToPropertiesModel() = default;
+	};
+}
diff --git a/Engine/Plugins/Developer/Concert/ConcertApp/MultiUserClient/Source/MultiUserReplicationEditor/Private/StreamEditor/Model/Object/AddOuterSource.cpp b/Engine/Plugins/Developer/Concert/ConcertApp/MultiUserClient/Source/MultiUserReplicationEditor/Private/StreamEditor/Model/Object/AddOuterSource.cpp
new file mode 100644
--- /dev/null
+++ b/Engine/Plugins/Developer/Concert/ConcertApp/MultiUserClient/Source/MultiUserReplicationEditor/Private/StreamEditor/Model/Object/AddOuterSource.cpp
@@ -0,0 +1,70 @@
+// Copyright Epic Games, Inc. All Rights Reserved.
+
+#include "AddOuterSource.h"
+
+#include "StreamEditor/Model/DisplayUtils.h"
+
+#include "GameFramework/Actor.h"
+
+#define LOCTEXT_NAMESPACE "FAddOuterSource"
+
+namespace UE::MultiUserReplicationEditor
+{
+	bool FAddOuterSource::TrySetObjectIfValid(UObject* InObject)
+	{
+		if (InObject
+			// Allow anything in actors but nothing above: persistent level should not be added.
+			&& InObject->IsInA(AActor::StaticClass()) && !InObject->IsA<AActor>())
+		{
+			Object = InObject;
+			return true;
+		}
+		return false;
+	}
+
+	ConcertSharedSlate::FSourceDisplayInfo FAddOuterSource::GetDisplayInfo() const
+	{
+		if (!ensure(Object.IsValid()))
+		{
+			return {{ LOCTEXT("Invalid", "Invalid") }};
+		}
+		
+		UObject* Outer = Object->GetOuter();
+		check(Outer);
+		
+		const FText Label = [Outer]()
+		{
+			if (Outer->IsA<AActor>())
+			{
+				return FText::Format(LOCTEXT("Add.OwningActor", "Add Actor {0}"), FText::FromString(DisplayUtils::GetObjectDisplayString(*Outer)));
+			}
+			if (UActorComponent* AsComponent = Cast<UActorComponent>(Outer))
+			{
+				return FText::Format(LOCTEXT("Add.Component", "Add Component {0} of {1}"),
+					FText::FromString(DisplayUtils::GetObjectDisplayString(*Outer)),
+					FText::FromString(DisplayUtils::GetObjectDisplayString(*AsComponent->GetOwner()))
+				);
+			}
+			return LOCTEXT("Add.Outer", "Add Outer");
+		}();
+		return {
+			{ Label, LOCTEXT("Tooltip", "Adds the owning object"), DisplayUtils::GetObjectIcon(*Outer) },
+			ConcertSharedSlate::ESourceType::AddOnClick
+		};
+	}
+
+	uint32 FAddOuterSource::GetNumSelectableItems() const
+	{
+		return Object.IsValid() ? 1 : 0;
+	}
+
+	void FAddOuterSource::EnumerateSelectableItems(TFunctionRef<EBreakBehavior(const FSelectableObjectInfo& SelectableOption)> Delegate) const
+	{
+		if (Object.IsValid())
+		{
+			Delegate(FSelectableObjectInfo{ *Object->GetOuter() });
+		}
+	}
+}
+
+#undef LOCTEXT_NAMESPACE
\ No newline at end of file
diff --git a/Engine/Plugins/Developer/Concert/ConcertApp/MultiUserClient/Source/MultiUserReplicationEditor/Private/StreamEditor/Model/Object/AddOuterSource.h b/Engine/Plugins/Developer/Concert/ConcertApp/MultiUserClient/Source/MultiUserReplicationEditor/Private/StreamEditor/Model/Object/AddOuterSource.h
new file mode 100644
--- /dev/null
+++ b/Engine/Plugins/Developer/Concert/ConcertApp/MultiUserClient/Source/MultiUserReplicationEditor/Private/StreamEditor/Model/Object/AddOuterSource.h
@@ -0,0 +1,29 @@
+// Copyright Epic Games, Inc. All Rights Reserved.
+
+#pragma once
+
+#include "StreamEditor/Model/Object/IObjectSourceModel.h"
+
+namespace UE::MultiUserReplicationEditor
+{
+	/** Adds the outer object as object. Used as source in context menu. */
+	class FAddOuterSource : public IObjectSourceModel
+	{
+	public:
+
+		bool TrySetObjectIfValid(UObject* InObject);
+
+		//~ Begin IObjectSourceModel Interface
+		virtual ConcertSharedSlate::FSourceDisplayInfo GetDisplayInfo() const override;
+		virtual uint32 GetNumSelectableItems() const override;
+		virtual void EnumerateSelectableItems(TFunctionRef<EBreakBehavior(const FSelectableObjectInfo& SelectableOption)> Delegate) const override;
+		//~ End IObjectSourceModel Interface
+
+	private:
+
+		/** The object whose outer to return. */
+		TWeakObjectPtr<UObject> Object;
+	};
+}
+
+
diff --git a/Engine/Plugins/Developer/Concert/ConcertApp/MultiUserClient/Source/MultiUserReplicationEditor/Private/StreamEditor/Model/Object/ComponentFromActorSource.cpp b/Engine/Plugins/Developer/Concert/ConcertApp/MultiUserClient/Source/MultiUserReplicationEditor/Private/StreamEditor/Model/Object/ComponentFromActorSource.cpp
new file mode 100644
--- /dev/null
+++ b/Engine/Plugins/Developer/Concert/ConcertApp/MultiUserClient/Source/MultiUserReplicationEditor/Private/StreamEditor/Model/Object/ComponentFromActorSource.cpp
@@ -0,0 +1,155 @@
+// Copyright Epic Games, Inc. All Rights Reserved.
+
+#include "ComponentFromActorSource.h"
+
+#include "StreamEditor/Model/DisplayUtils.h"
+#include "StreamEditor/Model/Object/IObjectSelectionSourceModel.h"
+
+#include "Algo/AnyOf.h"
+#include "Components/ActorComponent.h"
+#include "GameFramework/Actor.h"
+#include "Kismet2/ComponentEditorUtils.h"
+
+#define LOCTEXT_NAMESPACE "FComponentFromActorSource"
+
+namespace UE::MultiUserReplicationEditor
+{
+	bool FComponentFromActorSource_Base::TrySetActorIfValid(AActor* InComponentSource)
+	{
+		if (InComponentSource
+			// Excludes e.g. foliage and world settings
+			&& InComponentSource->IsListedInSceneOutliner())
+		{
+			ComponentSource = InComponentSource;
+			return true;
+		}
+		return false;
+	}
+
+	uint32 FComponentFromActorSource_Base::GetNumSelectableItems() const
+	{
+		return ComponentSource.IsValid() ? ComponentSource->GetComponents().Num() : 0;
+	}
+
+	void FComponentFromActorSource_Base::EnumerateSelectableItems(TFunctionRef<EBreakBehavior(const FSelectableObjectInfo& SelectableOption)> Delegate) const
+	{
+		if (!ComponentSource.IsValid())
+		{
+			return;
+		}
+
+		for (UActorComponent* Component : ComponentSource->GetComponents())
+		{
+			if (IsComponentValidReplicationObject(Component)
+				&& Delegate({ *Component }) == EBreakBehavior::Break)
+			{
+				break;
+			}
+		}
+	}
+	
+	bool FComponentFromActorSource_Base::IsComponentValidReplicationObject(UActorComponent* Component)
+	{
+		const USceneComponent* AsSceneComponent = Cast<USceneComponent>(Component);
+		const bool bIsParentValid = AsSceneComponent
+			&& AsSceneComponent->GetAttachParent()
+			&& Component->GetOwner() != nullptr
+			&& AsSceneComponent->GetAttachParent()->GetOwner() == Component->GetOwner();
+		const USceneComponent* Parent = bIsParentValid ? AsSceneComponent->GetAttachParent() : nullptr;
+
+		constexpr bool bAllowUserContructionScriptComps = true;
+		return  Component && FComponentEditorUtils::CanEditComponentInstance(Component, Parent, bAllowUserContructionScriptComps);
+	}
+
+	ConcertSharedSlate::FSourceDisplayInfo FComponentFromActorSource_ContextMenu::GetDisplayInfo() const
+	{
+		const FString ActorLabel = GetComponentSource().IsValid() ? DisplayUtils::GetObjectDisplayString(*GetComponentSource().Get()) : TEXT("None");
+		const FText ActorLabelText = FText::FromString(ActorLabel);
+		return
+		{
+			{
+			LOCTEXT("ContextMenu.Label", "Add component"),
+			FText::Format(LOCTEXT("ContextMenu.Tooltip", "Adds a component from {0}"), ActorLabelText),
+				FSlateIcon()
+			},
+			ConcertSharedSlate::ESourceType::ShowAsList
+		};
+	}
+
+	void FComponentFromActorSource_Root::AddComponentsAdditionToRootCategory(IObjectSourceModel& ActorSource, FObjectSourceCategory& Category)
+	{
+		FObjectSourceCategory SubCategory
+		{
+			{ LOCTEXT("AddComponent.Label", "Add Component"), LOCTEXT("AddComponent.Tooltip", "Add a component from an actor to the list.") },
+			{}
+		};
+
+		// The attribute's get function is called everytime the sub-category's menu entry is hovered.
+		// Once a lazy adder has been created once, avoid recreating it again.
+		using FActorToSourceMap = TMap<TWeakObjectPtr<AActor>, TSharedRef<FComponentFromActorSource_Root>>; 
+		TSharedRef<FActorToSourceMap> CachedLazilyCreatedComponentAdders = MakeShared<FActorToSourceMap>();
+		
+		ActorSource.EnumerateSelectableItems([&SubCategory, CachedLazilyCreatedComponentAdders](const FSelectableObjectInfo& SelectableOption) mutable
+		{
+			if (!SelectableOption.Object.IsValid())
+			{
+				return EBreakBehavior::Continue;
+			}
+			
+			AActor* Actor = Cast<AActor>(SelectableOption.Object.Get());
+			if (!Actor)
+			{
+				return EBreakBehavior::Continue;
+			}
+
+			const bool bHasAnyAddeableComponent = Algo::AnyOf(Actor->GetComponents(), [](UActorComponent* Component)
+			{
+				return IsComponentValidReplicationObject(Component);
+			});
+			if (!bHasAnyAddeableComponent)
+			{
+				return EBreakBehavior::Continue;
+			}
+
+			// The IObjectSourceModel should be created lazily because there may by thousands of actors in the world
+			TAttribute<TSharedPtr<IObjectSourceModel>> LazilyCreatedComponentSource = TAttribute<TSharedPtr<IObjectSourceModel>>::CreateLambda([WeakActor = TWeakObjectPtr<AActor>(Actor), CachedLazilyCreatedComponentAdders = MoveTemp(CachedLazilyCreatedComponentAdders)]()
+			{
+				if (const TSharedRef<FComponentFromActorSource_Root>* Cached = CachedLazilyCreatedComponentAdders->Find(WeakActor))
+				{
+					return *Cached;
+				}
+				
+				const TSharedRef<FComponentFromActorSource_Root> ComponentAdder = MakeShared<FComponentFromActorSource_Root>();
+				if (WeakActor.IsValid())
+				{
+					ComponentAdder->TrySetActorIfValid(WeakActor.Get());
+				}
+
+				CachedLazilyCreatedComponentAdders->Add(WeakActor, ComponentAdder);
+				return ComponentAdder;
+			});
+			SubCategory.Options.Emplace(LazilyCreatedComponentSource);
+			
+			return EBreakBehavior::Continue;
+		});
+		
+		Category.SubCategories.Emplace(MoveTemp(SubCategory));
+	}
+
+	ConcertSharedSlate::FSourceDisplayInfo FComponentFromActorSource_Root::GetDisplayInfo() const
+	{
+		const FString ActorLabel = GetComponentSource().IsValid() ? DisplayUtils::GetObjectDisplayString(*GetComponentSource().Get()) : TEXT("None");
+		const FText ActorLabelText = FText::FromString(ActorLabel);
+		const FSlateIcon SlateIcon = GetComponentSource().IsValid() ? DisplayUtils::GetObjectIcon(*GetComponentSource().Get()) : FSlateIcon();
+		return {
+			{
+				ActorLabelText,
+			FText::Format(LOCTEXT("ContextMenu.Tooltip", "Adds a component from {0}"), ActorLabelText),
+				SlateIcon
+			},
+			ConcertSharedSlate::ESourceType::ShowAsList
+		};
+	}
+}
+
+#undef LOCTEXT_NAMESPACE
\ No newline at end of file
diff --git a/Engine/Plugins/Developer/Concert/ConcertApp/MultiUserClient/Source/MultiUserReplicationEditor/Private/StreamEditor/Model/Object/ComponentFromActorSource.h b/Engine/Plugins/Developer/Concert/ConcertApp/MultiUserClient/Source/MultiUserReplicationEditor/Private/StreamEditor/Model/Object/ComponentFromActorSource.h
new file mode 100644
--- /dev/null
+++ b/Engine/Plugins/Developer/Concert/ConcertApp/MultiUserClient/Source/MultiUserReplicationEditor/Private/StreamEditor/Model/Object/ComponentFromActorSource.h
@@ -0,0 +1,65 @@
+// Copyright Epic Games, Inc. All Rights Reserved.
+
+#pragma once
+
+#include "IObjectSelectionSourceModel.h"
+#include "StreamEditor/Model/Object/IObjectSourceModel.h"
+
+class AActor;
+class UActorComponent;
+
+namespace UE::MultiUserReplicationEditor
+{
+	/** Gets the components on a given actor. Reused by the context menu and root filters. */
+	class FComponentFromActorSource_Base : public IObjectSourceModel
+	{
+	public:
+
+		bool TrySetActorIfValid(AActor* InComponentSource);
+
+		//~ Begin IObjectSourceModel Interface
+		virtual uint32 GetNumSelectableItems() const override;
+		virtual void EnumerateSelectableItems(TFunctionRef<EBreakBehavior(const FSelectableObjectInfo& SelectableOption)> Delegate) const override;
+		//~ End IObjectSourceModel Interface
+
+	protected:
+
+		/** Checks that the component would be visible in the detail's SSubobjectEditor */
+		static bool IsComponentValidReplicationObject(UActorComponent* Component);
+		
+		TWeakObjectPtr<AActor> GetComponentSource() const { return ComponentSource; }
+		
+	private:
+
+		/** The actor from which to add a component from. */
+		TWeakObjectPtr<AActor> ComponentSource;
+	};
+
+	/** Used from the context menu. The label is "Add component". */
+	class FComponentFromActorSource_ContextMenu : public FComponentFromActorSource_Base
+	{
+	public:
+
+		//~ Begin IObjectSourceModel Interface
+		virtual ConcertSharedSlate::FSourceDisplayInfo GetDisplayInfo() const override;
+		//~ End IObjectSourceModel Interface
+	};
+
+	/** A root filter. The label is the name of the actor. */
+	class FComponentFromActorSource_Root : public FComponentFromActorSource_Base
+	{
+	public:
+		
+		/**
+		 * Adds a "Add component" sub-category. When expanded, it will list all actors as sub-menus, which contain components to add.
+		 *
+		 * @param ActorSource This model provides the actors for which components can be added
+		 * @param OutCategory A new sub-category is added to this category
+		 */
+		static void AddComponentsAdditionToRootCategory(IObjectSourceModel& ActorSource, FObjectSourceCategory& OutCategory);
+
+		//~ Begin IObjectSourceModel Interface
+		virtual ConcertSharedSlate::FSourceDisplayInfo GetDisplayInfo() const override;
+		//~ End IObjectSourceModel Interface
+	};
+}
diff --git a/Engine/Plugins/Developer/Concert/ConcertApp/MultiUserClient/Source/MultiUserReplicationEditor/Private/StreamEditor/Model/Object/EditorObjectSelectionSourceModel.cpp b/Engine/Plugins/Developer/Concert/ConcertApp/MultiUserClient/Source/MultiUserReplicationEditor/Private/StreamEditor/Model/Object/EditorObjectSelectionSourceModel.cpp
new file mode 100644
--- /dev/null
+++ b/Engine/Plugins/Developer/Concert/ConcertApp/MultiUserClient/Source/MultiUserReplicationEditor/Private/StreamEditor/Model/Object/EditorObjectSelectionSourceModel.cpp
@@ -0,0 +1,85 @@
+// Copyright Epic Games, Inc. All Rights Reserved.
+
+#include "EditorObjectSelectionSourceModel.h"
+
+#include "AddOuterSource.h"
+#include "ComponentFromActorSource.h"
+#include "SelectedActorsSource.h"
+#include "WorldActorSource.h"
+
+#include "GameFramework/Actor.h"
+
+#define LOCTEXT_NAMESPACE "FEditorObjectSelectionSourceModel"
+
+namespace UE::MultiUserReplicationEditor
+{
+	FEditorObjectSelectionSourceModel::FEditorObjectSelectionSourceModel()
+		: WorldActorSource(MakeShared<FWorldActorSource>())
+		// Important: We dependent on the fact that WorldActorSource is declared (in the header) and hence initialized before BaseActorCategory
+		, BaseActorCategory([this]()
+		{
+			return FObjectSourceCategory
+			{
+				ConcertSharedSlate::FBaseDisplayInfo
+				{
+					LOCTEXT("ObjectsCategory.Label", "Add Objects"),
+					LOCTEXT("ObjectsCategory.Tooltip", "Options for adding objects from the open editor world")
+				},
+				{
+					WorldActorSource,
+					MakeShared<FSelectedActorsSource>()
+				}
+			};
+		}())
+	{}
+
+	TArray<FObjectSourceCategory> FEditorObjectSelectionSourceModel::GetRootSources() const
+	{
+		FObjectSourceCategory ActorCategory = BaseActorCategory;
+		FComponentFromActorSource_Root::AddComponentsAdditionToRootCategory(WorldActorSource.Get(), ActorCategory);
+		return { ActorCategory };
+	}
+
+	TArray<TSharedRef<IObjectSourceModel>> FEditorObjectSelectionSourceModel::GetContextMenuOptions(const FSoftObjectPath& Item)
+	{
+		return ContextMenuSources.BuildContextSources(Item);
+	}
+
+	EObjectItemValidity FEditorObjectSelectionSourceModel::GetItemValidity(const FSoftObjectPath& Item) const
+	{
+		if (!Item.ResolveObject() && !Item.TryLoad())
+		{
+			return EObjectItemValidity::DoesNotExist;
+		}
+		
+		return EObjectItemValidity::Invalid;
+	}
+
+	FEditorObjectSelectionSourceModel::FContextMenuSources::FContextMenuSources()
+		: AddComponent(MakeShared<FComponentFromActorSource_ContextMenu>())
+		, AddOuter(MakeShared<FAddOuterSource>())
+	{}
+
+	TArray<TSharedRef<IObjectSourceModel>> FEditorObjectSelectionSourceModel::FContextMenuSources::BuildContextSources(const FSoftObjectPath& Item)
+	{
+		UObject* ResolvedContextObject = Item.ResolveObject();
+		AActor* ResolvedContextActor = Cast<AActor>(ResolvedContextObject);
+		TArray<TSharedRef<IObjectSourceModel>> Result;
+
+		// The "Add Outer" option is not valid for all cases. For example adding an actor's outer is not valid.
+		if (AddOuter->TrySetObjectIfValid(ResolvedContextObject))
+		{
+			Result.Add(AddOuter);
+		}
+
+		// The "Add component" option is only valid if the context object is an actor displayed in the scene outliner
+		if (AddComponent->TrySetActorIfValid(ResolvedContextActor))
+		{
+			Result.Add(AddComponent);
+		}
+
+		return Result;
+	}
+}
+
+#undef LOCTEXT_NAMESPACE
\ No newline at end of file
diff --git a/Engine/Plugins/Developer/Concert/ConcertApp/MultiUserClient/Source/MultiUserReplicationEditor/Private/StreamEditor/Model/Object/EditorObjectSelectionSourceModel.h b/Engine/Plugins/Developer/Concert/ConcertApp/MultiUserClient/Source/MultiUserReplicationEditor/Private/StreamEditor/Model/Object/EditorObjectSelectionSourceModel.h
new file mode 100644
--- /dev/null
+++ b/Engine/Plugins/Developer/Concert/ConcertApp/MultiUserClient/Source/MultiUserReplicationEditor/Private/StreamEditor/Model/Object/EditorObjectSelectionSourceModel.h
@@ -0,0 +1,49 @@
+// Copyright Epic Games, Inc. All Rights Reserved.
+
+#pragma once
+
+#include "StreamEditor/Model/Object/IObjectSelectionSourceModel.h"
+
+namespace UE::MultiUserReplicationEditor
+{
+	class FAddOuterSource;
+	class FComponentFromActorSource_ContextMenu;
+	class FWorldActorSource;
+	
+	class FEditorObjectSelectionSourceModel : public IObjectSelectionSourceModel
+	{
+	public:
+
+		FEditorObjectSelectionSourceModel();
+
+		//~ Begin IObjectSelectionSourceModel Interface
+		virtual TArray<FObjectSourceCategory> GetRootSources() const override;
+		virtual TArray<TSharedRef<IObjectSourceModel>> GetContextMenuOptions(const FSoftObjectPath& Item) override;
+		virtual EObjectItemValidity GetItemValidity(const FSoftObjectPath& Item) const override;
+		//~ End IObjectSelectionSourceModel Interface
+
+	private:
+
+		/** Used to build Add Component source from BaseActorCategory. If you move this property, make sure it declared before BaseActorCategory, or update the constructor. */
+		TSharedRef<FWorldActorSource> WorldActorSource;
+
+		/** Contains actions around adding actor related objects. Displayed next to the search bar. Refined with dynamic options in GetRootSources. */
+		FObjectSourceCategory BaseActorCategory;
+
+		/** Sources which have to be instantiated with a context. */
+		struct FContextMenuSources
+		{
+			/** Add components from the actor for which the context menu is instantiated. */
+			TSharedRef<FComponentFromActorSource_ContextMenu> AddComponent;
+
+			/** Adds the object's outer as object. */
+			TSharedRef<FAddOuterSource> AddOuter;
+			
+			FContextMenuSources();
+
+			/** Gets all applicable context menu sources. */
+			TArray<TSharedRef<IObjectSourceModel>> BuildContextSources(const FSoftObjectPath& Item);
+		} ContextMenuSources;
+	};
+}
+
diff --git a/Engine/Plugins/Developer/Concert/ConcertApp/MultiUserClient/Source/MultiUserReplicationEditor/Private/StreamEditor/Model/Object/IObjectSelectionSourceModel.h b/Engine/Plugins/Developer/Concert/ConcertApp/MultiUserClient/Source/MultiUserReplicationEditor/Private/StreamEditor/Model/Object/IObjectSelectionSourceModel.h
new file mode 100644
--- /dev/null
+++ b/Engine/Plugins/Developer/Concert/ConcertApp/MultiUserClient/Source/MultiUserReplicationEditor/Private/StreamEditor/Model/Object/IObjectSelectionSourceModel.h
@@ -0,0 +1,45 @@
+// Copyright Epic Games, Inc. All Rights Reserved.
+
+#pragma once
+
+#include "Containers/Array.h"
+#include "IObjectSourceModel.h"
+#include "Model/Item/SourceSelectionCategory.h"
+#include "Templates/SharedPointer.h"
+
+namespace UE::MultiUserReplicationEditor
+{
+	using FObjectSourceCategory = ConcertSharedSlate::TSourceSelectionCategory<FSelectableObjectInfo>;
+	
+	enum class EObjectItemValidity : uint8
+	{
+		/** The item is allowed */
+		Valid,
+		/** The referenced object no longer exists. E.g. the actor was removed. */
+		DoesNotExist,
+		/** Any other reason, e.g. the object was allowed in a previous version but no longer is allowed now. */
+		Invalid
+	};
+	
+	/** Decides which objects can be added to a IEditableObjectToPropertiesModel. */
+	class IObjectSelectionSourceModel : public TSharedFromThis<IObjectSelectionSourceModel>
+	{
+	public:
+
+		/** Gets the "root set" of sources. Each one should be a button next to the search bar, e.g. "Add actor". */
+		virtual TArray<FObjectSourceCategory> GetRootSources() const = 0;
+
+		/** @return The sources that should be displayed in the right-click context menu when clicked in the menu. */
+		virtual TArray<TSharedRef<IObjectSourceModel>> GetContextMenuOptions(const FSoftObjectPath& Items) = 0;
+
+		/**
+		 * Checks whether the selected item is valid to be added to the IEditableObjectToPropertiesModel.
+		 *
+		 * This can be used to validate the asset when it is initially displayed but could also be useful for validating
+		 * actors / assets that were drag-dropped into the editor.
+		 */
+		virtual EObjectItemValidity GetItemValidity(const FSoftObjectPath& Item) const = 0;
+		
+		virtual ~IObjectSelectionSourceModel() = default;
+	};
+}
diff --git a/Engine/Plugins/Developer/Concert/ConcertApp/MultiUserClient/Source/MultiUserReplicationEditor/Private/StreamEditor/Model/Object/IObjectSourceModel.h b/Engine/Plugins/Developer/Concert/ConcertApp/MultiUserClient/Source/MultiUserReplicationEditor/Private/StreamEditor/Model/Object/IObjectSourceModel.h
new file mode 100644
--- /dev/null
+++ b/Engine/Plugins/Developer/Concert/ConcertApp/MultiUserClient/Source/MultiUserReplicationEditor/Private/StreamEditor/Model/Object/IObjectSourceModel.h
@@ -0,0 +1,26 @@
+// Copyright Epic Games, Inc. All Rights Reserved.
+
+#pragma once
+
+#include "Model/Item/IItemSourceModel.h"
+#include "UObject/WeakObjectPtr.h"
+#include "UObject/WeakObjectPtrTemplates.h"
+
+class UObject;
+
+namespace UE::MultiUserReplicationEditor
+{
+	/** Info about a selectable object */
+	struct FSelectableObjectInfo
+	{
+		FSelectableObjectInfo(UObject& Object) : Object(&Object) {}
+		
+		TWeakObjectPtr<UObject> Object;
+	};
+	
+	/**
+	 * A specific object source, e.g. like actors, components from an actor (right-click), etc.
+	 * @see IObjectSelectionSourceModel.
+	 */
+	using IObjectSourceModel = ConcertSharedSlate::IItemSourceModel<FSelectableObjectInfo>;
+}
diff --git a/Engine/Plugins/Developer/Concert/ConcertApp/MultiUserClient/Source/MultiUserReplicationEditor/Private/StreamEditor/Model/Object/SelectedActorsSource.cpp b/Engine/Plugins/Developer/Concert/ConcertApp/MultiUserClient/Source/MultiUserReplicationEditor/Private/StreamEditor/Model/Object/SelectedActorsSource.cpp
new file mode 100644
--- /dev/null
+++ b/Engine/Plugins/Developer/Concert/ConcertApp/MultiUserClient/Source/MultiUserReplicationEditor/Private/StreamEditor/Model/Object/SelectedActorsSource.cpp
@@ -0,0 +1,51 @@
+// Copyright Epic Games, Inc. All Rights Reserved.
+
+#include "SelectedActorsSource.h"
+
+#include "Editor.h"
+#include "Selection.h"
+
+#define LOCTEXT_NAMESPACE "FSelectedActorsSourceModel"
+
+namespace UE::MultiUserReplicationEditor
+{
+	ConcertSharedSlate::FSourceDisplayInfo FSelectedActorsSource::GetDisplayInfo() const
+	{
+		return
+		{
+			{
+				LOCTEXT("Label", "Selected Actors"),
+			   LOCTEXT("Tooltip", "Add the actors under your current selection in the world outliner"),
+			   FSlateIcon{}
+			},
+			ConcertSharedSlate::ESourceType::AddOnClick
+		};
+	}
+
+	uint32 FSelectedActorsSource::GetNumSelectableItems() const
+	{
+		return ensure(GEditor)
+			? GEditor->GetSelectedActors()->Num()
+			: 0;
+	}
+
+	void FSelectedActorsSource::EnumerateSelectableItems(TFunctionRef<EBreakBehavior(const FSelectableObjectInfo& SelectableOption)> Delegate) const
+	{
+		if (!ensure(GEditor))
+		{
+			return;
+		}
+
+		USelection* SelectedActors = GEditor->GetSelectedActors();
+		for (int32 i = 0; i < SelectedActors->Num(); ++i)
+		{
+			UObject* SelectedActor = SelectedActors->GetSelectedObject(i);
+			if (SelectedActor && Delegate({ *SelectedActor }) == EBreakBehavior::Break)
+			{
+				return;
+			}
+		}
+	}
+}
+
+#undef LOCTEXT_NAMESPACE
\ No newline at end of file
diff --git a/Engine/Plugins/Developer/Concert/ConcertApp/MultiUserClient/Source/MultiUserReplicationEditor/Private/StreamEditor/Model/Object/SelectedActorsSource.h b/Engine/Plugins/Developer/Concert/ConcertApp/MultiUserClient/Source/MultiUserReplicationEditor/Private/StreamEditor/Model/Object/SelectedActorsSource.h
new file mode 100644
--- /dev/null
+++ b/Engine/Plugins/Developer/Concert/ConcertApp/MultiUserClient/Source/MultiUserReplicationEditor/Private/StreamEditor/Model/Object/SelectedActorsSource.h
@@ -0,0 +1,20 @@
+// Copyright Epic Games, Inc. All Rights Reserved.
+
+#pragma once
+
+#include "StreamEditor/Model/Object/IObjectSourceModel.h"
+
+namespace UE::MultiUserReplicationEditor
+{
+	/** Gets the actors currently selected in the editor. */
+	class FSelectedActorsSource : public IObjectSourceModel
+	{
+	public:
+
+		//~ Begin IObjectSourceModel Interface
+		virtual ConcertSharedSlate::FSourceDisplayInfo GetDisplayInfo() const override;
+		virtual uint32 GetNumSelectableItems() const override;
+		virtual void EnumerateSelectableItems(TFunctionRef<EBreakBehavior(const FSelectableObjectInfo& SelectableOption)> Delegate) const override;
+		//~ End IObjectSourceModel Interface
+	};
+}
diff --git a/Engine/Plugins/Developer/Concert/ConcertApp/MultiUserClient/Source/MultiUserReplicationEditor/Private/StreamEditor/Model/Object/WorldActorSource.cpp b/Engine/Plugins/Developer/Concert/ConcertApp/MultiUserClient/Source/MultiUserReplicationEditor/Private/StreamEditor/Model/Object/WorldActorSource.cpp
new file mode 100644
--- /dev/null
+++ b/Engine/Plugins/Developer/Concert/ConcertApp/MultiUserClient/Source/MultiUserReplicationEditor/Private/StreamEditor/Model/Object/WorldActorSource.cpp
@@ -0,0 +1,55 @@
+// Copyright Epic Games, Inc. All Rights Reserved.
+
+#include "WorldActorSource.h"
+
+#include "Algo/Accumulate.h"
+#include "CoreGlobals.h"
+#include "Engine/World.h"
+#include "Engine/Level.h"
+#include "EngineUtils.h"
+#include "Algo/Count.h"
+#include "GameFramework/Actor.h"
+
+#define LOCTEXT_NAMESPACE "FWorldActorSource"
+
+namespace UE::MultiUserReplicationEditor
+{
+	ConcertSharedSlate::FSourceDisplayInfo FWorldActorSource::GetDisplayInfo() const
+	{
+		return {{
+			LOCTEXT("Label", "Add Actor from World"),
+		   LOCTEXT("Tooltip", "Pick an actor from the editor world")
+		}};
+	}
+
+	uint32 FWorldActorSource::GetNumSelectableItems() const
+	{
+		return ensure(GWorld)
+			? Algo::TransformAccumulate(GWorld->GetLevels(), [](ULevel* Level)
+			{
+				return Level
+					? Algo::CountIf(Level->Actors, [](AActor* Actor){ return Actor && Actor->IsListedInSceneOutliner(); })
+					: 0;
+			}, 0)
+			: 0;
+	}
+
+	void FWorldActorSource::EnumerateSelectableItems(TFunctionRef<EBreakBehavior(const FSelectableObjectInfo& SelectableOption)> Delegate) const
+	{
+		if (!ensure(GWorld))
+		{
+			return;
+		}
+
+		for (TActorIterator<AActor> ActorIt(GWorld); ActorIt; ++ActorIt)
+		{
+			AActor* Actor = *ActorIt;
+			if (Actor && Actor->IsListedInSceneOutliner() && Delegate({ *Actor }) == EBreakBehavior::Break)
+			{
+				return;
+			}
+		}
+	}
+}
+
+#undef LOCTEXT_NAMESPACE
\ No newline at end of file
diff --git a/Engine/Plugins/Developer/Concert/ConcertApp/MultiUserClient/Source/MultiUserReplicationEditor/Private/StreamEditor/Model/Object/WorldActorSource.h b/Engine/Plugins/Developer/Concert/ConcertApp/MultiUserClient/Source/MultiUserReplicationEditor/Private/StreamEditor/Model/Object/WorldActorSource.h
new file mode 100644
--- /dev/null
+++ b/Engine/Plugins/Developer/Concert/ConcertApp/MultiUserClient/Source/MultiUserReplicationEditor/Private/StreamEditor/Model/Object/WorldActorSource.h
@@ -0,0 +1,22 @@
+// Copyright Epic Games, Inc. All Rights Reserved.
+
+#pragma once
+
+#include "StreamEditor/Model/Object/IObjectSourceModel.h"
+
+namespace UE::MultiUserReplicationEditor
+{
+	/** Gets all the actors in the current GWorld */
+	class FWorldActorSource : public IObjectSourceModel
+	{
+	public:
+
+		//~ Begin IObjectSourceModel Interface
+		virtual ConcertSharedSlate::FSourceDisplayInfo GetDisplayInfo() const override;
+		virtual uint32 GetNumSelectableItems() const override;
+		virtual void EnumerateSelectableItems(TFunctionRef<EBreakBehavior(const FSelectableObjectInfo& SelectableOption)> Delegate) const override;
+		//~ End IObjectSourceModel Interface
+	};
+}
+
+
diff --git a/Engine/Plugins/Developer/Concert/ConcertApp/MultiUserClient/Source/MultiUserReplicationEditor/Private/StreamEditor/Model/Property/ConcertSyncCoreReplicatedPropertySource.cpp b/Engine/Plugins/Developer/Concert/ConcertApp/MultiUserClient/Source/MultiUserReplicationEditor/Private/StreamEditor/Model/Property/ConcertSyncCoreReplicatedPropertySource.cpp
new file mode 100644
--- /dev/null
+++ b/Engine/Plugins/Developer/Concert/ConcertApp/MultiUserClient/Source/MultiUserReplicationEditor/Private/StreamEditor/Model/Property/ConcertSyncCoreReplicatedPropertySource.cpp
@@ -0,0 +1,54 @@
+// Copyright Epic Games, Inc. All Rights Reserved.
+
+#include "ConcertSyncCoreReplicatedPropertySource.h"
+
+#include "Replication/PropertyChainUtils.h"
+#include "UObject/UnrealType.h"
+
+#define LOCTEXT_NAMESPACE "FSimplePropertyIterationSource"
+
+namespace UE::MultiUserReplicationEditor
+{
+	void FConcertSyncCoreReplicatedPropertySource::SetClass(UClass* InClass)
+	{
+		Class = InClass;
+		NumProperties = 0;
+		if (!Class.IsValid())
+		{
+			return;
+		}
+		
+		ConcertSyncCore::PropertyChain::ForEachReplicatableProperty(*Class.Get(), [this](const auto&, const auto&)
+		{
+			++NumProperties;
+			return EBreakBehavior::Continue;
+		});
+	}
+
+	ConcertSharedSlate::FSourceDisplayInfo FConcertSyncCoreReplicatedPropertySource::GetDisplayInfo() const
+	{
+		return {
+			{
+				LOCTEXT("Label", "Add Property"),
+			   FText::GetEmpty(),
+			   FSlateIcon()
+			},
+			ConcertSharedSlate::ESourceType::ShowAsList
+		};
+	}
+
+	void FConcertSyncCoreReplicatedPropertySource::EnumerateSelectableItems(TFunctionRef<EBreakBehavior(const FSelectablePropertyInfo& SelectableOption)> Delegate) const
+	{
+		if (!Class.IsValid())
+		{
+			return;
+		}
+
+		ConcertSyncCore::PropertyChain::ForEachReplicatableConcertProperty(*Class.Get(), [this, &Delegate](FConcertPropertyChain&& Property)
+		{
+			return Delegate({ MoveTemp(Property) });
+		});
+	}
+}
+
+#undef LOCTEXT_NAMESPACE
\ No newline at end of file
diff --git a/Engine/Plugins/Developer/Concert/ConcertApp/MultiUserClient/Source/MultiUserReplicationEditor/Private/StreamEditor/Model/Property/ConcertSyncCoreReplicatedPropertySource.h b/Engine/Plugins/Developer/Concert/ConcertApp/MultiUserClient/Source/MultiUserReplicationEditor/Private/StreamEditor/Model/Property/ConcertSyncCoreReplicatedPropertySource.h
new file mode 100644
--- /dev/null
+++ b/Engine/Plugins/Developer/Concert/ConcertApp/MultiUserClient/Source/MultiUserReplicationEditor/Private/StreamEditor/Model/Property/ConcertSyncCoreReplicatedPropertySource.h
@@ -0,0 +1,29 @@
+// Copyright Epic Games, Inc. All Rights Reserved.
+
+#pragma once
+
+#include "IPropertySourceModel.h"
+
+namespace UE::MultiUserReplicationEditor
+{
+	/** The allowed properties are those returned by UE::ConcertSyncCore::ForEachReplicatableProperty.*/
+	class FConcertSyncCoreReplicatedPropertySource : public IPropertySourceModel
+	{
+	public:
+
+		void SetClass(UClass* InClass);
+		
+		//~ Begin FSimplePropertyIterationSource Interface
+		virtual ConcertSharedSlate::FSourceDisplayInfo GetDisplayInfo() const override;
+		virtual uint32 GetNumSelectableItems() const override { return NumProperties; }
+		virtual void EnumerateSelectableItems(TFunctionRef<EBreakBehavior(const FSelectablePropertyInfo& SelectableOption)> Delegate) const override;
+		//~ End FSimplePropertyIterationSource Interface
+
+	private:
+
+		/** Weak ptr because it could be a Blueprint class that can be destroyed at any time. */
+		TWeakObjectPtr<UClass> Class;
+
+		int32 NumProperties = 0;
+	};
+}
diff --git a/Engine/Plugins/Developer/Concert/ConcertApp/MultiUserClient/Source/MultiUserReplicationEditor/Private/StreamEditor/Model/Property/IPropertySelectionSourceModel.h b/Engine/Plugins/Developer/Concert/ConcertApp/MultiUserClient/Source/MultiUserReplicationEditor/Private/StreamEditor/Model/Property/IPropertySelectionSourceModel.h
new file mode 100644
--- /dev/null
+++ b/Engine/Plugins/Developer/Concert/ConcertApp/MultiUserClient/Source/MultiUserReplicationEditor/Private/StreamEditor/Model/Property/IPropertySelectionSourceModel.h
@@ -0,0 +1,23 @@
+// Copyright Epic Games, Inc. All Rights Reserved.
+
+#pragma once
+
+#include "IPropertySourceModel.h"
+#include "Model/Item/SourceSelectionCategory.h"
+#include "Templates/SharedPointer.h"
+
+namespace UE::MultiUserReplicationEditor
+{
+	using FPropertySourceCategory = ConcertSharedSlate::TSourceSelectionCategory<FSelectablePropertyInfo>;
+	
+	/** Decides which properties can be added to a IObjectToPropertiesModel. */
+	class IPropertySelectionSourceModel : public TSharedFromThis<IPropertySelectionSourceModel>
+	{
+	public:
+
+		/** Gets the single source determining which properties can be selected. */
+		virtual TSharedRef<IPropertySourceModel> GetPropertySource(const FSoftClassPath& Class) const = 0;
+		
+		virtual ~IPropertySelectionSourceModel() = default;
+	};
+}
diff --git a/Engine/Plugins/Developer/Concert/ConcertApp/MultiUserClient/Source/MultiUserReplicationEditor/Private/StreamEditor/Model/Property/IPropertySourceModel.h b/Engine/Plugins/Developer/Concert/ConcertApp/MultiUserClient/Source/MultiUserReplicationEditor/Private/StreamEditor/Model/Property/IPropertySourceModel.h
new file mode 100644
--- /dev/null
+++ b/Engine/Plugins/Developer/Concert/ConcertApp/MultiUserClient/Source/MultiUserReplicationEditor/Private/StreamEditor/Model/Property/IPropertySourceModel.h
@@ -0,0 +1,21 @@
+// Copyright Epic Games, Inc. All Rights Reserved.
+
+#pragma once
+
+#include "Model/Item/IItemSourceModel.h"
+#include "Replication/Data/ConcertPropertySelection.h"
+
+namespace UE::MultiUserReplicationEditor
+{
+	/** Info about a selectable property */
+	struct FSelectablePropertyInfo
+	{
+		FConcertPropertyChain Property;
+	};
+	
+	/**
+	 * A specific object source, e.g. like actors, components from an actor (right-click), etc.
+	 * @see IPropertySelectionSourceModel.
+	 */
+	using IPropertySourceModel = ConcertSharedSlate::IItemSourceModel<FSelectablePropertyInfo>;
+}
diff --git a/Engine/Plugins/Developer/Concert/ConcertApp/MultiUserClient/Source/MultiUserReplicationEditor/Private/StreamEditor/Model/Property/SelectPropertyFromUClassModel.cpp b/Engine/Plugins/Developer/Concert/ConcertApp/MultiUserClient/Source/MultiUserReplicationEditor/Private/StreamEditor/Model/Property/SelectPropertyFromUClassModel.cpp
new file mode 100644
--- /dev/null
+++ b/Engine/Plugins/Developer/Concert/ConcertApp/MultiUserClient/Source/MultiUserReplicationEditor/Private/StreamEditor/Model/Property/SelectPropertyFromUClassModel.cpp
@@ -0,0 +1,23 @@
+// Copyright Epic Games, Inc. All Rights Reserved.
+
+#include "SelectPropertyFromUClassModel.h"
+
+#include "ConcertSyncCoreReplicatedPropertySource.h"
+
+#define LOCTEXT_NAMESPACE "FSelectPropertyFromUClassModel"
+
+namespace UE::MultiUserReplicationEditor
+{
+	FSelectPropertyFromUClassModel::FSelectPropertyFromUClassModel()
+		: UClassIteratorSource(MakeShared<FConcertSyncCoreReplicatedPropertySource>())
+	{}
+
+	TSharedRef<IPropertySourceModel> FSelectPropertyFromUClassModel::GetPropertySource(const FSoftClassPath& Class) const
+	{
+		UClass* LoadedClass = Class.TryLoadClass<UObject>();
+		UClassIteratorSource->SetClass(LoadedClass);
+		return UClassIteratorSource;
+	}
+}
+
+#undef LOCTEXT_NAMESPACE
\ No newline at end of file
diff --git a/Engine/Plugins/Developer/Concert/ConcertApp/MultiUserClient/Source/MultiUserReplicationEditor/Private/StreamEditor/Model/Property/SelectPropertyFromUClassModel.h b/Engine/Plugins/Developer/Concert/ConcertApp/MultiUserClient/Source/MultiUserReplicationEditor/Private/StreamEditor/Model/Property/SelectPropertyFromUClassModel.h
new file mode 100644
--- /dev/null
+++ b/Engine/Plugins/Developer/Concert/ConcertApp/MultiUserClient/Source/MultiUserReplicationEditor/Private/StreamEditor/Model/Property/SelectPropertyFromUClassModel.h
@@ -0,0 +1,30 @@
+// Copyright Epic Games, Inc. All Rights Reserved.
+
+#pragma once
+
+#include "IPropertySelectionSourceModel.h"
+
+namespace UE::MultiUserReplicationEditor
+{
+	class FConcertSyncCoreReplicatedPropertySource;
+	
+	/**
+	 * Decides which properties can be added to IEditableObjectToPropertiesModel.
+	 * The allowed properties are those returned by UE::ConcertSyncCore::ForEachReplicatableProperty.
+	 */
+	class FSelectPropertyFromUClassModel : public IPropertySelectionSourceModel
+	{
+	public:
+
+		FSelectPropertyFromUClassModel();
+
+		//~ Begin IPropertySelectionSourceModel Interface
+		virtual TSharedRef<IPropertySourceModel> GetPropertySource(const FSoftClassPath& Class) const override;
+		//~ End IPropertySelectionSourceModel Interface
+
+	private:
+
+		TSharedRef<FConcertSyncCoreReplicatedPropertySource> UClassIteratorSource;
+	};
+}
+
diff --git a/Engine/Plugins/Developer/Concert/ConcertApp/MultiUserClient/Source/MultiUserReplicationEditor/Private/StreamEditor/Model/PropertySelectionAssetModel.cpp b/Engine/Plugins/Developer/Concert/ConcertApp/MultiUserClient/Source/MultiUserReplicationEditor/Private/StreamEditor/Model/PropertySelectionAssetModel.cpp
new file mode 100644
--- /dev/null
+++ b/Engine/Plugins/Developer/Concert/ConcertApp/MultiUserClient/Source/MultiUserReplicationEditor/Private/StreamEditor/Model/PropertySelectionAssetModel.cpp
@@ -0,0 +1,326 @@
+// Copyright Epic Games, Inc. All Rights Reserved.
+
+#include "PropertySelectionAssetModel.h"
+
+#include "Replication/PropertyChainUtils.h"
+#include "Settings/DefaultPropertySelection.h"
+#include "Settings/MultiUserReplicationSettings.h"
+
+#include "Algo/AllOf.h"
+#include "ScopedTransaction.h"
+#include "Misc/TransactionObjectEvent.h"
+
+#define LOCTEXT_NAMESPACE "FPropertySelectionAssetModel"
+
+namespace UE::MultiUserReplicationEditor
+{
+	FPropertySelectionAssetModel::FPropertySelectionAssetModel(UMultiUserPropertyReplicationSelection& Object)
+		: Asset(&Object)
+	{}
+	
+	uint32 FPropertySelectionAssetModel::GetNumReplicatedObjects() const
+	{
+		return ensure(Asset.IsValid())
+			? Asset->ReplicationMap.ReplicatedObjects.Num()
+			: 0;
+	}
+
+	uint32 FPropertySelectionAssetModel::GetNumProperties(const FSoftObjectPath& Object) const
+	{
+		if (!ensure(Asset.IsValid()))
+		{
+			return false;
+		}
+
+		const FReplicatedObjectInfo* AssignedProperties = Asset->ReplicationMap.ReplicatedObjects.Find(Object);
+		return AssignedProperties
+			? AssignedProperties->ReplicatedProperties.ReplicatedProperties.Num()
+			: 0;
+	}
+
+	FSoftClassPath FPropertySelectionAssetModel::GetObjectClass(const FSoftObjectPath& Object) const
+	{
+		if (!ensure(Asset.IsValid()))
+		{
+			return {};
+		}
+		
+		const FReplicatedObjectInfo* AssignedProperties = Asset->ReplicationMap.ReplicatedObjects.Find(Object);
+		return AssignedProperties
+			? AssignedProperties->ClassPath
+			: FSoftClassPath{};
+	}
+
+	bool FPropertySelectionAssetModel::ContainsObjects(const TSet<FSoftObjectPath>& Objects) const
+	{
+		return ensure(Asset.IsValid())
+			&& Algo::AllOf(Objects, [this](const FSoftObjectPath& ObjectPath){ return Asset->ReplicationMap.ReplicatedObjects.Contains(ObjectPath); });
+	}
+
+	bool FPropertySelectionAssetModel::ContainsProperties(const FSoftObjectPath& Object, const TSet<FConcertPropertyChain>& Properties) const
+	{
+		if (!ensure(Asset.IsValid()))
+		{
+			return false;
+		}
+
+		const FReplicatedObjectInfo* ObjectInfo = Asset->ReplicationMap.ReplicatedObjects.Find(Object);
+		return ObjectInfo
+			&& Algo::AllOf(Properties, [ObjectInfo](const FConcertPropertyChain& Property){ return ObjectInfo->ReplicatedProperties.ReplicatedProperties.Contains(Property); });
+	}
+
+	bool FPropertySelectionAssetModel::ForEachReplicatedObject(TFunctionRef<EBreakBehavior(const FSoftObjectPath& Object)> Delegate) const
+	{
+		if (!ensure(Asset.IsValid()))
+		{
+			return false;
+		}
+
+		for (const TPair<FSoftObjectPath, FReplicatedObjectInfo>& ObjectMap: Asset->ReplicationMap.ReplicatedObjects)
+		{
+			if (Delegate(ObjectMap.Key) == EBreakBehavior::Break)
+			{
+				return true;
+			}
+		}
+
+		return !Asset->ReplicationMap.ReplicatedObjects.IsEmpty();
+	}
+
+	bool FPropertySelectionAssetModel::ForEachProperty(const FSoftObjectPath& Object, TFunctionRef<EBreakBehavior(const FConcertPropertyChain& Parent)> Delegate) const
+	{
+		if (!ensure(Asset.IsValid()))
+		{
+			return false;
+		}
+
+		const FReplicatedObjectInfo* AssignedProperties = Asset->ReplicationMap.ReplicatedObjects.Find(Object);
+		if (!AssignedProperties)
+		{
+			return false;
+		}
+
+		for (const FConcertPropertyChain& ReplicatedPropertyInfo : AssignedProperties->ReplicatedProperties.ReplicatedProperties)
+		{
+			if (Delegate(ReplicatedPropertyInfo) == EBreakBehavior::Break)
+			{
+				return true;
+			}
+		}
+		return !AssignedProperties->ReplicatedProperties.ReplicatedProperties.IsEmpty();
+	}
+
+	void FPropertySelectionAssetModel::AddObjects(TArrayView<UObject*> Objects)
+	{
+		if (!ensure(Asset.IsValid()) || Objects.IsEmpty())
+		{
+			return;
+		}
+
+		const FScopedTransaction Transaction(LOCTEXT("AddObjects", "Add replicated objects"));
+		Asset->Modify();
+		
+		TSet<UObject*> ObjectsNotAdded;
+		for (UObject* Object : Objects)
+		{
+			const FSoftObjectPath ObjectPath = Object;
+			if (ensureAlways(Object) && !Asset->ReplicationMap.ReplicatedObjects.Contains(ObjectPath))
+			{
+				FReplicatedObjectInfo& ObjectInfo = Asset->ReplicationMap.ReplicatedObjects.Add(ObjectPath);
+				ObjectInfo.ClassPath = Object->GetClass();
+				UMultiUserReplicationSettings::Get()->AddDefaultPropertiesFromSettings(ObjectInfo, *Object->GetClass());
+			}
+			else
+			{
+				ObjectsNotAdded.Add(Object);
+			}
+		}
+		
+		if (LIKELY(ObjectsNotAdded.IsEmpty()))
+		{
+			OnObjectsChangedDelegate.Broadcast(Objects, {}, EReplicatedObjectChangeReason::ChangedDirectly);
+		}
+		else if (ObjectsNotAdded.Num() < Objects.Num())
+		{
+			// Uncommon case so it is ok if it is suboptimal
+			TArray<UObject*> Added;
+			for (UObject* Object : Objects)
+			{
+				if (!ObjectsNotAdded.Contains(Object))
+				{
+					Added.Add(Object);
+				}
+			}
+			OnObjectsChangedDelegate.Broadcast(Added, {}, EReplicatedObjectChangeReason::ChangedDirectly);
+		}
+	}
+
+	void FPropertySelectionAssetModel::RemoveObjects(TArrayView<FSoftObjectPath> Objects)
+	{
+		if (!ensure(Asset.IsValid()) || Objects.IsEmpty())
+		{
+			return;
+		}
+
+		const FScopedTransaction Transaction(LOCTEXT("RemoveObjects", "Remove replicated objects"));
+		Asset->Modify();
+		
+		TSet<FSoftObjectPath> ObjectsNotRemoved;
+		for (const FSoftObjectPath& Object : Objects)
+		{
+			const bool bRemoved = Asset->ReplicationMap.ReplicatedObjects.Remove(Object) != 0;
+			if (!bRemoved)
+			{
+				ObjectsNotRemoved.Add(Object);
+			}
+		}
+		
+		if (LIKELY(ObjectsNotRemoved.IsEmpty()))
+		{
+			OnObjectsChangedDelegate.Broadcast({}, Objects, EReplicatedObjectChangeReason::ChangedDirectly);
+		}
+		else if (ObjectsNotRemoved.Num() < Objects.Num())
+		{
+			// Uncommon case so it is ok if it is suboptimal
+			TArray<FSoftObjectPath> Removed;
+			for (const FSoftObjectPath& Object : Objects)
+			{
+				if (!ObjectsNotRemoved.Contains(Object))
+				{
+					Removed.Add(Object);
+				}
+			}
+			OnObjectsChangedDelegate.Broadcast({}, Removed, EReplicatedObjectChangeReason::ChangedDirectly);
+		}
+	}
+
+	void FPropertySelectionAssetModel::AddProperties(const FSoftObjectPath& Object, TArrayView<FConcertPropertyChain> Properties)
+	{
+		if (!ensure(Asset.IsValid()))
+		{
+			return;
+		}
+		
+		const FScopedTransaction Transaction(LOCTEXT("AddProperties", "Add replicated properties"));
+		Asset->Modify();
+		
+		FReplicatedObjectInfo* AssignedProperties = Asset->ReplicationMap.ReplicatedObjects.Find(Object);
+		if (!AssignedProperties)
+		{
+			return;
+		}
+
+		UClass* Class = AssignedProperties->ClassPath.TryLoadClass<UObject>();
+		if (!Class)
+		{
+			return;
+		}
+
+		bool bAddedAtLeastOne = false;
+		TArray<FConcertPropertyChain>& ReplicatedProperties = AssignedProperties->ReplicatedProperties.ReplicatedProperties;
+		ReplicatedProperties.Reserve(ReplicatedProperties.Num() + Properties.Num());
+		for (const FConcertPropertyChain& AddedProperty : Properties)
+		{
+			bAddedAtLeastOne |= ReplicatedProperties.AddUnique(AddedProperty) != INDEX_NONE;
+
+			// Parent properties must also be added
+			// Not exactly efficient to iterate through the hierarchy for every removed item but it should be fine... Properties.Num() == 1 is the most common case
+			ConcertSyncCore::PropertyChain::ForEachReplicatableConcertProperty(*Class, [&bAddedAtLeastOne, &AssignedProperties, &AddedProperty](FConcertPropertyChain&& Property)
+			{
+				if (Property.IsParentOf(AddedProperty))
+				{
+					bAddedAtLeastOne |= AssignedProperties->ReplicatedProperties.ReplicatedProperties.AddUnique(Property) != INDEX_NONE;
+				}
+				return EBreakBehavior::Continue;
+			});
+		}
+
+		if (bAddedAtLeastOne)
+		{
+			OnPropertiesChangedDelegate.Broadcast();
+		}
+	}
+
+	void FPropertySelectionAssetModel::RemoveProperties(const FSoftObjectPath& Object, TArrayView<FConcertPropertyChain> Properties)
+	{
+		if (!ensure(Asset.IsValid()))
+		{
+			return;
+		}
+
+		const FScopedTransaction Transaction(LOCTEXT("RemoveProperties", "Remove replicated properties"));
+		Asset->Modify();
+		
+		FReplicatedObjectInfo* AssignedProperties = Asset->ReplicationMap.ReplicatedObjects.Find(Object);
+		if (!AssignedProperties)
+		{
+			return;
+		}
+
+		UClass* Class = AssignedProperties->ClassPath.TryLoadClass<UObject>();
+		int32 NumRemoved = 0;
+		for (const FConcertPropertyChain& RemovedProperty : Properties)
+		{
+			NumRemoved += AssignedProperties->ReplicatedProperties.ReplicatedProperties.Remove(RemovedProperty);
+
+			// Removal should not fail if the class is not available
+			if (!Class)
+			{
+				continue;
+			}
+
+			// Child properties must also be removed
+			// Not exactly efficient to iterate through the hierarchy for every removed item but it should be fine... Properties.Num() == 1 is the most common case
+			ConcertSyncCore::PropertyChain::ForEachReplicatableConcertProperty(*Class, [&NumRemoved, &AssignedProperties, &RemovedProperty](FConcertPropertyChain&& Property)
+			{
+				if (Property.IsChildOf(RemovedProperty))
+				{
+					NumRemoved += AssignedProperties->ReplicatedProperties.ReplicatedProperties.Remove(Property);
+				}
+				return EBreakBehavior::Continue;
+			});
+		}
+		
+		if (NumRemoved > 0)
+		{
+			OnPropertiesChangedDelegate.Broadcast();
+		}
+	}
+
+	bool FPropertySelectionAssetModel::MatchesContext(const FTransactionContext& InContext, const TArray<TPair<UObject*, FTransactionObjectEvent>>& TransactionObjectContexts) const
+	{
+		if (!ensure(Asset.IsValid()))
+		{
+			return false;
+		}
+		
+		for (const TPair<UObject*, FTransactionObjectEvent>& TransactionObjectPair : TransactionObjectContexts)
+		{
+			UObject* Object = TransactionObjectPair.Key;
+			while (Object != nullptr)
+			{
+				if (Object == Asset.Get())
+				{
+					return true;
+				}
+				Object = Object->GetOuter();
+			}
+		}
+
+		return false;
+	}
+
+	void FPropertySelectionAssetModel::PostUndo(bool bSuccess)
+	{
+		OnObjectsChangedDelegate.Broadcast({}, {}, EReplicatedObjectChangeReason::Transacted);
+		OnPropertiesChangedDelegate.Broadcast();
+	}
+
+	void FPropertySelectionAssetModel::PostRedo(bool bSuccess)
+	{
+		OnObjectsChangedDelegate.Broadcast({}, {}, EReplicatedObjectChangeReason::Transacted);
+		OnPropertiesChangedDelegate.Broadcast();
+	}
+}
+
+#undef LOCTEXT_NAMESPACE
\ No newline at end of file
diff --git a/Engine/Plugins/Developer/Concert/ConcertApp/MultiUserClient/Source/MultiUserReplicationEditor/Private/StreamEditor/Model/PropertySelectionAssetModel.h b/Engine/Plugins/Developer/Concert/ConcertApp/MultiUserClient/Source/MultiUserReplicationEditor/Private/StreamEditor/Model/PropertySelectionAssetModel.h
new file mode 100644
--- /dev/null
+++ b/Engine/Plugins/Developer/Concert/ConcertApp/MultiUserClient/Source/MultiUserReplicationEditor/Private/StreamEditor/Model/PropertySelectionAssetModel.h
@@ -0,0 +1,57 @@
+// Copyright Epic Games, Inc. All Rights Reserved.
+
+#pragma once
+
+#include "EditorUndoClient.h"
+#include "IEditableObjectToPropertiesModel.h"
+#include "MultiUserPropertyReplicationSelection.h"
+#include "UObject/WeakObjectPtrTemplates.h"
+
+class FTransactionObjectEvent;
+struct FTransactionContext;
+
+namespace UE::MultiUserReplicationEditor
+{
+	/** Edits a UMultiUserPropertyReplicationSelection object. */
+	class FPropertySelectionAssetModel
+		: public IEditableObjectToPropertiesModel
+		, public FSelfRegisteringEditorUndoClient
+	{
+	public:
+		
+		FPropertySelectionAssetModel(UMultiUserPropertyReplicationSelection& Object);
+		
+		//~ Begin IObjectToPropertiesModel Interface
+		virtual uint32 GetNumReplicatedObjects() const override;
+		virtual uint32 GetNumProperties(const FSoftObjectPath& Object) const override;
+		virtual FSoftClassPath GetObjectClass(const FSoftObjectPath& Object) const override;
+		virtual bool ContainsObjects(const TSet<FSoftObjectPath>& Objects) const override;
+		virtual bool ContainsProperties(const FSoftObjectPath& Object, const TSet<FConcertPropertyChain>& Properties) const override;
+		virtual bool ForEachReplicatedObject(TFunctionRef<EBreakBehavior(const FSoftObjectPath& Object)> Delegate) const override;
+		virtual bool ForEachProperty(const FSoftObjectPath& Object, TFunctionRef<EBreakBehavior(const FConcertPropertyChain& Parent)> Delegate) const override;
+		//~ End IObjectToPropertiesModel Interface
+		
+		//~ Begin IEditableObjectToPropertiesModel Interface
+		virtual void AddObjects(TArrayView<UObject*> Objects) override;
+		virtual void RemoveObjects(TArrayView<FSoftObjectPath> Objects) override;
+		virtual void AddProperties(const FSoftObjectPath&, TArrayView<FConcertPropertyChain> Properties) override;
+		virtual void RemoveProperties(const FSoftObjectPath&, TArrayView<FConcertPropertyChain> Properties) override;
+		virtual FOnObjectsChanged& OnObjectsChanged() override { return OnObjectsChangedDelegate; }
+		virtual FOnPropertiesChanged& OnPropertiesChanged() override { return OnPropertiesChangedDelegate; }
+		//~ End IEditableObjectToPropertiesModel Interface
+
+		//~ Begin FEditorUndoClient Interface
+		virtual bool MatchesContext(const FTransactionContext& InContext, const TArray<TPair<UObject*, FTransactionObjectEvent>>& TransactionObjectContexts) const override;
+		virtual void PostUndo(bool bSuccess) override;
+		virtual void PostRedo(bool bSuccess) override;
+		//~ End FEditorUndoClient Interface
+
+	private:
+
+		TWeakObjectPtr<UMultiUserPropertyReplicationSelection> Asset;
+
+		FOnObjectsChanged OnObjectsChangedDelegate;
+		FOnPropertiesChanged OnPropertiesChangedDelegate;
+	};
+}
+
diff --git a/Engine/Plugins/Developer/Concert/ConcertApp/MultiUserClient/Source/MultiUserReplicationEditor/Private/StreamEditor/ReplicationStreamAssetEditor.cpp b/Engine/Plugins/Developer/Concert/ConcertApp/MultiUserClient/Source/MultiUserReplicationEditor/Private/StreamEditor/ReplicationStreamAssetEditor.cpp
new file mode 100644
--- /dev/null
+++ b/Engine/Plugins/Developer/Concert/ConcertApp/MultiUserClient/Source/MultiUserReplicationEditor/Private/StreamEditor/ReplicationStreamAssetEditor.cpp
@@ -0,0 +1,20 @@
+// Copyright Epic Games, Inc. All Rights Reserved.
+
+#include "ReplicationStreamAssetEditor.h"
+
+#include "ReplicationStreamEditorToolkit.h"
+
+void UReplicationStreamAssetEditor::SetObjectToEdit(UObject* InObject)
+{
+	ObjectToEdit = InObject;
+}
+
+void UReplicationStreamAssetEditor::GetObjectsToEdit(TArray<UObject*>& OutObjectsToEdit)
+{
+	OutObjectsToEdit = { ObjectToEdit };
+}
+
+TSharedPtr<FBaseAssetToolkit> UReplicationStreamAssetEditor::CreateToolkit()
+{
+	return MakeShared<UE::MultiUserReplicationEditor::FReplicationStreamEditorToolkit>(this);
+}
diff --git a/Engine/Plugins/Developer/Concert/ConcertApp/MultiUserClient/Source/MultiUserReplicationEditor/Private/StreamEditor/ReplicationStreamAssetEditor.h b/Engine/Plugins/Developer/Concert/ConcertApp/MultiUserClient/Source/MultiUserReplicationEditor/Private/StreamEditor/ReplicationStreamAssetEditor.h
new file mode 100644
--- /dev/null
+++ b/Engine/Plugins/Developer/Concert/ConcertApp/MultiUserClient/Source/MultiUserReplicationEditor/Private/StreamEditor/ReplicationStreamAssetEditor.h
@@ -0,0 +1,30 @@
+// Copyright Epic Games, Inc. All Rights Reserved.
+
+#pragma once
+
+#include "Tools/UAssetEditor.h"
+#include "ReplicationStreamAssetEditor.generated.h"
+
+/**
+ * 
+ */
+UCLASS()
+class MULTIUSERREPLICATIONEDITOR_API UReplicationStreamAssetEditor : public UAssetEditor
+{
+	GENERATED_BODY()
+public:
+
+	void SetObjectToEdit(UObject* InObject);
+
+protected:
+
+	//~ Begin UAssetEditor Interface
+	virtual void GetObjectsToEdit(TArray<UObject*>& OutObjectsToEdit) override;
+	virtual TSharedPtr<FBaseAssetToolkit> CreateToolkit() override;
+	//~ End UAssetEditor Interface
+
+private:
+	
+	UPROPERTY(Transient)
+	TObjectPtr<UObject> ObjectToEdit;
+};
diff --git a/Engine/Plugins/Developer/Concert/ConcertApp/MultiUserClient/Source/MultiUserReplicationEditor/Private/StreamEditor/ReplicationStreamEditorToolkit.cpp b/Engine/Plugins/Developer/Concert/ConcertApp/MultiUserClient/Source/MultiUserReplicationEditor/Private/StreamEditor/ReplicationStreamEditorToolkit.cpp
new file mode 100644
--- /dev/null
+++ b/Engine/Plugins/Developer/Concert/ConcertApp/MultiUserClient/Source/MultiUserReplicationEditor/Private/StreamEditor/ReplicationStreamEditorToolkit.cpp
@@ -0,0 +1,79 @@
+// Copyright Epic Games, Inc. All Rights Reserved.
+
+#include "ReplicationStreamEditorToolkit.h"
+
+#include "MultiUserReplicationStreamAsset.h"
+#include "StreamEditor/View/SReplicationStreamEditor.h"
+
+#include "Framework/Docking/TabManager.h"
+#include "Framework/Docking/WorkspaceItem.h"
+#include "Widgets/Docking/SDockTab.h"
+
+#define LOCTEXT_NAMESPACE "FReplicationStreamEditorToolkit"
+
+namespace UE::MultiUserReplicationEditor
+{
+	const FName FReplicationStreamEditorToolkit::ContentTabId(TEXT("ReplicationStreamAssetEditor_Content"));
+	
+	FReplicationStreamEditorToolkit::FReplicationStreamEditorToolkit(UAssetEditor* InOwningAssetEditor)
+		: FBaseAssetToolkit(InOwningAssetEditor)
+	{
+		LayoutAppendix = TEXT("ReplicationStreamEditor_v1");
+		const FString LayoutString = TEXT("Standalone_Layout_") + LayoutAppendix;
+		StandaloneDefaultLayout = FTabManager::NewLayout(FName(LayoutString))
+			->AddArea
+			(
+				FTabManager::NewPrimaryArea()->SetOrientation(Orient_Vertical)
+				->Split
+				(
+					FTabManager::NewStack()
+					->SetSizeCoefficient(1.f)
+					->SetHideTabWell(true)
+					->AddTab(ContentTabId, ETabState::OpenedTab)
+				)
+			);
+	}
+
+	void FReplicationStreamEditorToolkit::CreateWidgets()
+	{
+		// Do not call Super because we do not want to create the viewport and details panel from FBaseAssetToolkit
+		
+	}
+
+	void FReplicationStreamEditorToolkit::SetEditingObject(UObject* InObject)
+	{
+		// Do not call Super because we do not wish to use the details view from FBaseAssetToolkit
+	}
+
+	void FReplicationStreamEditorToolkit::RegisterTabSpawners(const TSharedRef<FTabManager>& InTabManager)
+	{
+		WorkspaceMenuCategory = InTabManager->AddLocalWorkspaceMenuCategory(LOCTEXT("WorkspaceMenu_ReplicationStreamAssetEditor", "Replication Strean Asset Editor"));
+		FBaseAssetToolkit::RegisterTabSpawners(InTabManager);
+
+		InTabManager->RegisterTabSpawner(ContentTabId, FOnSpawnTab::CreateSP(this, &FReplicationStreamEditorToolkit::SpawnTab_Content))
+			.SetDisplayName(LOCTEXT("ConfigTab", "Config"))
+			.SetGroup(WorkspaceMenuCategory.ToSharedRef());
+	}
+
+	void FReplicationStreamEditorToolkit::UnregisterTabSpawners(const TSharedRef<FTabManager>& InTabManager)
+	{
+		FAssetEditorToolkit::UnregisterTabSpawners(InTabManager);
+		InTabManager->UnregisterTabSpawner(ContentTabId);
+	}
+
+	UMultiUserReplicationStreamAsset* FReplicationStreamEditorToolkit::GetEditedStreamAsset() const
+	{
+		return CastChecked<UMultiUserReplicationStreamAsset>(GetEditingObject());
+	}
+
+	TSharedRef<SDockTab> FReplicationStreamEditorToolkit::SpawnTab_Content(const FSpawnTabArgs& SpawnTabArgs)
+	{
+		return SNew(SDockTab)
+			.Label(LOCTEXT("BaseDetailsTitle", "Details"))
+			[
+				SAssignNew(RootWidget, SReplicationStreamEditor, *GetEditedStreamAsset())
+			];
+	}
+}
+
+#undef LOCTEXT_NAMESPACE
\ No newline at end of file
diff --git a/Engine/Plugins/Developer/Concert/ConcertApp/MultiUserClient/Source/MultiUserReplicationEditor/Private/StreamEditor/ReplicationStreamEditorToolkit.h b/Engine/Plugins/Developer/Concert/ConcertApp/MultiUserClient/Source/MultiUserReplicationEditor/Private/StreamEditor/ReplicationStreamEditorToolkit.h
new file mode 100644
--- /dev/null
+++ b/Engine/Plugins/Developer/Concert/ConcertApp/MultiUserClient/Source/MultiUserReplicationEditor/Private/StreamEditor/ReplicationStreamEditorToolkit.h
@@ -0,0 +1,39 @@
+// Copyright Epic Games, Inc. All Rights Reserved.
+
+#pragma once
+
+#include "Tools/BaseAssetToolkit.h"
+
+class UMultiUserReplicationStreamAsset;
+
+namespace UE::MultiUserReplicationEditor
+{
+	class SReplicationStreamEditor;
+
+	class FReplicationStreamEditorToolkit : public FBaseAssetToolkit
+	{
+	public:
+		
+		static const FName ContentTabId;
+
+		FReplicationStreamEditorToolkit(UAssetEditor* InOwningAssetEditor);
+
+		//~ Begin FAssetEditorToolkit Interface
+		virtual void CreateWidgets();
+		virtual void SetEditingObject(UObject* InObject) override;
+		virtual void RegisterTabSpawners(const TSharedRef<FTabManager>& InTabManager) override;
+		virtual void UnregisterTabSpawners(const TSharedRef<FTabManager>& InTabManager) override;
+		virtual FName GetToolkitFName() const override { return TEXT("ReplicationStreamEditorToolkit"); }
+		//~ End FAssetEditorToolkit Interface
+
+	private:
+
+		/** Root widget editing the asset */
+		TSharedPtr<SReplicationStreamEditor> RootWidget;
+		
+		UMultiUserReplicationStreamAsset* GetEditedStreamAsset() const;
+		
+		TSharedRef<SDockTab> SpawnTab_Content(const FSpawnTabArgs& SpawnTabArgs);
+	};
+}
+
diff --git a/Engine/Plugins/Developer/Concert/ConcertApp/MultiUserClient/Source/MultiUserReplicationEditor/Private/StreamEditor/View/ObjectEditor/FakeObjectToPropertiesEditorModel.cpp b/Engine/Plugins/Developer/Concert/ConcertApp/MultiUserClient/Source/MultiUserReplicationEditor/Private/StreamEditor/View/ObjectEditor/FakeObjectToPropertiesEditorModel.cpp
new file mode 100644
--- /dev/null
+++ b/Engine/Plugins/Developer/Concert/ConcertApp/MultiUserClient/Source/MultiUserReplicationEditor/Private/StreamEditor/View/ObjectEditor/FakeObjectToPropertiesEditorModel.cpp
@@ -0,0 +1,20 @@
+// Copyright Epic Games, Inc. All Rights Reserved.
+
+#include "FakeObjectToPropertiesEditorModel.h"
+
+#include "StreamEditor/Model/Property/IPropertySelectionSourceModel.h"
+#include "StreamEditor/Model/Property/IPropertySourceModel.h"
+#include "Replication/PropertyChainUtils.h"
+
+namespace UE::MultiUserReplicationEditor
+{
+	bool FFakeObjectToPropertiesEditorModel::ForEachProperty(const FSoftObjectPath& Object, TFunctionRef<EBreakBehavior(const FConcertPropertyChain& Property)> Delegate) const
+	{
+		PropertySelectionSource->GetPropertySource(GetObjectClass(Object))
+			->EnumerateSelectableItems([&Delegate](const FSelectablePropertyInfo& PropertyInfo)
+			{
+				return Delegate(PropertyInfo.Property);
+			});
+		return true;
+	}
+}
diff --git a/Engine/Plugins/Developer/Concert/ConcertApp/MultiUserClient/Source/MultiUserReplicationEditor/Private/StreamEditor/View/ObjectEditor/FakeObjectToPropertiesEditorModel.h b/Engine/Plugins/Developer/Concert/ConcertApp/MultiUserClient/Source/MultiUserReplicationEditor/Private/StreamEditor/View/ObjectEditor/FakeObjectToPropertiesEditorModel.h
new file mode 100644
--- /dev/null
+++ b/Engine/Plugins/Developer/Concert/ConcertApp/MultiUserClient/Source/MultiUserReplicationEditor/Private/StreamEditor/View/ObjectEditor/FakeObjectToPropertiesEditorModel.h
@@ -0,0 +1,48 @@
+// Copyright Epic Games, Inc. All Rights Reserved.
+
+#pragma once
+
+#include "StreamEditor/Model/IObjectToPropertiesModel.h"
+
+namespace UE::MultiUserReplicationEditor
+{
+	class IPropertySelectionSourceModel;
+	/**
+	 * This model is used by SPropertyReplicationSelectionEditor.
+	 * 
+	 * It forwards most functions except for ForEachProperty, which fakes to the SPropertyReplicationSelectionView
+	 * that all properties are selected. The properties are added to the real model via checking and unchecking the box.
+	 * 
+	 * This allows the SPropertyReplicationSelectionView to continue to show the properties it is being told are in the model,
+	 * which is important e.g. for the connection view (when a client is connected to a session and views the properties being sent).
+	 */
+	class FFakeObjectToPropertiesEditorModel : public IObjectToPropertiesModel
+	{
+	public:
+
+		FFakeObjectToPropertiesEditorModel(TSharedRef<IObjectToPropertiesModel> RealModel, TSharedRef<IPropertySelectionSourceModel> PropertySelectionSource)
+			: RealModel(MoveTemp(RealModel))
+			, PropertySelectionSource(MoveTemp(PropertySelectionSource))
+		{}
+
+		//~ Begin IObjectToPropertiesModel Interface
+		virtual uint32 GetNumReplicatedObjects() const override { return RealModel->GetNumReplicatedObjects(); }
+		virtual uint32 GetNumProperties(const FSoftObjectPath& Object) const override { return RealModel->GetNumProperties(Object); }
+		virtual FSoftClassPath GetObjectClass(const FSoftObjectPath& Object) const override { return RealModel->GetObjectClass(Object); }
+		virtual bool ContainsObjects(const TSet<FSoftObjectPath>& Objects) const override { return RealModel->ContainsObjects(Objects); }
+		// This function is by SPropertyReplicationSelectionEditor and not used by SPropertyReplicationSelectionView so forward it normally
+		virtual bool ContainsProperties(const FSoftObjectPath& Object, const TSet<FConcertPropertyChain>& Properties) const override { return RealModel->ContainsProperties(Object, Properties); }
+		virtual bool ForEachReplicatedObject(TFunctionRef<EBreakBehavior(const FSoftObjectPath& Object)> Delegate) const override { return RealModel->ForEachReplicatedObject(Delegate); }
+		virtual bool ForEachProperty(const FSoftObjectPath& Object, TFunctionRef<EBreakBehavior(const FConcertPropertyChain& Property)> Delegate) const override;
+		//~ End IObjectToPropertiesModel Interface
+
+	private:
+
+		/** The real model which is used to implement all the other functions. */
+		TSharedRef<IObjectToPropertiesModel> RealModel;
+
+		/** Determines the properties that can be selected. */
+		TSharedRef<IPropertySelectionSourceModel> PropertySelectionSource;
+	};
+}
+
diff --git a/Engine/Plugins/Developer/Concert/ConcertApp/MultiUserClient/Source/MultiUserReplicationEditor/Private/StreamEditor/View/ObjectEditor/SPropertyReplicationSelectionEditor.cpp b/Engine/Plugins/Developer/Concert/ConcertApp/MultiUserClient/Source/MultiUserReplicationEditor/Private/StreamEditor/View/ObjectEditor/SPropertyReplicationSelectionEditor.cpp
new file mode 100644
--- /dev/null
+++ b/Engine/Plugins/Developer/Concert/ConcertApp/MultiUserClient/Source/MultiUserReplicationEditor/Private/StreamEditor/View/ObjectEditor/SPropertyReplicationSelectionEditor.cpp
@@ -0,0 +1,216 @@
+// Copyright Epic Games, Inc. All Rights Reserved.
+
+#include "SPropertyReplicationSelectionEditor.h"
+
+#include "FakeObjectToPropertiesEditorModel.h"
+#include "StreamEditor/Model/DisplayUtils.h"
+#include "StreamEditor/Model/IEditableObjectToPropertiesModel.h"
+#include "StreamEditor/Model/Object/IObjectSelectionSourceModel.h"
+#include "Model/Item/SourceModelBuilders.h"
+#include "StreamEditor/Model/Property/IPropertySelectionSourceModel.h"
+#include "StreamEditor/View/ObjectViewer/ReplicatedObjectData.h"
+
+#include "Algo/AllOf.h"
+#include "Algo/AnyOf.h"
+#include "StreamEditor/View/ObjectViewer/ReplicatedPropertyData.h"
+#include "Widgets/SBoxPanel.h"
+
+#define LOCTEXT_NAMESPACE "SPropertyReplicationSelectionEditor"
+
+namespace UE::MultiUserReplicationEditor
+{
+	void SPropertyReplicationSelectionEditor::Construct(const FArguments& InArgs,
+		TSharedRef<IEditableObjectToPropertiesModel> InPropertiesModel,
+		TSharedRef<IObjectSelectionSourceModel> InObjectSelectionSource,
+		TSharedRef<IPropertySelectionSourceModel> InPropertySelectionSource)
+	{
+		ObjectSelectionSource = MoveTemp(InObjectSelectionSource);
+		PropertySelectionSource = MoveTemp(InPropertySelectionSource);
+		
+		EditablePropertiesModel = MoveTemp(InPropertiesModel);
+		EditablePropertiesModel->OnObjectsChanged().AddSP(this, &SPropertyReplicationSelectionEditor::OnObjectsChanged);
+		EditablePropertiesModel->OnPropertiesChanged().AddSP(this, &SPropertyReplicationSelectionEditor::OnPropertiesChanged);
+		PropertiesModelAdapter = MakeShared<FFakeObjectToPropertiesEditorModel>(EditablePropertiesModel.ToSharedRef(), PropertySelectionSource.ToSharedRef());
+		
+		SPropertyReplicationSelectionViewer::Construct(
+			SPropertyReplicationSelectionViewer::FArguments()
+				.AdditionalPropertyColumns({ ReplicationPropertyColumns::ReplicatesColumns(SharedThis(this), EditablePropertiesModel.ToSharedRef()) })
+				.OnDeleteObjects(this, &SPropertyReplicationSelectionEditor::OnDeleteObjects)
+				.OnObjectsContextMenuOpening(this, &SPropertyReplicationSelectionEditor::OnObjectsContextMenuOpening)
+				.SortPropertyRowPredicate(this, &SPropertyReplicationSelectionEditor::SortBySelectionThenByName_PropertyPredicate)
+				.LeftOfObjectSearchBar()
+				[
+					BuildRootAddObjectWidgets()
+				]
+				.LeftOfPropertySearchBar()
+				[
+					SAssignNew(AddPropertyWidgetContainer, SHorizontalBox)
+				],
+			PropertiesModelAdapter.ToSharedRef()
+			);
+	}
+
+	void SPropertyReplicationSelectionEditor::OnObjectsChanged(TArrayView<UObject*> AddedObjects, TArrayView<FSoftObjectPath> RemovedObjects, EReplicatedObjectChangeReason ChangeReason)
+	{
+		RefreshObjectData();
+
+		if (ChangeReason == EReplicatedObjectChangeReason::ChangedDirectly && !AddedObjects.IsEmpty())
+		{
+			for (const UObject* Object : AddedObjects)
+			{
+				// Example: Right-click component > Add Outer. You want to see the both the newly added actor and the component that you right-clicked.
+				bool bHasChildren = false;
+				GetObjectChildren(Object, [&bHasChildren](auto) mutable { bHasChildren = true; });
+				if (bHasChildren)
+				{
+					SetObjectsExpanded(TArray{ Object }, true);
+				}
+
+				// Example: Add a component to an actor. You want to see the object you just created.
+				if (const FSoftObjectPath ParentObject = GetParent(Object); !ParentObject.IsNull())
+				{
+					SetObjectsExpanded(TArray{ ParentObject }, true);
+				}
+			}
+			
+			// After adding an object, you usually instantly want to edit it so make the properties show up in the property view. Do this after making sure the item is visible.
+			SetSelectedObjects(AddedObjects, true);
+		}
+	}
+
+	void SPropertyReplicationSelectionEditor::OnPropertiesChanged()
+	{
+		RefreshPropertyData();
+	}
+
+	TSharedRef<SWidget> SPropertyReplicationSelectionEditor::BuildRootAddObjectWidgets() const
+	{
+		using namespace ConcertSharedSlate;
+		
+		const TSharedRef<SHorizontalBox> Root = SNew(SHorizontalBox);
+		const FSourceModelBuilders<FSelectableObjectInfo>::FItemPickerArgs Args = MakeObjectSourceBuilderArgs();
+		for (const TSourceSelectionCategory<FSelectableObjectInfo>& Category : ObjectSelectionSource->GetRootSources())
+		{
+			Root->AddSlot()
+				.AutoWidth()
+				[
+					FSourceModelBuilders<FSelectableObjectInfo>::BuildCategory(Category, Args)	
+				];
+		}
+		return Root;
+	}
+
+	void SPropertyReplicationSelectionEditor::OnObjectsSelectedForAdding(TArray<FSelectableObjectInfo> ObjectsToAdd) const
+	{
+		TArray<UObject*> Objects;
+		Algo::TransformIf(
+			ObjectsToAdd,
+			Objects,
+			[](const FSelectableObjectInfo& SelectableObject) { return SelectableObject.Object.IsValid(); },
+			[](const FSelectableObjectInfo& SelectableObject){ return SelectableObject.Object.Get(); }
+			);
+		
+		EditablePropertiesModel->AddObjects(Objects);
+	}
+
+	void SPropertyReplicationSelectionEditor::OnDeleteObjects(const TArray<TSharedPtr<FReplicatedObjectData>>& ObjectsToDelete) const
+	{
+		TArray<FSoftObjectPath> DeleteObjectPaths;
+		Algo::Transform(ObjectsToDelete, DeleteObjectPaths, [](const TSharedPtr<FReplicatedObjectData>& Data) { return Data->GetObjectPath(); });
+		EditablePropertiesModel->RemoveObjects(DeleteObjectPaths);
+	}
+
+	TSharedPtr<SWidget> SPropertyReplicationSelectionEditor::OnObjectsContextMenuOpening() const
+	{
+		FMenuBuilder MenuBuilder(true, nullptr);
+		AddObjectSourceContextMenuOptions(MenuBuilder);
+		MenuBuilder.AddMenuEntry(
+				LOCTEXT("DeleteItems", "Delete"),
+				FText::GetEmpty(),
+				FSlateIcon(),
+				FUIAction(FExecuteAction::CreateSP(this, &SPropertyReplicationSelectionEditor::OnDeleteObjects_PassByValue, GetSelectedObjects())),
+				NAME_None,
+				EUserInterfaceActionType::Button
+			);
+		
+		return MenuBuilder.MakeWidget();
+	}
+
+	void SPropertyReplicationSelectionEditor::AddObjectSourceContextMenuOptions(FMenuBuilder& MenuBuilder) const
+	{
+		using namespace ConcertSharedSlate;
+		
+		// Context menu generation is only supported for single items
+		const TArray<TSharedPtr<FReplicatedObjectData>> SelectedObjects = GetSelectedObjects();
+		if (SelectedObjects.Num() != 1)
+		{
+			return;
+		}
+		
+		const TArray<TSharedRef<IObjectSourceModel>> ContextMenuOptions = ObjectSelectionSource->GetContextMenuOptions(SelectedObjects[0]->GetObjectPath());
+		if (ContextMenuOptions.IsEmpty())
+		{
+			return;
+		}
+		
+		const FSourceModelBuilders<FSelectableObjectInfo>::FItemPickerArgs Args = MakeObjectSourceBuilderArgs();
+		for (const TSharedRef<IObjectSourceModel>& SourceModel : ContextMenuOptions)
+		{
+			bool bCanAddAnyObject = false; 
+			SourceModel->EnumerateSelectableItems([this, &bCanAddAnyObject](const FSelectableObjectInfo& SelectableOption) mutable
+			{
+				bCanAddAnyObject |= !EditablePropertiesModel->ContainsObjects({ SelectableOption.Object.Get() });
+				return bCanAddAnyObject ? EBreakBehavior::Break : EBreakBehavior::Continue;
+			});
+
+			// Skip showing context menu options which will not add anything new.
+			if (bCanAddAnyObject)
+			{
+				FSourceModelBuilders<FSelectableObjectInfo>::AddOptionToMenu(SourceModel, Args, MenuBuilder);
+			}
+		}
+		
+		MenuBuilder.AddSeparator();
+	}
+
+	ConcertSharedSlate::FSourceModelBuilders<FSelectableObjectInfo>::FItemPickerArgs SPropertyReplicationSelectionEditor::MakeObjectSourceBuilderArgs() const
+	{
+		using namespace ConcertSharedSlate;
+		
+		return FSourceModelBuilders<FSelectableObjectInfo>::FItemPickerArgs
+		{
+			FSourceModelBuilders<FSelectableObjectInfo>::FOnItemsSelected::CreateSP(this, &SPropertyReplicationSelectionEditor::OnObjectsSelectedForAdding),
+			FSourceModelBuilders<FSelectableObjectInfo>::FGetItemDisplayString::CreateLambda([](const FSelectableObjectInfo& Item)
+			{
+				return Item.Object.IsValid() ? DisplayUtils::GetObjectDisplayString(*Item.Object.Get()) : TEXT("");
+			}),
+			FSourceModelBuilders<FSelectableObjectInfo>::FGetItemIcon::CreateLambda([](const FSelectableObjectInfo& Item)
+			{
+				return Item.Object.IsValid() ? DisplayUtils::GetObjectIcon(*Item.Object.Get()) : FSlateIcon();
+			}),
+			FSourceModelBuilders<FSelectableObjectInfo>::FIsItemSelected::CreateLambda([this](const FSelectableObjectInfo& Item)
+			{
+				return EditablePropertiesModel->ContainsObjects({ Item.Object.Get() } );
+			})
+		}; 
+	}
+
+	bool SPropertyReplicationSelectionEditor::SortBySelectionThenByName_PropertyPredicate(const TSharedPtr<FReplicatedPropertyData>& Left, const TSharedPtr<FReplicatedPropertyData>& Right) const
+	{
+		const TArray<TSharedPtr<FReplicatedObjectData>> SelectedObjects = GetSelectedObjects();
+		const ECheckBoxState LeftCheckboxState = ReplicationPropertyColumns::GetPropertyCheckboxStateBasedOnSelection(*Left, SelectedObjects, *EditablePropertiesModel);
+		const ECheckBoxState RightCheckboxState = ReplicationPropertyColumns::GetPropertyCheckboxStateBasedOnSelection(*Right, SelectedObjects, *EditablePropertiesModel);
+		
+		// Secondary sort by name
+		if (LeftCheckboxState == RightCheckboxState)
+		{
+			return DisplayUtils::GetPropertyDisplayString(Left->GetProperty()) < DisplayUtils::GetPropertyDisplayString(Right->GetProperty());
+		}
+
+		// Selected properties should appear first
+		return LeftCheckboxState == ECheckBoxState::Checked
+			&& (RightCheckboxState == ECheckBoxState::Unchecked || RightCheckboxState == ECheckBoxState::Undetermined);
+	}
+}
+
+#undef LOCTEXT_NAMESPACE
\ No newline at end of file
diff --git a/Engine/Plugins/Developer/Concert/ConcertApp/MultiUserClient/Source/MultiUserReplicationEditor/Private/StreamEditor/View/ObjectEditor/SPropertyReplicationSelectionEditor.h b/Engine/Plugins/Developer/Concert/ConcertApp/MultiUserClient/Source/MultiUserReplicationEditor/Private/StreamEditor/View/ObjectEditor/SPropertyReplicationSelectionEditor.h
new file mode 100644
--- /dev/null
+++ b/Engine/Plugins/Developer/Concert/ConcertApp/MultiUserClient/Source/MultiUserReplicationEditor/Private/StreamEditor/View/ObjectEditor/SPropertyReplicationSelectionEditor.h
@@ -0,0 +1,79 @@
+// Copyright Epic Games, Inc. All Rights Reserved.
+
+#pragma once
+
+#include "Model/Item/SourceModelBuilders.h"
+#include "StreamEditor/Model/Property/IPropertySourceModel.h"
+#include "StreamEditor/View/ObjectViewer/SPropertyReplicationSelectionViewer.h"
+#include "Widgets/DeclarativeSyntaxSupport.h"
+
+class SHorizontalBox;
+
+namespace UE::MultiUserReplicationEditor
+{
+	class IEditableObjectToPropertiesModel;
+	class FFakeObjectToPropertiesEditorModel;
+	class IPropertySelectionSourceModel;
+	class IObjectSelectionSourceModel;
+	class IObjectToPropertiesModel;
+	
+	enum class EReplicatedObjectChangeReason : uint8;
+	
+	struct FSelectableObjectInfo;
+	
+	/**
+	 * Root widget for editing UMultiUserPropertyReplicationSelection.
+	 * This widget extends SPropertyReplicationSelectionViewer by using IEditableObjectToPropertiesModel for editing.
+	 */
+	class MULTIUSERREPLICATIONEDITOR_API SPropertyReplicationSelectionEditor : public SPropertyReplicationSelectionViewer
+	{
+	public:
+		
+		SLATE_BEGIN_ARGS(SPropertyReplicationSelectionEditor)
+		{}
+		SLATE_END_ARGS()
+
+		void Construct(const FArguments& InArgs,
+			TSharedRef<IEditableObjectToPropertiesModel> InPropertiesModel,
+			TSharedRef<IObjectSelectionSourceModel> InObjectSelectionSource,
+		    TSharedRef<IPropertySelectionSourceModel> InPropertySelectionSource
+		);
+
+	private:
+
+		/** For reading and writting to the edited asset */
+		TSharedPtr<IEditableObjectToPropertiesModel> EditablePropertiesModel;
+		/**
+		 * Fakes to SPropertyReplicationSelectionViewer that all UClass properties are contained.
+		 * We inject checkboxes to SPropertyReplicationSelectionViewer which do the actual adding and removing.
+		 */
+		TSharedPtr<FFakeObjectToPropertiesEditorModel> PropertiesModelAdapter;
+		
+		/** For deciding which objects can be added to EditablePropertiesModel. */
+		TSharedPtr<IObjectSelectionSourceModel> ObjectSelectionSource;
+		/** For deciding which properties can be added to EditablePropertiesModel. */
+		TSharedPtr<IPropertySelectionSourceModel> PropertySelectionSource;
+
+		/** Contains the widgets generated for PropertySelectionSource (since they depend on the selected objects). */
+		TSharedPtr<SHorizontalBox> AddPropertyWidgetContainer;
+		
+		void OnObjectsChanged(TArrayView<UObject*> AddedObjects, TArrayView<FSoftObjectPath> RemovedObjects, EReplicatedObjectChangeReason ChangeReason);
+		void OnPropertiesChanged();
+
+		// Customizing left objects search bar
+		TSharedRef<SWidget> BuildRootAddObjectWidgets() const;
+		void OnObjectsSelectedForAdding(TArray<FSelectableObjectInfo> ObjectsToAdd) const;
+
+		// Respond to object editing events
+		void OnDeleteObjects_PassByValue(TArray<TSharedPtr<FReplicatedObjectData>> CopiedObjectsToDelete) const { OnDeleteObjects(CopiedObjectsToDelete); }
+		void OnDeleteObjects(const TArray<TSharedPtr<FReplicatedObjectData>>& ObjectsToDelete) const;
+		TSharedPtr<SWidget> OnObjectsContextMenuOpening() const;
+		void AddObjectSourceContextMenuOptions(FMenuBuilder& MenuBuilder) const;
+		
+		// Utils for building item source widgets
+		ConcertSharedSlate::FSourceModelBuilders<FSelectableObjectInfo>::FItemPickerArgs MakeObjectSourceBuilderArgs() const;
+
+		/** Util for sorting property data by*/
+		bool SortBySelectionThenByName_PropertyPredicate(const TSharedPtr<FReplicatedPropertyData>& Left, const TSharedPtr<FReplicatedPropertyData>& Right) const;
+	};
+}
diff --git a/Engine/Plugins/Developer/Concert/ConcertApp/MultiUserClient/Source/MultiUserReplicationEditor/Private/StreamEditor/View/ObjectViewer/ReplicatedObjectData.h b/Engine/Plugins/Developer/Concert/ConcertApp/MultiUserClient/Source/MultiUserReplicationEditor/Private/StreamEditor/View/ObjectViewer/ReplicatedObjectData.h
new file mode 100644
--- /dev/null
+++ b/Engine/Plugins/Developer/Concert/ConcertApp/MultiUserClient/Source/MultiUserReplicationEditor/Private/StreamEditor/View/ObjectViewer/ReplicatedObjectData.h
@@ -0,0 +1,33 @@
+// Copyright Epic Games, Inc. All Rights Reserved.
+
+#pragma once
+
+#include "UObject/SoftObjectPath.h"
+
+namespace UE::MultiUserReplicationEditor
+{
+	/** Instanced for each object row in SPropertyReplicationSelectionViewer.*/
+	class FReplicatedObjectData
+	{
+	public:
+
+		FReplicatedObjectData(FSoftObjectPath ObjectPath)
+			: ObjectPath(MoveTemp(ObjectPath))
+		{}
+		
+		const FSoftObjectPath& GetObjectPath() const { return ObjectPath; }
+		void SetObjectPath(FSoftObjectPath InObjectPath) { ObjectPath = MoveTemp(InObjectPath); }
+
+	private:
+
+		/** The object identified by the */
+		FSoftObjectPath ObjectPath;
+	};
+
+	/** Instanced for each object row SPropertyReplicationSelectionEditor. */
+	class FReplicatedObjectData_Editor : public FReplicatedObjectData
+	{
+		// Empty and unused for now. I just want to highlight that if you want custom data for SPropertyReplicationSelectionEditor,
+		// you should add it here and override SPropertyReplicationSelectionViewer::AllocateObjectData.
+	};
+}
\ No newline at end of file
diff --git a/Engine/Plugins/Developer/Concert/ConcertApp/MultiUserClient/Source/MultiUserReplicationEditor/Private/StreamEditor/View/ObjectViewer/ReplicatedPropertyData.h b/Engine/Plugins/Developer/Concert/ConcertApp/MultiUserClient/Source/MultiUserReplicationEditor/Private/StreamEditor/View/ObjectViewer/ReplicatedPropertyData.h
new file mode 100644
--- /dev/null
+++ b/Engine/Plugins/Developer/Concert/ConcertApp/MultiUserClient/Source/MultiUserReplicationEditor/Private/StreamEditor/View/ObjectViewer/ReplicatedPropertyData.h
@@ -0,0 +1,37 @@
+// Copyright Epic Games, Inc. All Rights Reserved.
+
+#pragma once
+
+#include "Replication/Data/ConcertPropertySelection.h"
+
+namespace UE::MultiUserReplicationEditor
+{
+	/** Instanced for each property row in SPropertyReplicationSelectionViewer.*/
+	class FReplicatedPropertyData
+	{
+	public:
+		
+		FReplicatedPropertyData(FSoftClassPath OwningClass, FConcertPropertyChain Object)
+			: OwningClass(MoveTemp(OwningClass))
+			, Property(MoveTemp(Object))
+		{}
+		
+		const FConcertPropertyChain& GetProperty() const { return Property; }
+		const FSoftClassPath& GetOwningClass() const { return OwningClass; }
+
+	private:
+
+		/** The class with which the FProperty can be determined. */
+		FSoftClassPath OwningClass;
+		
+		/** The object identified by the */
+		FConcertPropertyChain Property;
+	};
+
+	/** Instanced for each property row SPropertyReplicationSelectionEditor. */
+	class FReplicatedPropertyData_Editor : public FReplicatedPropertyData
+	{
+		// Empty and unused for now. I just want to highlight that if you want custom data for SPropertyReplicationSelectionEditor,
+		// you should add it here and override SPropertyReplicationSelectionViewer::AllocatePropertyData.
+	};
+}
\ No newline at end of file
diff --git a/Engine/Plugins/Developer/Concert/ConcertApp/MultiUserClient/Source/MultiUserReplicationEditor/Private/StreamEditor/View/ObjectViewer/ReplicationColumn.h b/Engine/Plugins/Developer/Concert/ConcertApp/MultiUserClient/Source/MultiUserReplicationEditor/Private/StreamEditor/View/ObjectViewer/ReplicationColumn.h
new file mode 100644
--- /dev/null
+++ b/Engine/Plugins/Developer/Concert/ConcertApp/MultiUserClient/Source/MultiUserReplicationEditor/Private/StreamEditor/View/ObjectViewer/ReplicationColumn.h
@@ -0,0 +1,72 @@
+// Copyright Epic Games, Inc. All Rights Reserved.
+
+#pragma once
+
+#include "Widgets/SNullWidget.h"
+#include "Widgets/Views/SHeaderRow.h"
+
+namespace UE::MultiUserReplicationEditor
+{
+	template<typename TListItemType>
+	class TReplicationColumn : public SHeaderRow::FColumn
+	{
+	public:
+
+		struct FBuildArgs
+		{
+			TSharedPtr<FText> HighlightText;
+			TListItemType RowData;
+		};
+
+		/** Generates the widget content for a column in a row */
+		DECLARE_DELEGATE_RetVal_OneParam(TSharedRef<SWidget>, FGenerateColumnWidget,
+			const FBuildArgs& InArgs
+			);
+
+		DECLARE_DELEGATE_TwoParams(FPopulateSearchString,
+			const TListItemType& InArgs,
+			TArray<FString>& InOutSearchStrings
+			);
+		
+		SLATE_BEGIN_ARGS(TReplicationColumn)
+		{}
+			/** Callback to use for generating the widget in the column */
+			SLATE_EVENT(FGenerateColumnWidget, GenerateWidgetColumn)
+			/** Callback used to generate search items for this column */
+			SLATE_EVENT(FPopulateSearchString, PopulateSearchItems)
+			/** Where in the row this column is found with respect to the other columns */
+			SLATE_ARGUMENT(int32, ColumnSortOrder)
+		SLATE_END_ARGS()
+
+		TReplicationColumn(const TReplicationColumn::FArguments& InArgs, const FColumn::FArguments& InColumnArgs)
+			: FColumn(InColumnArgs)
+			, GenerateColumnWidgetCallback(InArgs._GenerateWidgetColumn)
+			, PopulateSearchStringCallback(InArgs._PopulateSearchItems)
+			, ColumnSortOrderValue(InArgs._ColumnSortOrder)
+		{}
+
+		TSharedRef<SWidget> BuildColumnWidget(const FBuildArgs& InArgs) const
+		{
+			return ensure(GenerateColumnWidgetCallback.IsBound())
+				? GenerateColumnWidgetCallback.Execute(InArgs)
+				: SNullWidget::NullWidget;
+		}
+
+		void ExecutePopulateSearchString(const TListItemType& InRowData, TArray<FString>& InOutSearchStrings) const
+		{
+			if (PopulateSearchStringCallback.IsBound())
+			{
+				PopulateSearchStringCallback.Execute(InRowData, InOutSearchStrings);
+			}
+		}
+		
+		int32 GetColumnSortOrderValue() const { return ColumnSortOrderValue; }
+		
+	private:
+		
+		FGenerateColumnWidget GenerateColumnWidgetCallback;
+		FPopulateSearchString PopulateSearchStringCallback;
+		/** Determines whether this column is the first, etc. */
+		int32 ColumnSortOrderValue = 0;
+	};
+}
\ No newline at end of file
diff --git a/Engine/Plugins/Developer/Concert/ConcertApp/MultiUserClient/Source/MultiUserReplicationEditor/Private/StreamEditor/View/ObjectViewer/SPropertyReplicationSelectionViewer.cpp b/Engine/Plugins/Developer/Concert/ConcertApp/MultiUserClient/Source/MultiUserReplicationEditor/Private/StreamEditor/View/ObjectViewer/SPropertyReplicationSelectionViewer.cpp
new file mode 100644
--- /dev/null
+++ b/Engine/Plugins/Developer/Concert/ConcertApp/MultiUserClient/Source/MultiUserReplicationEditor/Private/StreamEditor/View/ObjectViewer/SPropertyReplicationSelectionViewer.cpp
@@ -0,0 +1,398 @@
+// Copyright Epic Games, Inc. All Rights Reserved.
+
+#include "SPropertyReplicationSelectionViewer.h"
+
+#include "ConcertFrontendUtils.h"
+#include "StreamEditor/Model/IObjectToPropertiesModel.h"
+#include "ReplicatedObjectData.h"
+#include "ReplicatedPropertyData.h"
+#include "SelectionViewerColumns.h"
+
+#include "Algo/AllOf.h"
+#include "Algo/ForEach.h"
+#include "Components/ActorComponent.h"
+#include "GameFramework/Actor.h"
+#include "Widgets/Layout/SBorder.h"
+#include "Widgets/Layout/SSplitter.h"
+#include "Widgets/Layout/SWidgetSwitcher.h"
+#include "Widgets/SBoxPanel.h"
+#include "Widgets/Text/STextBlock.h"
+
+#define LOCTEXT_NAMESPACE "SPropertyReplicationSelectionViewer"
+
+namespace UE::MultiUserReplicationEditor
+{
+	void SPropertyReplicationSelectionViewer::Construct(const FArguments& InArgs, TSharedRef<IObjectToPropertiesModel> InPropertiesModel)
+	{
+		PropertiesModel = MoveTemp(InPropertiesModel);
+
+		SortPropertyRowPredicate = InArgs._SortPropertyRowPredicate;
+		
+		ChildSlot
+		[
+			SNew(SVerticalBox)
+
+			+SVerticalBox::Slot()
+			[
+				SNew(SSplitter)
+				.Orientation(Orient_Vertical)
+
+				// Actors
+				+SSplitter::Slot()
+				.SizeRule(TAttribute<SSplitter::ESizeRule>(this, &SPropertyReplicationSelectionViewer::GetActorAreaSizeRule))
+				.Value(0.55f)
+				[
+					CreateActorsSection(InArgs)
+				]
+				
+				// Properties
+				+SSplitter::Slot()
+				.SizeRule(TAttribute<SSplitter::ESizeRule>(this, &SPropertyReplicationSelectionViewer::GetPropertyAreaSizeRule))
+				.Value(0.45f)
+				[
+					CreatePropertiesSection(InArgs)
+				]
+			]
+		];
+
+		RefreshObjectData();
+		RefreshPropertyData();
+
+		ActorArea->SetExpanded(true);
+		PropertyArea->SetExpanded(true);
+	}
+
+	void SPropertyReplicationSelectionViewer::RefreshObjectData()
+	{
+		const int32 NumElements = PropertiesModel->GetNumReplicatedObjects();
+		// Re-using existing instances is tricky: we cannot update the object path in an item because the list view will no detect this change;
+		// list view only looks at the shared ptr address. So the UI will not be refreshed. Since the number of items will be small, just reallocate... 
+		ObjectRowData.Empty(NumElements);
+
+		// Try to re-use old instances by using the old PathToObjectDataCache. This is also done so the expansion states restore correctly in the tree view.
+		TMap<FSoftObjectPath, TSharedPtr<FReplicatedObjectData>> NewPathToObjectDataCache;
+		
+		// Do a complete refresh.
+		// Complete refresh is acceptable because the list is updated infrequently and typically small < 500 items.
+		// An alternative would be to change RefreshObjectData to be called with two variables ObjectsAdded and ObjectsRemoved.
+		PropertiesModel->ForEachReplicatedObject([this, &NewPathToObjectDataCache](const FSoftObjectPath& Path) mutable
+		{
+			const TSharedPtr<FReplicatedObjectData>* ExistingItem = PathToObjectDataCache.Find(Path);
+			const TSharedRef<FReplicatedObjectData> Item = ExistingItem ? ExistingItem->ToSharedRef() : AllocateObjectData(Path);
+			ObjectRowData.Emplace(Item);
+			NewPathToObjectDataCache.Emplace(Path, Item);
+			return EBreakBehavior::Continue;
+		});
+		
+		// If an item was removed, then NewPathToObjectDataCache does not contain it. 
+		PathToObjectDataCache = MoveTemp(NewPathToObjectDataCache);
+
+		// The tree view requires the item source to only contain the root items. Children are discovered via GetObjectRowChildren. We re-use GetObjectRowChildren to remove any non-root nodes.
+		BuildRootObjectRowData();
+		ReplicatedObjects->OnItemsChanged();
+	}
+
+	void SPropertyReplicationSelectionViewer::RefreshPropertyData()
+	{
+		const TArray<TSharedPtr<FReplicatedObjectData>> SelectedObjects = ReplicatedObjects->GetSelectedItems();
+		if (SelectedObjects.IsEmpty())
+		{
+			SetPropertyContent(EReplicatedPropertyContent::NoSelection);
+			return;
+		}
+
+		// Technically, the classes just need to be compatible with each other... but it is easier to just allow the same class.
+		const TOptional<FSoftClassPath> SharedClass = GetClassForPropertiesFromSelection();
+		if (!SharedClass)
+		{
+			SetPropertyContent(EReplicatedPropertyContent::SelectionTooBig);
+			return;
+		}
+		const FSoftClassPath Class = *SharedClass;
+
+		// Build the set of properties that are shared by all of the selected objects
+		TSet<FConcertPropertyChain> SharedProperties = PropertiesModel->GetAllProperties(SelectedObjects[0]->GetObjectPath());
+		for (int32 i = 1; i < SelectedObjects.Num(); ++i)
+		{
+			SharedProperties = PropertiesModel->GetAllProperties(SelectedObjects[i]->GetObjectPath()).Union(SharedProperties);
+		}
+		
+		// Try to re-use old instances by using the old PathToObjectDataCache. This is also done so the expansion states restore correctly in the tree view.
+		TMap<FConcertPropertyChain, TSharedPtr<FReplicatedPropertyData>> NewPathToPropertyDataCache;
+		
+		PropertyRowData.Empty(PropertiesModel->GetNumProperties(SelectedObjects[0]->GetObjectPath()));
+		for (const FConcertPropertyChain& PropertyChain : SharedProperties)
+		{
+			const TSharedPtr<FReplicatedPropertyData>* ExistingItem = ChainToPropertyDataCache.Find(PropertyChain);
+			const TSharedRef<FReplicatedPropertyData> Item = ExistingItem ? ExistingItem->ToSharedRef() : AllocatePropertyData(Class, PropertyChain);
+			PropertyRowData.Emplace(Item);
+			NewPathToPropertyDataCache.Emplace(PropertyChain, Item);
+		}
+
+		// If an item was removed, then NewPathToPropertyDataCache does not contain it. 
+		ChainToPropertyDataCache = MoveTemp(NewPathToPropertyDataCache);
+		
+		// The tree view requires the item source to only contain the root items.
+		BuildRootPropertyRowData();
+		SetPropertyContent(EReplicatedPropertyContent::Properties);
+		ReplicatedProperties->OnItemsChanged();
+	}
+
+	TOptional<FSoftClassPath> SPropertyReplicationSelectionViewer::GetClassForPropertiesFromSelection() const
+	{
+		const TArray<TSharedPtr<FReplicatedObjectData>> SelectedObjects = ReplicatedObjects->GetSelectedItems();
+		const FSoftClassPath Class = PropertiesModel->GetObjectClass(SelectedObjects[0]->GetObjectPath());
+		const bool bAllHaveSameClass = Algo::AllOf(SelectedObjects, [this, Class](const TSharedPtr<FReplicatedObjectData>& Object)
+		{
+			return PropertiesModel->GetObjectClass(Object->GetObjectPath()) == Class;
+		});
+		return bAllHaveSameClass ? Class : TOptional<FSoftClassPath>{};
+	}
+
+	FSoftObjectPath SPropertyReplicationSelectionViewer::GetParent(const FSoftObjectPath& Path) const
+	{
+		UObject* Object = Path.ResolveObject();
+		if (!Object)
+		{
+			return nullptr;
+		}
+
+		for (UObject* Current = Object->GetOuter(); Current; Current = Current->GetOuter())
+		{
+			if (PathToObjectDataCache.Contains(Current))
+			{
+				return Current;
+			}
+		}
+
+		return nullptr;
+	}
+
+	TSharedRef<FReplicatedObjectData> SPropertyReplicationSelectionViewer::AllocateObjectData(FSoftObjectPath ObjectPath)
+	{
+		return MakeShared<FReplicatedObjectData>(MoveTemp(ObjectPath));
+	}
+
+	TSharedRef<FReplicatedPropertyData> SPropertyReplicationSelectionViewer::AllocatePropertyData(FSoftClassPath OwningClass, FConcertPropertyChain PropertyChain)
+	{
+		return MakeShared<FReplicatedPropertyData>(MoveTemp(OwningClass), MoveTemp(PropertyChain));
+	}
+
+	TSharedRef<SWidget> SPropertyReplicationSelectionViewer::CreateActorsSection(const FArguments& InArgs)
+	{
+		TArray<ReplicationObjectColumns::FReplicationObjectColumn> Columns
+		{
+			ReplicationObjectColumns::IconColumn(PropertiesModel.ToSharedRef()),
+			ReplicationObjectColumns::LabelColumn(),
+			ReplicationObjectColumns::TypeColumn(PropertiesModel.ToSharedRef())
+		};
+		Columns.Append(InArgs._AdditionalObjectColumns);
+		
+		return SNew(SBorder)
+			.BorderImage(FAppStyle::Get().GetBrush("ToolPanel.GroupBorder"))
+			.Padding(0.0f)
+			[
+				SAssignNew(ActorArea, SExpandableArea)
+				.InitiallyCollapsed(true) 
+				.BorderBackgroundColor(FLinearColor(0.6f, 0.6f, 0.6f, 1.0f))
+				.BorderImage_Lambda([this]() { return ConcertFrontendUtils::GetExpandableAreaBorderImage(*ActorArea); })
+				.BodyBorderImage(FAppStyle::Get().GetBrush("ToolPanel.GroupBorder"))
+				.BodyBorderBackgroundColor(FLinearColor::White)
+				.OnAreaExpansionChanged(this, &SPropertyReplicationSelectionViewer::OnActorAreaExpansionChanged)
+				.Padding(0.0f)
+				.HeaderContent()
+				[
+					SNew(STextBlock)
+					.Text(LOCTEXT("ReplicatedObjects", "Replicated Objects"))
+					.Font(FAppStyle::Get().GetFontStyle("DetailsView.CategoryFontStyle"))
+					.ShadowOffset(FVector2D(1.0f, 1.0f))
+				]
+				.BodyContent()
+				[
+					SAssignNew(ReplicatedObjects, SReplicationTreeView<TSharedPtr<FReplicatedObjectData>>)
+					.RootItemsSource(&RootObjectRowData)
+					.OnGetChildren(this, &SPropertyReplicationSelectionViewer::GetObjectRowChildren)
+					.OnContextMenuOpening(InArgs._OnObjectsContextMenuOpening)
+					.OnDeleteItems(InArgs._OnDeleteObjects)
+					.OnSelectionChanged_Lambda([this](){ RefreshPropertyData(); })
+					.Columns(Columns)
+					.ExpandableColumnLabel(ReplicationObjectColumns::LabelColumnId)
+					.SelectionMode(ESelectionMode::Multi)
+					.LeftOfSearchBar()
+					[
+						InArgs._LeftOfObjectSearchBar.Widget
+					]
+				]
+			];
+	}
+
+	TSharedRef<SWidget> SPropertyReplicationSelectionViewer::CreatePropertiesSection(const FArguments& InArgs)
+	{
+		TArray<ReplicationPropertyColumns::FReplicationPropertyColumn> Columns
+		{
+			ReplicationPropertyColumns::LabelColumn()
+		};
+		Columns.Append(InArgs._AdditionalPropertyColumns);
+		
+		return SNew(SBorder)
+			.BorderImage(FAppStyle::Get().GetBrush("ToolPanel.GroupBorder"))
+			.Padding(0.0f)
+			[
+				SAssignNew(PropertyArea, SExpandableArea)
+				.InitiallyCollapsed(true) 
+				.BorderBackgroundColor(FLinearColor(0.6f, 0.6f, 0.6f, 1.0f))
+				.BorderImage_Lambda([this]() { return ConcertFrontendUtils::GetExpandableAreaBorderImage(*PropertyArea); })
+				.BodyBorderImage(FAppStyle::Get().GetBrush("ToolPanel.GroupBorder"))
+				.BodyBorderBackgroundColor(FLinearColor::White)
+				.OnAreaExpansionChanged(this, &SPropertyReplicationSelectionViewer::OnPropertyAreaExpansionChanged)
+				.Padding(0.0f)
+				.HeaderContent()
+				[
+					SNew(STextBlock)
+					.Text(LOCTEXT("ReplicatedProperties", "Replicated Properties"))
+					.Font(FAppStyle::Get().GetFontStyle("DetailsView.CategoryFontStyle"))
+					.ShadowOffset(FVector2D(1.0f, 1.0f))
+				]
+				.BodyContent()
+				[
+					// Make sure the slots are coherent with the order of EReplicatedPropertyContent!
+					SAssignNew(PropertyContent, SWidgetSwitcher)
+					.WidgetIndex(static_cast<int32>(EReplicatedPropertyContent::NoSelection))
+					
+					// EReplicatedPropertyContent::Properties
+					+SWidgetSwitcher::Slot()
+					[
+						SAssignNew(ReplicatedProperties, SReplicationTreeView<TSharedPtr<FReplicatedPropertyData>>)
+						.RootItemsSource(&RootPropertyRowData)
+						.OnGetChildren(this, &SPropertyReplicationSelectionViewer::GetPropertyRowChildren)
+						.Columns(Columns)
+						.ExpandableColumnLabel(ReplicationPropertyColumns::LabelColumnId)
+						.SelectionMode(ESelectionMode::Multi)
+						.LeftOfSearchBar()
+						[
+							InArgs._LeftOfPropertySearchBar.Widget
+						]
+					]
+					
+					// EReplicatedPropertyContent::NoSelection
+					+SWidgetSwitcher::Slot()
+					.VAlign(VAlign_Center)
+					.HAlign(HAlign_Center)
+					[
+						SNew(STextBlock)
+						.Text(LOCTEXT("NoSelection", "Select an object to see selected properties"))
+					]
+					
+					// EReplicatedPropertyContent::SelectionTooBig
+					+SWidgetSwitcher::Slot()
+					.VAlign(VAlign_Center)
+					.HAlign(HAlign_Center)
+					[
+						SNew(STextBlock)
+						.Text(LOCTEXT("SelectionTooBig", "Select objects of the same type type to see selected properties"))
+					]
+				]
+			];
+	}
+	
+	void SPropertyReplicationSelectionViewer::BuildRootObjectRowData()
+	{
+		TSet<TSharedPtr<FReplicatedObjectData>> NonRootNodes;
+		for (const TSharedPtr<FReplicatedObjectData>& Node : ObjectRowData)
+		{
+			GetObjectRowChildren(Node, [&NonRootNodes](TSharedPtr<FReplicatedObjectData> Child)
+			{
+				NonRootNodes.Add(MoveTemp(Child));
+			});
+		}
+
+		// Make RootObjectRowData only contains those nodes which were not listed as children 
+		RootObjectRowData.Empty(NonRootNodes.Num());
+		for (const TSharedPtr<FReplicatedObjectData>& Node : ObjectRowData)
+		{
+			if (!NonRootNodes.Contains(Node))
+			{
+				RootObjectRowData.Add(Node);
+			}
+		}
+		
+		RootObjectRowData.Sort([](const TSharedPtr<FReplicatedObjectData>& Left, const TSharedPtr<FReplicatedObjectData>& Right)
+		{
+			return Left->GetObjectPath().GetSubPathString() < Right->GetObjectPath().GetSubPathString();
+		});
+	}
+
+	void SPropertyReplicationSelectionViewer::BuildRootPropertyRowData()
+	{
+		RootPropertyRowData.Empty(PropertyRowData.Num());
+		for (const TSharedPtr<FReplicatedPropertyData>& PropertyData : PropertyRowData)
+		{
+			if (PropertyData->GetProperty().IsRootProperty())
+			{
+				RootPropertyRowData.Emplace(PropertyData);
+			}
+		}
+
+		SortPropertyRowArray(RootPropertyRowData);
+	}
+
+	void SPropertyReplicationSelectionViewer::GetObjectRowChildren(TSharedPtr<FReplicatedObjectData> ReplicatedObjectData, TFunctionRef<void(TSharedPtr<FReplicatedObjectData>)> ProcessChild)
+	{
+		const UObject* ResolvedObject = ReplicatedObjectData->GetObjectPath().ResolveObject();
+		const AActor* ResolvedActor = Cast<AActor>(ResolvedObject);
+
+		GetActorComponentsAsChildRows(ResolvedActor, ProcessChild);
+	}
+
+	void SPropertyReplicationSelectionViewer::GetActorComponentsAsChildRows(const AActor* ResolvedActor, TFunctionRef<void(TSharedPtr<FReplicatedObjectData>)> ProcessChild)
+	{
+		if (!ResolvedActor)
+		{
+			return;
+		}
+
+		for (UActorComponent* Component : ResolvedActor->GetComponents())
+		{
+			const FSoftObjectPath ComponentPath = Component;
+			if (const TSharedPtr<FReplicatedObjectData>* ComponentChild = PathToObjectDataCache.Find(ComponentPath))
+			{
+				ProcessChild(*ComponentChild);
+			}
+		}
+	}
+
+	void SPropertyReplicationSelectionViewer::GetPropertyRowChildren(TSharedPtr<FReplicatedPropertyData> ReplicatedPropertyData, TFunctionRef<void(TSharedPtr<FReplicatedPropertyData>)> ProcessChild)
+	{
+		TArray<TSharedPtr<FReplicatedPropertyData>> Children;
+		
+		// Not the most efficient but it should be fine.
+		for (const TSharedPtr<FReplicatedPropertyData>& Data : PropertyRowData)
+		{
+			if (Data->GetProperty().IsDirectChildOf(ReplicatedPropertyData->GetProperty()))
+			{
+				Children.Add(Data);
+			}
+		}
+
+		SortPropertyRowArray(Children);
+		Algo::ForEach(Children, [&ProcessChild](const TSharedPtr<FReplicatedPropertyData>& Data){ ProcessChild(Data); });
+	}
+	
+	void SPropertyReplicationSelectionViewer::SortPropertyRowArray(TArray<TSharedPtr<FReplicatedPropertyData>>& ToSort) const
+	{
+		if (SortPropertyRowPredicate.IsBound())
+		{
+			ToSort.Sort([this](const TSharedPtr<FReplicatedPropertyData>& Left, const TSharedPtr<FReplicatedPropertyData>& Right)
+			{
+				return SortPropertyRowPredicate.Execute(Left, Right);
+			});
+		}
+	}
+
+	void SPropertyReplicationSelectionViewer::SetPropertyContent(EReplicatedPropertyContent Content) const
+	{
+		PropertyContent->SetActiveWidgetIndex(static_cast<int32>(Content));
+	}
+}
+
+#undef LOCTEXT_NAMESPACE
\ No newline at end of file
diff --git a/Engine/Plugins/Developer/Concert/ConcertApp/MultiUserClient/Source/MultiUserReplicationEditor/Private/StreamEditor/View/ObjectViewer/SPropertyReplicationSelectionViewer.h b/Engine/Plugins/Developer/Concert/ConcertApp/MultiUserClient/Source/MultiUserReplicationEditor/Private/StreamEditor/View/ObjectViewer/SPropertyReplicationSelectionViewer.h
new file mode 100644
--- /dev/null
+++ b/Engine/Plugins/Developer/Concert/ConcertApp/MultiUserClient/Source/MultiUserReplicationEditor/Private/StreamEditor/View/ObjectViewer/SPropertyReplicationSelectionViewer.h
@@ -0,0 +1,197 @@
+// Copyright Epic Games, Inc. All Rights Reserved.
+
+#pragma once
+
+#include "Replication/Data/ConcertPropertySelection.h"
+#include "ReplicatedObjectData.h"
+#include "SReplicationTreeView.h"
+#include "StreamEditor/View/ObjectViewer/SelectionViewerColumns.h"
+
+#include "Algo/Transform.h"
+#include "Misc/Optional.h"
+#include "Widgets/DeclarativeSyntaxSupport.h"
+#include "Widgets/SCompoundWidget.h"
+
+class SExpandableArea;
+class SWidgetSwitcher;
+struct FSoftClassPath;
+struct FSoftObjectPath;
+
+namespace UE::MultiUserReplicationEditor
+{
+	class FReplicatedPropertyData;
+	class FReplicatedObjectData;
+	class IEditableObjectToPropertiesModel;
+	class IObjectToPropertiesModel;
+
+	DECLARE_DELEGATE(FOnPropertiesRefreshed);
+	DECLARE_DELEGATE_RetVal_TwoParams(bool, FSortPropertyPredicate, const TSharedPtr<FReplicatedPropertyData>& Left, const TSharedPtr<FReplicatedPropertyData>& Right);
+	
+	/**
+	 * Root widget for viewing UMultiUserPropertyReplicationSelection.
+	 * This widget knows how to display IObjectToPropertiesModel.
+	 * 
+	 * The underlying data is modified by SPropertyReplicationSelectionEditor, which uses this widget's extension
+	 * points to call functions on IEditableObjectToPropertiesModel.
+	 */
+	class MULTIUSERREPLICATIONEDITOR_API SPropertyReplicationSelectionViewer : public SCompoundWidget
+	{
+	public:
+
+		SLATE_BEGIN_ARGS(SPropertyReplicationSelectionViewer)
+		{}
+			/** Additional columns to add to the object view */
+			SLATE_ARGUMENT(TArray<ReplicationObjectColumns::FReplicationObjectColumn>, AdditionalObjectColumns)
+			/** Additional columns to add to the property view */
+			SLATE_ARGUMENT(TArray<ReplicationPropertyColumns::FReplicationPropertyColumn>, AdditionalPropertyColumns)
+
+			/** Optional. Called when the delete key is pressed in the object view. */
+			SLATE_EVENT(SReplicationTreeView<TSharedPtr<FReplicatedObjectData>>::FDeleteItems, OnDeleteObjects)
+		
+			/** Called to generate the context menu for objects. */
+			SLATE_EVENT(FOnContextMenuOpening, OnObjectsContextMenuOpening)
+
+			/** Optional. Used for determining the order in which properties are displayed. */
+			SLATE_EVENT(FSortPropertyPredicate, SortPropertyRowPredicate)
+		
+			/** Optional widget to add to the left of the object list search bar. */
+			SLATE_NAMED_SLOT(FArguments, LeftOfObjectSearchBar)
+			/** Optional widget to add to the left of the property list search bar. */
+			SLATE_NAMED_SLOT(FArguments, LeftOfPropertySearchBar)
+		SLATE_END_ARGS()
+
+		void Construct(const FArguments& InArgs, TSharedRef<IObjectToPropertiesModel> InPropertiesModel);
+
+		void RefreshObjectData();
+		void RefreshPropertyData();
+
+		/** Sets only these objects to be selected. */
+		template<typename TObjectRangeType>
+		void SetSelectedObjects(const TObjectRangeType& ObjectsToSelect, bool bIsSelected);
+		/** Expands the given items. Does not unexpand already expanded items. */
+		template<typename TObjectRangeType>
+		void SetObjectsExpanded(const TObjectRangeType& ObjectsToSelect, bool bIsExpanded);
+
+		/** Given the selected objects, determines whether they all have the same class and returns it if so. */
+		TOptional<FSoftClassPath> GetClassForPropertiesFromSelection() const;
+
+		/** Gets the child items of the given object. */
+		template<typename TObjectType>
+		void GetObjectChildren(const TObjectType& ReplicatedObjectData, TFunctionRef<void(const FSoftObjectPath&)> ProcessChild);
+		/** @return The parent of the object associated with Path or null (check IsNull) if none exists. */
+		FSoftObjectPath GetParent(const FSoftObjectPath& Path) const;
+
+		const TSet<TSharedPtr<FReplicatedObjectData>>& GetObjectRowData() const { return ObjectRowData; }
+		const TArray<TSharedPtr<FReplicatedPropertyData>>& GetPropertyRowData() const { return PropertyRowData; }
+
+		TArray<TSharedPtr<FReplicatedObjectData>> GetSelectedObjects() const { return ReplicatedObjects->GetSelectedItems(); }
+	
+	private:
+
+		/** The model this view is visualizing. */
+		TSharedPtr<IObjectToPropertiesModel> PropertiesModel;
+
+		/** Lists the selected actors */
+		TSharedPtr<SExpandableArea> ActorArea;
+		/** Lists the properties of the selected actor */
+		TSharedPtr<SExpandableArea> PropertyArea;
+
+		/** Tree view for replicated objects. */
+		TSharedPtr<SReplicationTreeView<TSharedPtr<FReplicatedObjectData>>> ReplicatedObjects;
+		/** Tree view for replicated properties. Content depends on the current object selected. */
+		TSharedPtr<SReplicationTreeView<TSharedPtr<FReplicatedPropertyData>>> ReplicatedProperties;
+
+		enum class EReplicatedPropertyContent
+		{
+			/** Shows the properties */
+			Properties,
+			/** Prompts: "Select an object to see selected properties" */
+			NoSelection,
+			/** Prompts: "Select objects of the same type type to see selected properties" */
+			SelectionTooBig
+		};
+		/** Determines the content displayed for PropertyArea. */
+		TSharedPtr<SWidgetSwitcher> PropertyContent;
+		
+		/**
+		 * These instances can be subclasses of FReplicatedObjectData, e.g. FReplicatedObjectData_Editor.
+		 * Their type can be overridden by subclasses.
+		 * They only have the FReplicatedObjectData type so they can be passed efficiently to SPropertyReplicationSelectionViewer.
+		 * @see GetObjectData
+		 */
+		TSet<TSharedPtr<FReplicatedObjectData>> ObjectRowData;
+		/** The instances of ObjectRowData which do not have any parents. This acts as the item source for the tree view. */
+		TArray<TSharedPtr<FReplicatedObjectData>> RootObjectRowData;
+		/** Inverse map of ObjectRowData using FReplicatedObjectData::GetObjectPath as key. Contains all elements of ObjectRowData. */
+		TMap<FSoftObjectPath, TSharedPtr<FReplicatedObjectData>> PathToObjectDataCache;
+		
+		/**
+		 * These instances can be subclasses of FReplicatedPropertyData, e.g. FReplicatedPropertyData_Editor.
+		 * Their type can be overridden by subclasses.
+		 * They only have the FReplicatedPropertyData type so they can be passed efficiently to SPropertyReplicationSelectionViewer.
+		 * @see GetPropertyData
+		 */
+		TArray<TSharedPtr<FReplicatedPropertyData>> PropertyRowData;
+		/** The instances of ObjectRowData which do not have any parents. This acts as the item source for the tree view. */
+		TArray<TSharedPtr<FReplicatedPropertyData>> RootPropertyRowData;
+		/** Inverse map of PropertyRowData using FReplicatedPropertyData::GetProperty as key. Contains all elements of PropertyRowData. */
+		TMap<FConcertPropertyChain, TSharedPtr<FReplicatedPropertyData>> ChainToPropertyDataCache;
+
+		bool bIsActorAreaExpanded = false;
+		bool bIsPropertyAreaExpanded = false;
+
+		/** Used for determining the order in which properties are displayed. Default: Sort by name. */
+		FSortPropertyPredicate SortPropertyRowPredicate;
+
+		virtual TSharedRef<FReplicatedObjectData> AllocateObjectData(FSoftObjectPath ObjectPath);
+		virtual TSharedRef<FReplicatedPropertyData> AllocatePropertyData(FSoftClassPath OwningClass, FConcertPropertyChain PropertyChain);
+
+		TSharedRef<SWidget> CreateActorsSection(const FArguments& InArgs);
+		TSharedRef<SWidget> CreatePropertiesSection(const FArguments& InArgs);
+
+		/** Sets RootObjectRowData to all non-root nodes from ObjectRowData. */
+		void BuildRootObjectRowData();
+		void BuildRootPropertyRowData();
+		
+		void GetObjectRowChildren(TSharedPtr<FReplicatedObjectData> ReplicatedObjectData, TFunctionRef<void(TSharedPtr<FReplicatedObjectData>)> ProcessChild);
+		void GetActorComponentsAsChildRows(const AActor* ResolvedActor, TFunctionRef<void(TSharedPtr<FReplicatedObjectData>)> ProcessChild);
+		void GetPropertyRowChildren(TSharedPtr<FReplicatedPropertyData> ReplicatedPropertyData, TFunctionRef<void(TSharedPtr<FReplicatedPropertyData>)> ProcessChild);
+
+		/** Handles how much space the 'Clients' area uses with respect to its expansion state. */
+		SSplitter::ESizeRule GetActorAreaSizeRule() const { return bIsActorAreaExpanded ? SSplitter::ESizeRule::FractionOfParent : SSplitter::ESizeRule::SizeToContent; }
+		void OnActorAreaExpansionChanged(bool bExpanded) { bIsActorAreaExpanded = bExpanded; }
+
+		/** Handles how much space the 'Clients' area uses with respect to its expansion state. */
+		SSplitter::ESizeRule GetPropertyAreaSizeRule() const { return bIsPropertyAreaExpanded ? SSplitter::ESizeRule::FractionOfParent : SSplitter::ESizeRule::SizeToContent; }
+		void OnPropertyAreaExpansionChanged(bool bExpanded) { bIsPropertyAreaExpanded = bExpanded; }
+
+		// Utils
+		void SortPropertyRowArray(TArray<TSharedPtr<FReplicatedPropertyData>>& ToSort) const;
+		void SetPropertyContent(EReplicatedPropertyContent Content) const;
+	};
+
+	template <typename TObjectRangeType>
+	void SPropertyReplicationSelectionViewer::SetSelectedObjects(const TObjectRangeType& ObjectsToSelect, bool bIsSelected)
+	{
+		TArray<TSharedPtr<FReplicatedObjectData>> ObjectDataToSelect;
+		Algo::Transform(ObjectsToSelect, ObjectDataToSelect, [this](const FSoftObjectPath& Path){ return PathToObjectDataCache[Path]; });
+		ReplicatedObjects->SetSelectedItems(ObjectDataToSelect, bIsSelected);
+	}
+
+	template <typename TObjectRangeType>
+	void SPropertyReplicationSelectionViewer::SetObjectsExpanded(const TObjectRangeType& ObjectsToSelect, bool bIsExpanded)
+	{
+		TArray<TSharedPtr<FReplicatedObjectData>> ObjectDataToSelect;
+		Algo::Transform(ObjectsToSelect, ObjectDataToSelect, [this](const FSoftObjectPath& Path){ return PathToObjectDataCache[Path]; });
+		ReplicatedObjects->SetExpandedItems(ObjectDataToSelect, bIsExpanded);
+	}
+
+	template <typename TObjectType>
+	void SPropertyReplicationSelectionViewer::GetObjectChildren(const TObjectType& ReplicatedObjectData, TFunctionRef<void(const FSoftObjectPath&)> ProcessChild)
+	{
+		if (const TSharedPtr<FReplicatedObjectData>* Data = PathToObjectDataCache.Find(ReplicatedObjectData))
+		{
+			GetObjectRowChildren(*Data, [&ProcessChild](const TSharedPtr<FReplicatedObjectData>& Child){ ProcessChild(Child->GetObjectPath()); });
+		}
+	}
+}
diff --git a/Engine/Plugins/Developer/Concert/ConcertApp/MultiUserClient/Source/MultiUserReplicationEditor/Private/StreamEditor/View/ObjectViewer/SReplicationColumnRow.h b/Engine/Plugins/Developer/Concert/ConcertApp/MultiUserClient/Source/MultiUserReplicationEditor/Private/StreamEditor/View/ObjectViewer/SReplicationColumnRow.h
new file mode 100644
--- /dev/null
+++ b/Engine/Plugins/Developer/Concert/ConcertApp/MultiUserClient/Source/MultiUserReplicationEditor/Private/StreamEditor/View/ObjectViewer/SReplicationColumnRow.h
@@ -0,0 +1,101 @@
+// Copyright Epic Games, Inc. All Rights Reserved.
+
+#pragma once
+
+#include "MultiUserReplicationEditorStyle.h"
+#include "ReplicationColumn.h"
+#include "Widgets/Layout/SBox.h"
+#include "Widgets/SNullWidget.h"
+#include "Widgets/Views/SExpanderArrow.h"
+#include "Widgets/Views/STableRow.h"
+#include "Widgets/Views/STableViewBase.h"
+
+namespace UE::MultiUserReplicationEditor
+{
+	/**
+	 * Displays the summary of an activity recorded and recoverable in the SConcertSessionRecovery list view.
+	 */
+	template<typename TListItemType>
+	class SReplicationColumnRow : public SMultiColumnTableRow<TSharedPtr<TListItemType>>
+	{
+	public:
+
+		using TColumType = TReplicationColumn<TListItemType>;
+		
+		DECLARE_DELEGATE_RetVal_OneParam(const TColumType*, FGetColumn,
+			const FName& ColumnId
+			);
+		
+		SLATE_BEGIN_ARGS(SReplicationColumnRow)
+		{}
+			/** Used for highlighting the text being searched. */
+			SLATE_ARGUMENT(TSharedPtr<FText>, HighlightText)
+
+			/** Gets columns info about a certain column */
+			SLATE_EVENT(FGetColumn, ColumnGetter)
+
+			/** The data to pass to TReplicationColumn::BuildColumnWidget. */
+			SLATE_ARGUMENT(TListItemType, RowData)
+		
+			/** The name of the column that will have the SExpanderArrow for the tree view. */
+			SLATE_ARGUMENT(FName, ExpandableColumnLabel)
+		SLATE_END_ARGS()
+
+		void Construct(
+			const FArguments& InArgs,
+			TSharedRef<STableViewBase> InOwner)
+		{
+			ColumnGetterDelegate = InArgs._ColumnGetter;
+			HighlightText = InArgs._HighlightText;
+			RowData = InArgs._RowData;
+			ExpandableColumnLabel = InArgs._ExpandableColumnLabel;
+			
+			SMultiColumnTableRow<TSharedPtr<TListItemType>>::Construct({}, InOwner);
+		}
+
+		/** Generates the widget representing this row. */
+		virtual TSharedRef<SWidget> GenerateWidgetForColumn(const FName& ColumnName) override
+		{
+			const TColumType* Column = ColumnGetterDelegate.Execute(ColumnName); ensure(Column);
+			if (!Column)
+			{
+				return SNullWidget::NullWidget;
+			}
+			
+			const TSharedRef<SWidget> ColumnWidget = Column->BuildColumnWidget({ HighlightText, RowData });
+			const bool bNeedsExpanderArrow = ColumnName == ExpandableColumnLabel;
+			if (!bNeedsExpanderArrow)
+			{
+				return ColumnWidget;
+			}
+			
+			return SNew(SBox)
+				.MinDesiredHeight(FMultiUserReplicationEditorStyle::Get().GetFloat("ReplicationTreeView.RowHeight"))
+				[
+					SNew(SHorizontalBox)
+
+					+SHorizontalBox::Slot()
+					.AutoWidth()
+					.Padding(6, 0, 0, 0)
+					[
+						SNew(SExpanderArrow, SReplicationColumnRow::SharedThis(this))
+						.IndentAmount(12)
+					]
+
+					+SHorizontalBox::Slot()
+					.FillWidth(1.0f)
+					[
+						ColumnWidget
+					]
+				];
+	
+		}
+
+	private:
+		
+		FGetColumn ColumnGetterDelegate;
+		TSharedPtr<FText> HighlightText;
+		TListItemType RowData;
+		FName ExpandableColumnLabel;
+	};
+}
\ No newline at end of file
diff --git a/Engine/Plugins/Developer/Concert/ConcertApp/MultiUserClient/Source/MultiUserReplicationEditor/Private/StreamEditor/View/ObjectViewer/SReplicationTreeView.h b/Engine/Plugins/Developer/Concert/ConcertApp/MultiUserClient/Source/MultiUserReplicationEditor/Private/StreamEditor/View/ObjectViewer/SReplicationTreeView.h
new file mode 100644
--- /dev/null
+++ b/Engine/Plugins/Developer/Concert/ConcertApp/MultiUserClient/Source/MultiUserReplicationEditor/Private/StreamEditor/View/ObjectViewer/SReplicationTreeView.h
@@ -0,0 +1,336 @@
+// Copyright Epic Games, Inc. All Rights Reserved.
+
+#pragma once
+
+#include "ReplicationColumn.h"
+#include "SReplicationColumnRow.h"
+
+#include "Algo/RemoveIf.h"
+#include "Misc/TextFilter.h"
+#include "Widgets/DeclarativeSyntaxSupport.h"
+#include "Widgets/Input/SSearchBox.h"
+#include "Widgets/Layout/SScrollBox.h"
+#include "Widgets/SBoxPanel.h"
+#include "Widgets/Views/STreeView.h"
+
+#define LOCTEXT_NAMESPACE "SReplicationListView"
+
+namespace UE::MultiUserReplicationEditor
+{
+	/**
+	 * Shared code for the list view for replicated actors and properties.
+	 * It is a table view that is searchable with a search box.
+	 */
+	template<typename TItemType>
+	class SReplicationTreeView : public SCompoundWidget
+	{
+	public:
+
+		DECLARE_DELEGATE_OneParam(FDeleteItems, const TArray<TItemType>& SelectedItems);
+		DECLARE_DELEGATE_TwoParams(FGetItemChildren, TItemType Item, TFunctionRef<void(TItemType)> ProcessChild);
+		DECLARE_DELEGATE(FOnSelectionChanged);
+
+		SLATE_BEGIN_ARGS(SReplicationTreeView<TItemType>)
+			: _SelectionMode(ESelectionMode::Single)
+		{}
+			/** The items to display */
+			SLATE_ARGUMENT(TArray<TItemType>*, RootItemsSource)
+
+			/** Gets an items children for the tree view */
+			SLATE_EVENT(FGetItemChildren, OnGetChildren)
+
+			/** Optional. Called when the user presses the delete key */
+			SLATE_EVENT(FDeleteItems, OnDeleteItems)
+
+			/** Called to generate the context menu for an item */
+			SLATE_EVENT(FOnContextMenuOpening, OnContextMenuOpening)
+
+			/** Called when the selection changes. Call GetSelectedItems to get the selected items. */
+			SLATE_EVENT(FOnSelectionChanged, OnSelectionChanged)
+			
+			/** The columns this list should have */
+			SLATE_ARGUMENT(TArray<TReplicationColumn<TItemType>>, Columns)
+			/** The name of the column that will have the SExpanderArrow for the tree view. */
+			SLATE_ARGUMENT(FName, ExpandableColumnLabel)
+		
+			/** How many items are to allowed to be selected */
+			SLATE_ARGUMENT(ESelectionMode::Type, SelectionMode)
+
+			/** Optional widget to add to the left of the search bar. */
+			SLATE_NAMED_SLOT(FArguments, LeftOfSearchBar)
+		SLATE_END_ARGS()
+
+		void Construct(const FArguments& InArgs)
+		{
+			AllRootItems = InArgs._RootItemsSource;
+			check(AllRootItems);
+
+			OnGetChildrenDelegate = InArgs._OnGetChildren;
+			OnDeleteItemsDelegate = InArgs._OnDeleteItems;
+			ExpandableColumnId = InArgs._ExpandableColumnLabel;
+			
+			SearchText = MakeShared<FText>();
+			SearchTextFilter = MakeShared<TTextFilter<const TItemType&>>(TTextFilter<const TItemType&>::FItemToStringArray::CreateSP(this, &SReplicationTreeView::PopulateSearchStrings));
+			SearchTextFilter->OnChanged().AddSP(this, &SReplicationTreeView::OnFilterChanged);
+			
+			ChildSlot
+			[
+				SNew(SVerticalBox)
+
+				// Search
+				+SVerticalBox::Slot()
+				.Padding(1.f)
+				.AutoHeight()
+				[
+					SNew(SHorizontalBox)
+
+					+SHorizontalBox::Slot()
+					.Padding(1.f)
+					.AutoWidth()
+					[
+						InArgs._LeftOfSearchBar.Widget
+					]
+
+					+SHorizontalBox::Slot()
+					.Padding(1.f)
+					.FillWidth(1.f)
+					[
+						SAssignNew(SearchBox, SSearchBox)
+						.HintText(LOCTEXT("SearchHint", "Search..."))
+						.OnTextChanged(this, &SReplicationTreeView::OnSearchTextChanged)
+						.OnTextCommitted(this, &SReplicationTreeView::OnSearchTextCommitted)
+						.DelayChangeNotificationsWhileTyping(true)
+					]
+				]
+
+				// Table row
+				+SVerticalBox::Slot()
+				.Padding(1.f)
+				.FillHeight(1.f)
+				[
+					SNew(SScrollBox)
+					+SScrollBox::Slot()
+					[
+						CreateTreeView(InArgs)
+					]
+				]
+			];
+		}
+		
+		void OnItemsChanged()
+		{
+			// Re-filter everything. There should not be many items so filtering everything again should be fine. Calls RequestListRefresh as well.
+			OnFilterChanged();
+		}
+		
+		void SetSelectedItems(const TArray<TItemType>& ObjectsToSelect, bool bIsSelected)
+		{
+			TreeView->ClearSelection();
+			TreeView->SetItemSelection(ObjectsToSelect, bIsSelected);
+		}
+		void SetExpandedItems(const TArray<TItemType>& ObjectsToSelect, bool bIsExpanded)
+		{
+			for (TItemType Item : ObjectsToSelect)
+			{
+				TreeView->SetItemExpansion(Item, bIsExpanded);
+			}
+		}
+		
+		TArray<TItemType> GetSelectedItems() const { return TreeView->GetSelectedItems(); }
+		
+		virtual FReply OnKeyDown(const FGeometry& MyGeometry, const FKeyEvent& InKeyEvent) override;
+
+	private:
+
+		/** The widget being used for search. */
+		TSharedPtr<SSearchBox> SearchBox;
+		/** Used to highlight text in text widgets */
+		TSharedPtr<FText> SearchText;
+		/** Performs text search */
+		TSharedPtr<TTextFilter<const TItemType&>> SearchTextFilter;
+
+		/** ListView's header row */
+		TSharedPtr<SHeaderRow> HeaderRow;
+		/** Displays the contents */
+		TSharedPtr<STreeView<TItemType>> TreeView;
+		/** The name of the column which will have the SExpandableArrow widget for the tree view. */
+		FName ExpandableColumnId;
+
+		TArray<TItemType>* AllRootItems = nullptr;
+		TArray<TItemType> FilteredRootItems;
+
+		/** Callback for getting an item's children. */
+		FGetItemChildren OnGetChildrenDelegate;
+		/** Optional delegate for responding to pressing the delete button */
+		FDeleteItems OnDeleteItemsDelegate;
+		
+
+		TSharedRef<SWidget> CreateTreeView(const FArguments& InArgs);
+		TSharedRef<SHeaderRow> CreateHeaderRow(const FArguments& InArgs);
+		TSharedRef<ITableRow> OnGenerateRowWidget(TItemType Item, const TSharedRef<STableViewBase>& OwnerTable);
+		void GetRowChildren(TItemType Item, TArray<TItemType>& OutChildren);
+		
+		void OnSearchTextCommitted(const FText& InFilterText, ETextCommit::Type CommitType);
+		void OnSearchTextChanged(const FText& InSearchText);
+
+		void PopulateSearchStrings(const TItemType& Item, TArray<FString>& OutSearchStrings);
+		void OnFilterChanged();
+		bool PassesFilters(const TItemType& Item);
+	};
+
+	template <typename TItemType>
+	FReply SReplicationTreeView<TItemType>::OnKeyDown(const FGeometry& MyGeometry, const FKeyEvent& InKeyEvent)
+	{
+		if (InKeyEvent.GetKey() == EKeys::Delete && OnDeleteItemsDelegate.IsBound())
+		{
+			OnDeleteItemsDelegate.Execute(TreeView->GetSelectedItems());
+			return FReply::Handled();
+		}
+	
+		return SCompoundWidget::OnKeyDown(MyGeometry, InKeyEvent);
+	}
+
+	template <typename TItemType>
+	TSharedRef<SWidget> SReplicationTreeView<TItemType>::CreateTreeView(const FArguments& InArgs)
+	{
+		return SNew(SBorder)
+			.BorderImage(FAppStyle::Get().GetBrush("ToolPanel.GroupBorder"))
+			.BorderBackgroundColor(FSlateColor(FLinearColor(0.6, 0.6, 0.6)))
+			.Padding(0)
+			[
+				SAssignNew(TreeView, STreeView<TItemType>)
+				.OnGetChildren(this, &SReplicationTreeView::GetRowChildren)
+				.TreeItemsSource(&FilteredRootItems)
+				.OnGenerateRow(this, &SReplicationTreeView::OnGenerateRowWidget)
+				.OnContextMenuOpening(InArgs._OnContextMenuOpening)
+				.OnSelectionChanged_Lambda([OnSelectionChanged = InArgs._OnSelectionChanged](auto, auto){ OnSelectionChanged.ExecuteIfBound(); })
+				.SelectionMode(InArgs._SelectionMode)
+				.AllowOverscroll(EAllowOverscroll::No)
+				.HeaderRow(CreateHeaderRow(InArgs))
+			];
+	}
+
+	template <typename TItemType>
+	TSharedRef<SHeaderRow> SReplicationTreeView<TItemType>::CreateHeaderRow(const FArguments& InArgs)
+	{
+		TArray<TReplicationColumn<TItemType>> Columns = InArgs._Columns;
+		Columns.Sort([](const TReplicationColumn<TItemType>& Left, const TReplicationColumn<TItemType>& Right) { return Left.GetColumnSortOrderValue() < Right.GetColumnSortOrderValue(); });
+		
+		HeaderRow = SNew(SHeaderRow);
+		TSet<FName> DuplicateColumnDetection;
+		for (TReplicationColumn<TItemType>& Column : Columns)
+		{
+			check(!DuplicateColumnDetection.Contains(Column.ColumnId));
+			DuplicateColumnDetection.Add(Column.ColumnId);
+			
+			// SHeaderRow owns the columns and deletes them when destroyed
+			TReplicationColumn<TItemType>* ManagedByHeaderRow = new TReplicationColumn<TItemType>(MoveTemp(Column));
+			HeaderRow->AddColumn(*ManagedByHeaderRow);
+		}
+
+		return HeaderRow.ToSharedRef();
+	}
+
+	template <typename TItemType>
+	TSharedRef<ITableRow> SReplicationTreeView<TItemType>::OnGenerateRowWidget(TItemType Item, const TSharedRef<STableViewBase>& OwnerTable)
+	{
+		const typename SReplicationColumnRow<TItemType>::FGetColumn ColumnGetter = SReplicationColumnRow<TItemType>::FGetColumn::CreateLambda([this](const FName& ColumnId) -> const TReplicationColumn<TItemType>*
+		{
+			for (const SHeaderRow::FColumn& Column : HeaderRow->GetColumns())
+			{
+				if (Column.ColumnId == ColumnId)
+				{
+					return static_cast<const TReplicationColumn<TItemType>*>(&Column);
+				}
+			}
+
+			checkNoEntry();
+			return nullptr;
+		});
+	
+		return SNew(SReplicationColumnRow<TItemType>, OwnerTable)
+			.HighlightText(SearchText)
+			.ColumnGetter(ColumnGetter)
+			.RowData(Item)
+			.ExpandableColumnLabel(ExpandableColumnId);
+	}
+
+	template <typename TItemType>
+	void SReplicationTreeView<TItemType>::GetRowChildren(TItemType Item, TArray<TItemType>& OutChildren)
+	{
+		if (OnGetChildrenDelegate.IsBound())
+		{
+			OnGetChildrenDelegate.Execute(Item, [this, &OutChildren](TItemType ItemToAdd)
+			{
+				if (PassesFilters(ItemToAdd))
+				{
+					OutChildren.Add(ItemToAdd);
+				}
+			});
+		}
+	}
+
+	template <typename TItemType>
+	void SReplicationTreeView<TItemType>::OnSearchTextCommitted(const FText& InFilterText, ETextCommit::Type CommitType)
+	{
+		if (!InFilterText.EqualTo(*SearchText))
+		{
+			OnSearchTextChanged(InFilterText);
+		}
+	}
+
+	template <typename TItemType>
+	void SReplicationTreeView<TItemType>::OnSearchTextChanged(const FText& InSearchText)
+	{
+		*SearchText = InSearchText;
+		SearchTextFilter->SetRawFilterText(InSearchText);
+		SearchBox->SetError(SearchTextFilter->GetFilterErrorText());
+	}
+
+	template <typename TItemType>
+	void SReplicationTreeView<TItemType>::PopulateSearchStrings(const TItemType& Item, TArray<FString>& OutSearchStrings)
+	{
+		for (const SHeaderRow::FColumn& Column : HeaderRow->GetColumns())
+		{
+			const TReplicationColumn<TItemType>& CastColumn = static_cast<const TReplicationColumn<TItemType>&>(Column);
+			CastColumn.ExecutePopulateSearchString(Item, OutSearchStrings);
+		}
+	}
+
+	template <typename TItemType>
+	void SReplicationTreeView<TItemType>::OnFilterChanged()
+	{
+		// Try preserving the selected activity.
+		TArray<TItemType> SelectedItems = TreeView->GetSelectedItems();
+
+		// Reset the list of displayed activities.
+		FilteredRootItems.Reset(AllRootItems->Num());
+
+		// Apply the filter.
+		for (const TItemType& Activity : *AllRootItems)
+		{
+			if (PassesFilters(Activity))
+			{
+				FilteredRootItems.Add(Activity);
+			}
+		}
+
+		// Restore/reset the selected activity.
+		SelectedItems.SetNum(Algo::RemoveIf(SelectedItems, [this](const TItemType& Item){ return !FilteredRootItems.Contains(Item); }));
+		if (!SelectedItems.IsEmpty())
+		{
+			TreeView->SetItemSelection(SelectedItems, true); // Restore previous selection.
+			TreeView->RequestScrollIntoView(SelectedItems[0]);
+		}
+
+		TreeView->RequestListRefresh();
+	}
+
+	template <typename TItemType>
+	bool SReplicationTreeView<TItemType>::PassesFilters(const TItemType& Item)
+	{
+		return SearchTextFilter->PassesFilter(Item);
+	}
+}
+
+#undef LOCTEXT_NAMESPACE
\ No newline at end of file
diff --git a/Engine/Plugins/Developer/Concert/ConcertApp/MultiUserClient/Source/MultiUserReplicationEditor/Private/StreamEditor/View/ObjectViewer/SelectionViewerColumns.cpp b/Engine/Plugins/Developer/Concert/ConcertApp/MultiUserClient/Source/MultiUserReplicationEditor/Private/StreamEditor/View/ObjectViewer/SelectionViewerColumns.cpp
new file mode 100644
--- /dev/null
+++ b/Engine/Plugins/Developer/Concert/ConcertApp/MultiUserClient/Source/MultiUserReplicationEditor/Private/StreamEditor/View/ObjectViewer/SelectionViewerColumns.cpp
@@ -0,0 +1,179 @@
+// Copyright Epic Games, Inc. All Rights Reserved.
+
+#include "SelectionViewerColumns.h"
+
+#include "ReplicatedPropertyData.h"
+#include "StreamEditor/View/ObjectViewer/ReplicatedObjectData.h"
+
+#include "StreamEditor/Model/DisplayUtils.h"
+#include "Internationalization/Internationalization.h"
+#include "StreamEditor/Model/IEditableObjectToPropertiesModel.h"
+#include "StreamEditor/View/ObjectEditor/SPropertyReplicationSelectionEditor.h"
+#include "Textures/SlateIcon.h"
+#include "Widgets/Images/SImage.h"
+#include "Widgets/SBoxPanel.h"
+#include "Widgets/Text/STextBlock.h"
+#include "Widgets/Input/SCheckBox.h"
+
+#define LOCTEXT_NAMESPACE "ReplicationObjectColumns"
+
+namespace UE::MultiUserReplicationEditor::ReplicationObjectColumns
+{
+	const FName IconColumnId = TEXT("IconColumn");
+	const FName LabelColumnId = TEXT("LabelColumn");
+	const FName TypeColumnId = TEXT("TypeColumn");
+
+	FReplicationObjectColumn IconColumn(TSharedRef<IObjectToPropertiesModel> Model)
+	{
+		return FReplicationObjectColumn(
+			FReplicationObjectColumn::FArguments()
+				.GenerateWidgetColumn_Lambda([Model](const FReplicationObjectColumn::FBuildArgs& Args)
+				{
+					return SNew(SHorizontalBox)
+						+SHorizontalBox::Slot()
+						.HAlign(HAlign_Center)
+						.VAlign(VAlign_Center)
+						[
+							SNew(SImage)
+							.Image(DisplayUtils::GetObjectIcon(*Model, Args.RowData->GetObjectPath()).GetOptionalIcon())
+						];
+				})
+				.ColumnSortOrder(static_cast<int32>(EReplicationColumnOrder::Icon)),
+			SHeaderRow::Column(IconColumnId)
+				.DefaultLabel(FText::GetEmpty())
+				.FixedWidth(FMultiUserReplicationEditorStyle::Get().GetFloat("ReplicationTreeView.IconColumnWidth"))
+			);
+	}
+	
+	FReplicationObjectColumn LabelColumn()
+	{
+		return FReplicationObjectColumn(
+			FReplicationObjectColumn::FArguments()
+				.GenerateWidgetColumn_Lambda([](const FReplicationObjectColumn::FBuildArgs& Args)
+				{
+					return SNew(STextBlock)
+						.HighlightText(TAttribute<FText>::CreateLambda([HighlightText = Args.HighlightText](){ return *HighlightText; }))
+						.Text(DisplayUtils::GetObjectDisplayText(Args.RowData->GetObjectPath()));
+				})
+				.PopulateSearchItems_Lambda([](const TSharedPtr<FReplicatedObjectData>& ObjectData, TArray<FString>& InOutSearchStrings)
+				{
+					InOutSearchStrings.Add(DisplayUtils::GetObjectDisplayText(ObjectData->GetObjectPath()).ToString());
+				})
+				.ColumnSortOrder(static_cast<int32>(EReplicationColumnOrder::Label)),
+			SHeaderRow::Column(LabelColumnId)
+				.DefaultLabel(LOCTEXT("LabelColumnLabel", "Label"))
+				.FillWidth(1.f)
+			);
+	}
+	
+	FReplicationObjectColumn TypeColumn(TSharedRef<IObjectToPropertiesModel> Model)
+	{
+		return FReplicationObjectColumn(
+			FReplicationObjectColumn::FArguments()
+				.GenerateWidgetColumn_Lambda([Model](const FReplicationObjectColumn::FBuildArgs& Args)
+				{
+					return SNew(STextBlock)
+						.HighlightText(TAttribute<FText>::CreateLambda([HighlightText = Args.HighlightText](){ return *HighlightText; }))
+						.Text(DisplayUtils::GetObjectTypeText(*Model, Args.RowData->GetObjectPath()));
+				})
+				.PopulateSearchItems_Lambda([Model](const TSharedPtr<FReplicatedObjectData>& ObjectData, TArray<FString>& InOutSearchStrings)
+				{
+					InOutSearchStrings.Add(DisplayUtils::GetObjectTypeText(Model.Get(), ObjectData->GetObjectPath()).ToString());
+				})
+				.ColumnSortOrder(static_cast<int32>(EReplicationColumnOrder::Type)),
+			SHeaderRow::Column(TypeColumnId)
+				.DefaultLabel(LOCTEXT("TypeColumnLabel", "Type"))
+				.FillWidth(1.f)
+			);
+	}
+}
+
+#undef LOCTEXT_NAMESPACE
+#define LOCTEXT_NAMESPACE "ReplicationPropertyColumns"
+
+namespace UE::MultiUserReplicationEditor::ReplicationPropertyColumns
+{
+	const FName ReplicatesColumnId = TEXT("ReplicatesColumn");
+	const FName LabelColumnId = TEXT("LabelColumn");
+
+	FReplicationPropertyColumn ReplicatesColumns(TSharedRef<SPropertyReplicationSelectionEditor> EditorWidget, TSharedRef<IEditableObjectToPropertiesModel> Model)
+	{
+		return FReplicationPropertyColumn(
+			FReplicationPropertyColumn::FArguments()
+				.GenerateWidgetColumn_Lambda([EditorWidget = MoveTemp(EditorWidget), Model = MoveTemp(Model)](const FReplicationPropertyColumn::FBuildArgs& Args)
+				{
+					return SNew(SCheckBox)
+						.ToolTipText(LOCTEXT("ReplicatesCheckbox.Tooltip", "Should replicate?"))
+						.IsChecked_Lambda([EditorWidget, Model, RowData = Args.RowData]()
+						{
+							return GetPropertyCheckboxStateBasedOnSelection(*RowData.Get(), EditorWidget->GetSelectedObjects(), *Model);
+						})
+						.OnCheckStateChanged_Lambda([EditorWidget, Model, RowData = Args.RowData](ECheckBoxState NewState)
+						{
+							TArray Properties{ RowData->GetProperty() };
+							for (const TSharedPtr<FReplicatedObjectData>& SelectedObject : EditorWidget->GetSelectedObjects())
+							{
+								if (NewState == ECheckBoxState::Checked)
+								{
+									Model->AddProperties(SelectedObject->GetObjectPath(), Properties);
+								}
+								else
+								{
+									Model->RemoveProperties(SelectedObject->GetObjectPath(), Properties);
+								}
+							}
+						});
+				})
+				.ColumnSortOrder(static_cast<int32>(EReplicationPropertyColumnOrder::ReplicatesCheckbox)),
+			SHeaderRow::Column(ReplicatesColumnId)
+				.DefaultLabel(FText::GetEmpty())
+				.FixedWidth(FMultiUserReplicationEditorStyle::Get().GetFloat("ReplicationTreeView.IconColumnWidth"))
+			);
+	}
+	
+	FReplicationPropertyColumn LabelColumn()
+	{
+		return FReplicationPropertyColumn(
+			FReplicationPropertyColumn::FArguments()
+				.GenerateWidgetColumn_Lambda([](const FReplicationPropertyColumn::FBuildArgs& Args)
+				{
+					return SNew(STextBlock)
+						.HighlightText(TAttribute<FText>::CreateLambda([HighlightText = Args.HighlightText](){ return *HighlightText; }))
+						.Text(DisplayUtils::GetPropertyDisplayText(Args.RowData->GetProperty()));
+				})
+				.PopulateSearchItems_Lambda([](const TSharedPtr<FReplicatedPropertyData>& ObjectData, TArray<FString>& InOutSearchStrings)
+				{
+					InOutSearchStrings.Add(DisplayUtils::GetPropertyDisplayText(ObjectData->GetProperty()).ToString());
+				})
+				.ColumnSortOrder(static_cast<int32>(EReplicationPropertyColumnOrder::Label)),
+			SHeaderRow::Column(LabelColumnId)
+				.DefaultLabel(LOCTEXT("LabelColumnLabel", "Label"))
+				.FillWidth(1.f)
+			);
+	}
+
+	ECheckBoxState GetPropertyCheckboxStateBasedOnSelection(const FReplicatedPropertyData& RowData, TConstArrayView<TSharedPtr<FReplicatedObjectData>> Selection, const IObjectToPropertiesModel& Model)
+	{
+		ECheckBoxState CheckBoxState = ECheckBoxState::Undetermined;
+		for (const TSharedPtr<FReplicatedObjectData>& SelectedObject : Selection)
+		{
+			const bool bContainsProperty = Model.ContainsProperties(SelectedObject->GetObjectPath(), { RowData.GetProperty() });
+			const ECheckBoxState StateForThisObject = bContainsProperty
+				? ECheckBoxState::Checked
+				: ECheckBoxState::Unchecked;
+			// The first object?
+			if (CheckBoxState == ECheckBoxState::Undetermined)
+			{
+				CheckBoxState = StateForThisObject;
+			}
+			else if (CheckBoxState != StateForThisObject)
+			{
+				// Two boxes do not have the same value
+				return ECheckBoxState::Undetermined;
+			}
+		}
+		return CheckBoxState;
+	}
+}
+
+#undef LOCTEXT_NAMESPACE
\ No newline at end of file
diff --git a/Engine/Plugins/Developer/Concert/ConcertApp/MultiUserClient/Source/MultiUserReplicationEditor/Private/StreamEditor/View/ObjectViewer/SelectionViewerColumns.h b/Engine/Plugins/Developer/Concert/ConcertApp/MultiUserClient/Source/MultiUserReplicationEditor/Private/StreamEditor/View/ObjectViewer/SelectionViewerColumns.h
new file mode 100644
--- /dev/null
+++ b/Engine/Plugins/Developer/Concert/ConcertApp/MultiUserClient/Source/MultiUserReplicationEditor/Private/StreamEditor/View/ObjectViewer/SelectionViewerColumns.h
@@ -0,0 +1,67 @@
+// Copyright Epic Games, Inc. All Rights Reserved.
+
+#pragma once
+
+#include "Containers/ArrayView.h"
+#include "ReplicationColumn.h"
+#include "Templates/SharedPointer.h"
+
+namespace UE::MultiUserReplicationEditor
+{
+	class FReplicatedObjectData;
+	class FReplicatedPropertyData;
+	class IEditableObjectToPropertiesModel;
+	class IObjectToPropertiesModel;
+	class SPropertyReplicationSelectionEditor;
+}
+
+namespace UE::MultiUserReplicationEditor::ReplicationObjectColumns
+{
+	using FReplicationObjectColumn = TReplicationColumn<TSharedPtr<FReplicatedObjectData>>;
+	
+	enum class EReplicationColumnOrder : int32
+	{
+		/** Displays the class icon */
+		Icon = 10,
+		/** Label of the object */
+		Label = 20,
+		/** Class of the object */
+		Type = 30,
+	};
+	
+	extern const FName IconColumnId;
+	extern const FName LabelColumnId;
+	extern const FName TypeColumnId;
+
+	FReplicationObjectColumn IconColumn(TSharedRef<IObjectToPropertiesModel> Model);
+	FReplicationObjectColumn LabelColumn();
+	FReplicationObjectColumn TypeColumn(TSharedRef<IObjectToPropertiesModel> Model);
+}
+
+namespace UE::MultiUserReplicationEditor::ReplicationPropertyColumns
+{
+	using FReplicationPropertyColumn = TReplicationColumn<TSharedPtr<FReplicatedPropertyData>>;
+	
+	enum class EReplicationPropertyColumnOrder : int32
+	{
+		/** The checkbox in SPropertyReplicationSelectionEditor determining whether the property is in the selection*/
+		ReplicatesCheckbox = 0,
+		/** Label of the property */
+		Label = 10
+	};
+	
+	/** The checkbox in SPropertyReplicationSelectionEditor determining whether the property is in the selection*/
+	extern const FName ReplicatesColumnId;
+	extern const FName LabelColumnId;
+
+	FReplicationPropertyColumn ReplicatesColumns(TSharedRef<SPropertyReplicationSelectionEditor> EditorWidget, TSharedRef<IEditableObjectToPropertiesModel> Model);
+	FReplicationPropertyColumn LabelColumn();
+
+	/**
+	 * Goes through all selected objects, checks whether the property is checked on it or not, and returns a checkbox state.
+	 * The state is checked or unchecked if all selected properties have the property enabled or disabled, respectively, and undetermined otherwise, i.e. if it is mixed.
+	 *
+	 * This is used by ReplicatesColumns.
+	 */
+	ECheckBoxState GetPropertyCheckboxStateBasedOnSelection(const FReplicatedPropertyData& RowData, TConstArrayView<TSharedPtr<FReplicatedObjectData>> Selection, const IObjectToPropertiesModel& Model);
+}
diff --git a/Engine/Plugins/Developer/Concert/ConcertApp/MultiUserClient/Source/MultiUserReplicationEditor/Private/StreamEditor/View/SReplicationStreamEditor.cpp b/Engine/Plugins/Developer/Concert/ConcertApp/MultiUserClient/Source/MultiUserReplicationEditor/Private/StreamEditor/View/SReplicationStreamEditor.cpp
new file mode 100644
--- /dev/null
+++ b/Engine/Plugins/Developer/Concert/ConcertApp/MultiUserClient/Source/MultiUserReplicationEditor/Private/StreamEditor/View/SReplicationStreamEditor.cpp
@@ -0,0 +1,31 @@
+// Copyright Epic Games, Inc. All Rights Reserved.
+
+#include "SReplicationStreamEditor.h"
+
+#include "MultiUserReplicationStreamAsset.h"
+#include "StreamEditor/Model/Object/EditorObjectSelectionSourceModel.h"
+#include "StreamEditor/Model/Property/SelectPropertyFromUClassModel.h"
+#include "StreamEditor/Model/PropertySelectionAssetModel.h"
+#include "StreamEditor/View/ObjectEditor/SPropertyReplicationSelectionEditor.h"
+#include "Widgets/SBoxPanel.h"
+
+namespace UE::MultiUserReplicationEditor
+{
+	void SReplicationStreamEditor::Construct(const FArguments& InArgs, UMultiUserReplicationStreamAsset& InEditedAsset)
+	{
+		AssetReadWriteModel = MakeShared<FPropertySelectionAssetModel>(*InEditedAsset.ReplicationList);
+		ObjectSourceModel = MakeShared<FEditorObjectSelectionSourceModel>();
+		PropertySourceModel = MakeShared<FSelectPropertyFromUClassModel>();
+		ChildSlot
+		[
+			SNew(SVerticalBox)
+
+			+SVerticalBox::Slot()
+			[
+				SNew(SPropertyReplicationSelectionEditor, AssetReadWriteModel.ToSharedRef(), ObjectSourceModel.ToSharedRef(), PropertySourceModel.ToSharedRef())
+			]
+		];
+	}
+}
+
+
diff --git a/Engine/Plugins/Developer/Concert/ConcertApp/MultiUserClient/Source/MultiUserReplicationEditor/Private/StreamEditor/View/SReplicationStreamEditor.h b/Engine/Plugins/Developer/Concert/ConcertApp/MultiUserClient/Source/MultiUserReplicationEditor/Private/StreamEditor/View/SReplicationStreamEditor.h
new file mode 100644
--- /dev/null
+++ b/Engine/Plugins/Developer/Concert/ConcertApp/MultiUserClient/Source/MultiUserReplicationEditor/Private/StreamEditor/View/SReplicationStreamEditor.h
@@ -0,0 +1,39 @@
+// Copyright Epic Games, Inc. All Rights Reserved.
+
+#pragma once
+
+#include "Widgets/DeclarativeSyntaxSupport.h"
+#include "Widgets/SCompoundWidget.h"
+
+class UMultiUserReplicationStreamAsset;
+
+namespace UE::MultiUserReplicationEditor
+{
+	class FEditorObjectSelectionSourceModel;
+	class FSelectPropertyFromUClassModel;
+	class FPropertySelectionAssetModel;
+	class SPropertyReplicationSelectionEditor;
+	
+	/** Root widget for editing UMultiUserReplicationStreamAsset. */
+	class MULTIUSERREPLICATIONEDITOR_API SReplicationStreamEditor : public SCompoundWidget
+	{
+	public:
+		SLATE_BEGIN_ARGS(SReplicationStreamEditor)
+		{}
+		SLATE_END_ARGS()
+
+		void Construct(const FArguments& InArgs, UMultiUserReplicationStreamAsset& InEditedAsset);
+
+	private:
+
+		/** Edits the selected properties */
+		TSharedPtr<SPropertyReplicationSelectionEditor> PropertySelectionEditor;
+		
+		/** Passed to SPropertyReplicationSelectionEditor which will read and write to the UMultiUserReplicationStreamAsset */
+		TSharedPtr<FPropertySelectionAssetModel> AssetReadWriteModel;
+		/** Passed to SPropertyReplicationSelectionEditor. It controls which objects can be selected for the replication list. */
+		TSharedPtr<FEditorObjectSelectionSourceModel> ObjectSourceModel;
+		/** Passed to SPropertyReplicationSelectionEditor. It controls which properties can be added to objects. */
+		TSharedPtr<FSelectPropertyFromUClassModel> PropertySourceModel;
+	};
+}
diff --git a/Engine/Plugins/Developer/Concert/ConcertApp/MultiUserClient/Source/MultiUserReplicationEditor/Public/IMultiUserReplicationEditorModule.h b/Engine/Plugins/Developer/Concert/ConcertApp/MultiUserClient/Source/MultiUserReplicationEditor/Public/IMultiUserReplicationEditorModule.h
new file mode 100644
--- /dev/null
+++ b/Engine/Plugins/Developer/Concert/ConcertApp/MultiUserClient/Source/MultiUserReplicationEditor/Public/IMultiUserReplicationEditorModule.h
@@ -0,0 +1,43 @@
+// Copyright Epic Games, Inc. All Rights Reserved.
+
+#pragma once
+
+#include "Modules/ModuleInterface.h"
+#include "Modules/ModuleManager.h"
+#include "UObject/NameTypes.h"
+
+struct FAssetCategoryPath;
+
+namespace UE::MultiUserReplicationEditor
+{
+	class IMultiUserReplicationEditorModule : public IModuleInterface
+	{
+	public:
+
+		/**
+		 * Singleton-like access to this module's interface.  This is just for convenience!
+		 * Beware of calling this during the shutdown phase, though.  Your module might have been unloaded already.
+		 *
+		 * @return Returns singleton instance, loading the module on demand if needed
+		 */
+		static inline IMultiUserReplicationEditorModule& Get()
+		{
+			static const FName ModuleName = "MultiUserReplicationEditor";
+			return FModuleManager::LoadModuleChecked<IMultiUserReplicationEditorModule>(ModuleName);
+		}
+
+		/**
+		 * Checks to see if this module is loaded and ready.  It is only valid to call Get() during shutdown if IsAvailable() returns true.
+		 *
+		 * @return True if the module is loaded and ready to use
+		 */
+		static inline bool IsAvailable()
+		{
+			static const FName ModuleName = "MultiUserReplicationEditor";
+			return FModuleManager::Get().IsModuleLoaded(ModuleName);
+		}
+
+		/** Gets the category under which assets for Multi-User Replication are nested. */
+		virtual FAssetCategoryPath GetMultiUserReplicationCategory() const = 0;
+	};
+}
\ No newline at end of file
diff --git a/Engine/Plugins/Developer/Concert/ConcertApp/MultiUserClient/Source/MultiUserReplicationEditor/Public/Settings/DefaultPropertySelection.h b/Engine/Plugins/Developer/Concert/ConcertApp/MultiUserClient/Source/MultiUserReplicationEditor/Public/Settings/DefaultPropertySelection.h
new file mode 100644
--- /dev/null
+++ b/Engine/Plugins/Developer/Concert/ConcertApp/MultiUserClient/Source/MultiUserReplicationEditor/Public/Settings/DefaultPropertySelection.h
@@ -0,0 +1,23 @@
+// Copyright Epic Games, Inc. All Rights Reserved.
+
+#pragma once
+
+#include "DefaultPropertySelection.generated.h"
+
+USTRUCT()
+struct FDefaultPropertySelection
+{
+	GENERATED_BODY()
+
+	/**
+	 * A list of properties that should be selected by default when you add a new object type.
+	 * Specify properties by using the names as the appear in the replication editor and each property with ".".
+	 * Example: "RelativeLocation.X"
+	 */
+	UPROPERTY(EditAnywhere, Category = "Config")
+	TArray<FString> DefaultSelectedProperties;
+
+	/** Whether to inherit the properties from the base class (you need not list them in DefaultSelectedProperties). */
+	UPROPERTY(EditAnywhere, Category = "Config")
+	bool bInheritFromBase = true;
+};
diff --git a/Engine/Plugins/Developer/Concert/ConcertApp/MultiUserClient/Source/MultiUserReplicationEditor/Public/Settings/MultiUserReplicationSettings.h b/Engine/Plugins/Developer/Concert/ConcertApp/MultiUserClient/Source/MultiUserReplicationEditor/Public/Settings/MultiUserReplicationSettings.h
new file mode 100644
--- /dev/null
+++ b/Engine/Plugins/Developer/Concert/ConcertApp/MultiUserClient/Source/MultiUserReplicationEditor/Public/Settings/MultiUserReplicationSettings.h
@@ -0,0 +1,29 @@
+// Copyright Epic Games, Inc. All Rights Reserved.
+
+#pragma once
+
+#include "UObject/Object.h"
+#include "MultiUserReplicationClientProfileAsset.h"
+#include "Settings/DefaultPropertySelection.h"
+#include "MultiUserReplicationSettings.generated.h"
+
+/** Asset for users to describe a client to the server. */
+UCLASS(Config=MultiUserClient)
+class MULTIUSERREPLICATIONEDITOR_API UMultiUserReplicationSettings : public UObject
+{
+	GENERATED_BODY()
+public:
+	
+	/** The client profile to use when connecting to a session. */
+	UPROPERTY(EditAnywhere, Config, Category = "Replication")
+	TSoftObjectPtr<UMultiUserReplicationClientProfileAsset> DefaultProfile;
+
+	/** Properties you want selected by default when you add a new replicated object in the editor */
+	UPROPERTY(EditAnywhere, Config, Category = "Replication|Editor")
+	TMap<FSoftClassPath, FDefaultPropertySelection> DefaultPropertySelection;
+
+	static UMultiUserReplicationSettings* Get() { return GetMutableDefault<UMultiUserReplicationSettings>(); }
+
+	/** Reads DefaultPropertySelection and applies any default property selections to Info based on the Class just added. */
+	void AddDefaultPropertiesFromSettings(FReplicatedObjectInfo& Info, UClass& Class);
+};
\ No newline at end of file
diff --git a/Engine/Plugins/Developer/Concert/ConcertSync/ConcertSyncCore/Source/ConcertSyncCore/Private/Replication/Data/ConcertPropertySelection.cpp b/Engine/Plugins/Developer/Concert/ConcertSync/ConcertSyncCore/Source/ConcertSyncCore/Private/Replication/Data/ConcertPropertySelection.cpp
new file mode 100644
--- /dev/null
+++ b/Engine/Plugins/Developer/Concert/ConcertSync/ConcertSyncCore/Source/ConcertSyncCore/Private/Replication/Data/ConcertPropertySelection.cpp
@@ -0,0 +1,137 @@
+// Copyright Epic Games, Inc. All Rights Reserved.
+
+#include "Replication/Data/ConcertPropertySelection.h"
+
+#include "ConcertLogGlobal.h"
+#include "Replication/PropertyChainUtils.h"
+
+#include "Serialization/ArchiveSerializedPropertyChain.h"
+
+const FName FConcertPropertyChain::InternalContainerPropertyValueName(TEXT("Value"));
+
+TOptional<FConcertPropertyChain> FConcertPropertyChain::CreateFromPath(UStruct& Class, const TArray<FName>& NamePath)
+{
+	TOptional<FConcertPropertyChain> Result;
+	// The goal here is to find a valid FProperty path based on the NamePath...
+	UE::ConcertSyncCore::PropertyChain::ForEachReplicatableConcertProperty(Class, [&NamePath, &Result](FConcertPropertyChain&& Chain) mutable
+	{
+		// ... and performance could be improved by not expanding down the search tree by introducing ETreeBreakBehavior::SkipSubtree
+		if (Chain == NamePath)
+		{
+			Result.Emplace(MoveTemp(Chain));
+			return EBreakBehavior::Break;
+		}
+		return EBreakBehavior::Continue;
+	});
+	return Result;
+}
+
+FConcertPropertyChain::FConcertPropertyChain(const FArchiveSerializedPropertyChain* OptionalChain, const FProperty& LeafProperty)
+{
+	using namespace UE::ConcertSyncCore::PropertyChain;
+	
+	if (OptionalChain)
+	{
+		PathToProperty.Reserve(OptionalChain->GetNumProperties() + 1);
+		for (int32 i = 0; i < OptionalChain->GetNumProperties(); ++i)
+		{
+			const FProperty* CurrentProperty = OptionalChain->GetPropertyFromRoot(i);
+			
+			const FProperty* ParentProperty = CurrentProperty->GetOwner<FProperty>();
+			const bool bIsKeyOfMap = ParentProperty && ParentProperty->IsA(FMapProperty::StaticClass()) && CastField<FMapProperty>(ParentProperty)->KeyProp == CurrentProperty;
+			if (!ensureMsgf(!bIsKeyOfMap, TEXT("Key child properties never appear in a path because all of them are assumed to be replicated. See documentation of PathToProperty!")))
+			{
+				// The input was invalid but we've constructed FConcertPropertyChain with valid state: it ends at the TMap property.
+				return;
+			}
+			
+			if (!IsInnerContainerProperty(*CurrentProperty))
+			{
+				PathToProperty.Add(OptionalChain->GetPropertyFromRoot(i)->GetFName());
+			}
+
+			if (!ensureMsgf(!IsNativeStructProperty(*CurrentProperty), TEXT("Child properties of structs implementing a custom Serialize function never appear in a path. See documentation of PathToProperty!")))
+			{
+				// The input was invalid but we've constructed FConcertPropertyChain with valid state: it ends at the native struct property.
+				return;
+			}
+		}
+	}
+	else if (!ensureMsgf(LeafProperty.GetOwner<FProperty>() == nullptr, TEXT("Not a leaf property!")))
+	{
+		// The input was invalid but we've constructed FConcertPropertyChain with valid state: no property.
+		return;
+	}
+
+	// LeafProperty is allowed to be an inner property of FArrayProperty, FSetProperty, FMapProperty::ValueProp but only if it is primitive or a native structs
+	// In that case the property display label is InternalContainerPropertyValueName, which is "Value"
+	if (IsInnerContainerProperty(LeafProperty))
+	{
+		if (ensureMsgf(IsPrimitiveProperty(LeafProperty) || IsNativeStructProperty(LeafProperty), TEXT("The chain only contains inner properties of primitives or native structs!")))
+		{
+			PathToProperty.Add(InternalContainerPropertyValueName);
+		}
+	}
+	else
+	{
+		PathToProperty.Add(LeafProperty.GetFName());
+	}
+
+	UE_CLOG(!IsReplicatableProperty(LeafProperty), LogConcert, Warning, TEXT("Instantiated property chain with non replicable property: %s"), *ToString());
+}
+
+bool FConcertPropertyChain::IsChildOf(const FConcertPropertyChain& ParentToCheck) const
+{
+	if (PathToProperty.Num() < ParentToCheck.PathToProperty.Num())
+	{
+		return false;
+	}
+
+	for (int32 i = 0; i < ParentToCheck.PathToProperty.Num(); ++i)
+	{
+		if (PathToProperty[i] != ParentToCheck.PathToProperty[i])
+		{
+			return false;
+		}
+	}
+
+	return true;
+}
+
+bool FConcertPropertyChain::IsDirectChildOf(const FConcertPropertyChain& ParentToCheck) const
+{
+	if (ParentToCheck.PathToProperty.Num() != PathToProperty.Num() - 1)
+	{
+		return false;
+	}
+
+	bool bArePathsEqual = true;
+	for (int32 i = 0; bArePathsEqual && i < ParentToCheck.PathToProperty.Num(); ++i)
+	{
+		bArePathsEqual &= ParentToCheck.PathToProperty[i] == PathToProperty[i];
+	}
+	return bArePathsEqual;
+}
+
+FString FConcertPropertyChain::ToString(EToStringMethod Method) const
+{
+	switch (Method)
+	{
+	case EToStringMethod::Path:
+		return FString::JoinBy(PathToProperty, TEXT("."), [](const FName& Name){ return Name.ToString(); });
+	case EToStringMethod::LeafProperty:
+		return PathToProperty.IsEmpty()
+			? FName(NAME_None).ToString()
+			: PathToProperty[PathToProperty.Num() - 1].ToString();
+	default:
+		checkNoEntry();
+		return FName(NAME_None).ToString();
+	}
+	
+}
+
+uint32 GetTypeHash(const FConcertPropertyChain& Chain)
+{
+	return GetTypeHash(Chain.GetPathToProperty());
+}
+
diff --git a/Engine/Plugins/Developer/Concert/ConcertSync/ConcertSyncCore/Source/ConcertSyncCore/Private/Replication/Data/ReplicationStreamDescription.cpp b/Engine/Plugins/Developer/Concert/ConcertSync/ConcertSyncCore/Source/ConcertSyncCore/Private/Replication/Data/ReplicationStreamDescription.cpp
new file mode 100644
--- /dev/null
+++ b/Engine/Plugins/Developer/Concert/ConcertSync/ConcertSyncCore/Source/ConcertSyncCore/Private/Replication/Data/ReplicationStreamDescription.cpp
@@ -0,0 +1,18 @@
+// Copyright Epic Games, Inc. All Rights Reserved.
+
+#include "Replication/Data/ReplicationStreamDescription.h"
+
+#include "Misc/Optional.h"
+
+FReplicationStreamDescription_NetPacked FReplicationStreamDescription::Pack() const
+{
+	// TODO: Implement once stream attributes are added
+	return FReplicationStreamDescription_NetPacked{ BaseDescription };
+}
+
+TOptional<FReplicationStreamDescription> FReplicationStreamDescription_NetPacked::Unpack(FString* OutErrorMessage) const
+{
+	// TODO: Implement once stream attributes are added
+	return FReplicationStreamDescription{ BaseDescription };
+}
+
diff --git a/Engine/Plugins/Developer/Concert/ConcertSync/ConcertSyncCore/Source/ConcertSyncCore/Private/Replication/PropertyChainUtils.cpp b/Engine/Plugins/Developer/Concert/ConcertSync/ConcertSyncCore/Source/ConcertSyncCore/Private/Replication/PropertyChainUtils.cpp
new file mode 100644
--- /dev/null
+++ b/Engine/Plugins/Developer/Concert/ConcertSync/ConcertSyncCore/Source/ConcertSyncCore/Private/Replication/PropertyChainUtils.cpp
@@ -0,0 +1,197 @@
+// Copyright Epic Games, Inc. All Rights Reserved.
+
+#include "Replication/PropertyChainUtils.h"
+
+#include "Misc/ScopeExit.h"
+#include "Serialization/ArchiveSerializedPropertyChain.h"
+#include "UObject/UnrealType.h"
+
+namespace UE::ConcertSyncCore::PropertyChain
+{
+	namespace Private
+	{
+		EBreakBehavior VisitPropertyRecursive(FArchiveSerializedPropertyChain& Chain, FProperty& Property, TFunctionRef<EBreakBehavior(const FArchiveSerializedPropertyChain& Chain, const FProperty&LeafProperty)> ProcessProperty);
+		
+		EBreakBehavior VisitStructPropertyRecursive(
+			FArchiveSerializedPropertyChain& Chain,
+			FStructProperty& StructProperty,
+			TFunctionRef<EBreakBehavior(const FArchiveSerializedPropertyChain& Chain, const FProperty&LeafProperty)> ProcessProperty
+			)
+		{
+			// If the struct is in a container, then the container should be pushed, not the inner property.
+			FProperty& PropertyToPush = IsInnerContainerProperty(StructProperty)
+				? *StructProperty.GetOwner<FProperty>()
+				: StructProperty;
+			Chain.PushProperty(&PropertyToPush, PropertyToPush.IsEditorOnlyProperty());
+			ON_SCOPE_EXIT { Chain.PopProperty(&PropertyToPush, PropertyToPush.IsEditorOnlyProperty()); };
+			
+			for (TFieldIterator<FProperty> FieldIt(StructProperty.Struct); FieldIt; ++FieldIt)
+			{
+				if (VisitPropertyRecursive(Chain, **FieldIt, ProcessProperty) == EBreakBehavior::Break)
+				{
+					return EBreakBehavior::Break;
+				}
+			}
+			
+			return EBreakBehavior::Continue;
+		}
+		
+		EBreakBehavior VisitPropertyRecursive(
+			FArchiveSerializedPropertyChain& Chain,
+			FProperty& Property,
+			TFunctionRef<EBreakBehavior(const FArchiveSerializedPropertyChain& Chain, const FProperty&LeafProperty)> ProcessProperty
+			)
+		{
+			if (!IsReplicatableProperty(Property))
+			{
+				return EBreakBehavior::Continue;
+			}
+
+			if (ProcessProperty(Chain, Property) == EBreakBehavior::Break)
+			{
+				return EBreakBehavior::Break;
+			}
+			
+			if (FStructProperty* StructProperty = CastField<FStructProperty>(&Property))
+			{
+				// Handle FConcertPropertyChain::InternalContainerPropertyValueName case
+				// If the struct defines a custom serialize function, it does not make sense to list any uproperties we find.
+				// The serialize function can (and often will) skip certain properties so exposing them to a user for selection makes no sense.
+				// If this struct property is supposed to be replicated, then it is all the Serialize function serializes or nothing.
+				return IsNativeStructProperty(*StructProperty)
+					? EBreakBehavior::Continue
+					: VisitStructPropertyRecursive(Chain, *StructProperty, ProcessProperty);
+			}
+
+			auto HandleContainer = [&Chain, &Property, &ProcessProperty](FProperty& Inner)
+			{
+				// Handle FConcertPropertyChain::InternalContainerPropertyValueName case
+				if (IsPrimitiveProperty(Inner)
+					|| (Inner.IsA(FStructProperty::StaticClass()) && IsNativeStructProperty(Inner)))
+				{
+					Chain.PushProperty(&Property, Property.IsEditorOnlyProperty());
+					ON_SCOPE_EXIT{ Chain.PopProperty(&Property, Property.IsEditorOnlyProperty()); };
+					return ProcessProperty(Chain, Inner);
+				}
+				
+				FStructProperty* InnerStructProperty = CastField<FStructProperty>(&Inner);
+				return InnerStructProperty && IsReplicatableProperty(Inner)
+					? VisitStructPropertyRecursive(Chain, *InnerStructProperty, ProcessProperty)
+					: EBreakBehavior::Continue;
+			};
+
+			if (FArrayProperty* ArrayProperty = CastField<FArrayProperty>(&Property))
+			{
+				return HandleContainer(*ArrayProperty->Inner);
+			}
+
+			if (FSetProperty* SetProperty = CastField<FSetProperty>(&Property))
+			{
+				
+				return HandleContainer(*SetProperty->ElementProp);
+			}
+
+			if (FMapProperty* MapProperty = CastField<FMapProperty>(&Property))
+			{
+				return IsReplicatableProperty(*MapProperty->KeyProp)
+					? HandleContainer(*MapProperty->ValueProp)
+					: EBreakBehavior::Continue;
+			}
+			
+			return EBreakBehavior::Continue;
+		}
+	}
+	
+	void ForEachReplicatableProperty(
+		UStruct& Class,
+		TFunctionRef<EBreakBehavior(const FArchiveSerializedPropertyChain& Chain, const FProperty& LeafProperty)> ProcessProperty)
+	{
+		FArchiveSerializedPropertyChain Chain;
+		for (TFieldIterator<FProperty> FieldIt(&Class); FieldIt; ++FieldIt)
+		{
+			if (Private::VisitPropertyRecursive(Chain, **FieldIt, ProcessProperty) == EBreakBehavior::Break)
+			{
+				break;
+			}
+		}
+	}
+
+	void ForEachReplicatableConcertProperty(
+		UStruct& Class,
+		TFunctionRef<EBreakBehavior(FConcertPropertyChain&& PropertyChain)> ProcessProperty)
+	{
+		ForEachReplicatableProperty(Class, [&ProcessProperty](const FArchiveSerializedPropertyChain& Chain, const FProperty& LeafProperty)
+		{
+			return ProcessProperty(FConcertPropertyChain(&Chain, LeafProperty));
+		});
+	}
+
+	void BulkConstructConcertChainsFromPaths(UStruct& Class, uint32 NumPaths, TFunctionRef<bool(const FArchiveSerializedPropertyChain& Chain, const FProperty& LeafProperty)> MatchesPath)
+	{
+		if (NumPaths == 0)
+		{
+			return;
+		}
+		
+		ForEachReplicatableProperty(Class, [&NumPaths, &MatchesPath](const FArchiveSerializedPropertyChain& Chain, const FProperty& LeafProperty) mutable
+		{
+			if (MatchesPath(Chain, LeafProperty))
+			{
+				--NumPaths;
+				return NumPaths == 0
+					? EBreakBehavior::Break
+					: EBreakBehavior::Continue;
+			}
+			return EBreakBehavior::Continue;
+		});
+	}
+	
+	bool DoPathAndChainsMatch(const TArray<FName>& Path, const FArchiveSerializedPropertyChain& Chain, const FProperty& LeafProperty)
+	{
+		if (Chain.GetNumProperties() + 1 != Path.Num())
+		{
+			return false;
+		}
+		
+		for (int32 i = 0; i < Path.Num() - 1; ++i)
+		{
+			if (Path[i] != Chain.GetPropertyFromRoot(0)->GetFName())
+			{
+				return false;
+			}
+		}
+
+		return Path[Path.Num() - 1] == LeafProperty.GetFName();
+	}
+	
+	bool IsReplicatableProperty(const FProperty& LeafProperty)
+	{
+		return !LeafProperty.HasAnyPropertyFlags(
+			// Replicating delegates makes no sense
+			CPF_BlueprintAssignable
+			// It does not make sense to serialize the reference of an instanced subobject. Instead the subobject should be added to the list of replicated properties.
+			| CPF_InstancedReference /* for object ptrs */ | CPF_ContainsInstancedReference /* for containers of object ptrs */
+			);
+	}
+
+	bool IsInnerContainerProperty(const FProperty& Property)
+	{
+		const FProperty* ParentProperty = Property.GetOwner<FProperty>();
+		const bool bIsInnerContainerProperty = ParentProperty
+			&& (ParentProperty->IsA(FArrayProperty::StaticClass()) || ParentProperty->IsA(FSetProperty::StaticClass()) || ParentProperty->IsA(FMapProperty::StaticClass()));
+		return bIsInnerContainerProperty;
+	}
+
+	bool IsPrimitiveProperty(const FProperty& Property)
+	{
+		return Property.IsA(FNumericProperty::StaticClass())
+			|| Property.IsA(FBoolProperty::StaticClass())
+			|| Property.IsA(FEnumProperty::StaticClass());
+	}
+	
+	bool IsNativeStructProperty(const FProperty& Property)
+	{
+		return Property.IsA(FStructProperty::StaticClass())
+			&& CastField<FStructProperty>(&Property)->Struct->GetCppStructOps()->HasSerializer();
+	}
+}
diff --git a/Engine/Plugins/Developer/Concert/ConcertSync/ConcertSyncCore/Source/ConcertSyncCore/Public/ConcertSyncSessionDatabase.h b/Engine/Plugins/Developer/Concert/ConcertSync/ConcertSyncCore/Source/ConcertSyncCore/Public/ConcertSyncSessionDatabase.h
--- a/Engine/Plugins/Developer/Concert/ConcertSync/ConcertSyncCore/Source/ConcertSyncCore/Public/ConcertSyncSessionDatabase.h
+++ b/Engine/Plugins/Developer/Concert/ConcertSync/ConcertSyncCore/Source/ConcertSyncCore/Public/ConcertSyncSessionDatabase.h
@@ -3,6 +3,7 @@
 #pragma once
 
 #include "Delegates/Delegate.h"
+#include "Misc/EBreakBehavior.h"
 #include "Templates/PimplPtr.h"
 
 class FStructOnScope;
@@ -31,12 +32,6 @@ class FConcertSyncSessionDatabaseStatements;
 class FSQLiteDatabase;
 enum class ESQLiteDatabaseOpenMode : uint8;
 
-enum class EBreakBehavior
-{
-	Break,
-	Continue
-};
-
 using FConsumePackageActivityFunc = TFunctionRef<void(FConcertSyncActivity&&/*BasePart*/, FConcertSyncPackageEventData& /*EventPart*/)>;
 using FIteratePackageActivityFunc = TFunctionRef<EBreakBehavior(FConcertSyncActivity&&/*BasePart*/, FConcertSyncPackageEventData& /*EventPart*/)>;
 using FIterateActivityFunc = TFunctionRef<EBreakBehavior(FConcertSyncActivity&&)>;
diff --git a/Engine/Plugins/Developer/Concert/ConcertSync/ConcertSyncCore/Source/ConcertSyncCore/Public/Misc/EBreakBehavior.h b/Engine/Plugins/Developer/Concert/ConcertSync/ConcertSyncCore/Source/ConcertSyncCore/Public/Misc/EBreakBehavior.h
new file mode 100644
--- /dev/null
+++ b/Engine/Plugins/Developer/Concert/ConcertSync/ConcertSyncCore/Source/ConcertSyncCore/Public/Misc/EBreakBehavior.h
@@ -0,0 +1,11 @@
+// Copyright Epic Games, Inc. All Rights Reserved.
+
+#pragma once
+
+#include "HAL/Platform.h"
+
+enum class EBreakBehavior : uint8
+{
+	Break,
+	Continue
+};
\ No newline at end of file
diff --git a/Engine/Plugins/Developer/Concert/ConcertSync/ConcertSyncCore/Source/ConcertSyncCore/Public/Replication/Data/ConcertPropertySelection.h b/Engine/Plugins/Developer/Concert/ConcertSync/ConcertSyncCore/Source/ConcertSyncCore/Public/Replication/Data/ConcertPropertySelection.h
new file mode 100644
--- /dev/null
+++ b/Engine/Plugins/Developer/Concert/ConcertSync/ConcertSyncCore/Source/ConcertSyncCore/Public/Replication/Data/ConcertPropertySelection.h
@@ -0,0 +1,160 @@
+// Copyright Epic Games, Inc. All Rights Reserved.
+
+#pragma once
+
+#include "ConcertPropertySelection.generated.h"
+
+class UStruct;
+struct FArchiveSerializedPropertyChain;
+
+/** Describes the path to a FProperty replicated by Concert.  */
+USTRUCT()
+struct CONCERTSYNCCORE_API FConcertPropertyChain
+{
+	GENERATED_BODY()
+
+	/** The name of inner properties of containers (array, set, map) which are either primitive or native structs. */
+	static const FName InternalContainerPropertyValueName;
+
+	/** Constructs a FConcertPropertyChain from a path if it is valid. If you need to create many paths in one go, use PropertyUtils::BulkConstructConcertChainsFromPaths instead. */
+	static TOptional<FConcertPropertyChain> CreateFromPath(UStruct& Class, const TArray<FName>& NamePath);
+
+	FConcertPropertyChain() = default;
+	/**
+	 * @param OptionalChain The chain leading up to LeafProperty. If it is a root property it can either be empty or nullptr. This mimics the behaviours of FArchive.
+	 * @param LeafProperty The property the chain leads to. It will be the last property in PathToProperty. This can be the inner property of a container but ONLY for primitive types (float, etc.) or structs with a custom Serialize function. 
+	 */
+	FConcertPropertyChain(const FArchiveSerializedPropertyChain* OptionalChain, const FProperty& LeafProperty);
+	
+	/** Gets the leaf property, which is the property the path leads towards. */
+	FName GetLeafProperty() const { return PathToProperty.IsEmpty() ? NAME_None : PathToProperty[PathToProperty.Num() - 1]; }
+	bool IsRootProperty() const { return PathToProperty.Num() == 1; }
+
+	/** @return Whether this is a parent of ChildToCheck */
+	bool IsParentOf(const FConcertPropertyChain& ChildToCheck) const { return ChildToCheck.IsChildOf(*this); }
+	
+	/** @return Whether the leaf property is a child of the given property chain. */
+	bool IsChildOf(const FConcertPropertyChain& ParentToCheck) const;
+	/** @return Whether the leaf property is a direct child of the given property chain. */
+	bool IsDirectChildOf(const FConcertPropertyChain& ParentToCheck) const;
+
+	const TArray<FName>& GetPathToProperty() const { return PathToProperty; }
+
+	enum class EToStringMethod
+	{
+		Path,
+		LeafProperty
+	};
+	FString ToString(EToStringMethod Method = EToStringMethod::Path) const;
+	
+	friend bool operator==(const FConcertPropertyChain& Left, const FConcertPropertyChain& Right)
+	{
+		return Left.PathToProperty == Right.PathToProperty;
+	}
+	friend bool operator!=(const FConcertPropertyChain& Left, const FConcertPropertyChain& Right)
+	{
+		return !(Left == Right);
+	}
+
+	friend bool operator==(const FConcertPropertyChain& Left, const TArray<FName>& Path)
+	{
+		return Left.PathToProperty == Path;
+	}
+	friend bool operator!=(const FConcertPropertyChain& Left, const TArray<FName>& Path)
+	{
+		return Left.PathToProperty != Path;
+	}
+	
+	friend bool operator==(const TArray<FName>& Path, const FConcertPropertyChain& Left)
+	{
+		return Left.PathToProperty == Path;
+	}
+	friend bool operator!=(const TArray<FName>& Path, const FConcertPropertyChain& Left)
+	{
+		return Left.PathToProperty != Path;
+	}
+	
+private:
+	
+	/**
+	 * Path from root of UObject to leaf property. Includes the leaf property.
+	 * This property is kept private to force the use of the exposed constructors.
+	 *
+	 * Suppose:
+	 * class AFooActor
+	 * {
+	 *		UPROPERTY()
+	 *		FFooStruct Struct;
+	 * 
+	 *		UPROPERTY()
+	 *		TArray<FFooStruct> ArrayOfStructs;
+	 *
+	 *		UPROPERTY()
+	 *		TMap<int32, FFooStruct> MapOfStructs;
+	 *
+	 *		UPROPERTY()
+	 *		TMap<int32, float> MapIntToFloat;
+	 *
+	 *		UPROPERTY(Instanced)
+	 *		TArray<UFooSubobject*> UnsupportedInstanced;
+	 * };
+	 *
+	 * struct FFooStruct
+	 * {
+	 *		UPROPERTY()
+	 *		int32 Foo;
+	 *
+	 *		UPROPERTY()
+	 *		int32 Bar;
+	 * };
+	 *
+	 * Listing some paths:
+	 *  - { "Struct", "Foo" }
+	 *  - { "ArrayOfStructs", "Foo" }
+	 *  - { "MapOfStructs" } > Replicate all keys but not the struct values
+	 *  - { "MapOfStructs", "Foo" } > Replicate all keys and the Foo struct value.
+	 *	- { "MapIntToFloat" } > Replicate all keys but not the values
+	 *	- { "MapIntToFloat", "Value" } > Replicate all the keys and the values. Special non-ustruct exception where the inner property is listed! 
+	 *  
+	 * The inner properties of containers, e.g. FArrayProperty::Inner, FMapProperty::KeyProp, etc., generally do not appear in the path
+	 * but there is one special case for TMap where it does.
+	 * More details:
+	 *	- TArray, TSet: The inner properties are "obvious" so we can shorten the path.
+	 *	- TMap:
+	 *	  - Key: The key is required. We assume reasonably that all properties of the key property must be replicated. Hence, there is no point in listing
+	 *	  the key inner property nor any of its subproperties. Example: If key is FSoftObjectPath, we do not list FSoftObjectPath::AssetPath, etc.
+	 *	  - Value:
+	 *	    - If the inner value property is is ustruct, then it is not listed. The child properties of the inner property are listed.
+	 *	    - If the inner value property is not a ustruct (meaning a primitive float, int, etc.), then the child property is listed and called "Value".
+	 *	      This is needed to differentiate between the cases "serializes only the keys" and "serializes the keys and the values".
+	 *	  - Final word about FConcertPropertySelection:
+	 *		- { "MapOfStructs", "Foo" } means all keys and the Foo property is serialized (nothing is said about Bar - it is replicated if it is also in the property selection).
+	 *		- { "MapOfStructs" } means only the keys are serialized and the values are skipped (Foo and Bar will be value-initialized)
+	 *
+	 * Paths do not include the child properties of uStructs implementing a custom Serialize function with struct ops. In that case, the user can only specify
+	 * include the struct or skip it.
+	 *
+	 * FConcertPropertyChains do NOT cross the UObject border. For example, there would be no thing as [1] "UnsupportedInstanced" [x] ...
+	 * You'd start a new FConcertPropertyChain for properties for objects of type UFooSubobject.
+	 */
+	UPROPERTY()
+	TArray<FName> PathToProperty;
+};
+
+/** List of properties to be replicated for a given object */
+USTRUCT()
+struct CONCERTSYNCCORE_API FConcertPropertySelection
+{
+	GENERATED_BODY()
+
+	/**
+	 * List of replicated objects.
+	 *
+	 * Can be actors or its subobjects.
+	 * Technically, this can also be non-UWorld objects.
+	 */
+	UPROPERTY()
+	TArray<FConcertPropertyChain> ReplicatedProperties;
+};
+
+CONCERTSYNCCORE_API uint32 GetTypeHash(const FConcertPropertyChain& Chain);
\ No newline at end of file
diff --git a/Engine/Plugins/Developer/Concert/ConcertSync/ConcertSyncCore/Source/ConcertSyncCore/Public/Replication/Data/ObjectReplicationMap.h b/Engine/Plugins/Developer/Concert/ConcertSync/ConcertSyncCore/Source/ConcertSyncCore/Public/Replication/Data/ObjectReplicationMap.h
new file mode 100644
--- /dev/null
+++ b/Engine/Plugins/Developer/Concert/ConcertSync/ConcertSyncCore/Source/ConcertSyncCore/Public/Replication/Data/ObjectReplicationMap.h
@@ -0,0 +1,34 @@
+// Copyright Epic Games, Inc. All Rights Reserved.
+
+#pragma once
+
+#include "Replication/Data/ConcertPropertySelection.h"
+#include "ObjectReplicationMap.generated.h"
+
+USTRUCT()
+struct FReplicatedObjectInfo
+{
+	GENERATED_BODY()
+	
+	UPROPERTY()
+	FSoftClassPath ClassPath;
+	
+	UPROPERTY()
+	FConcertPropertySelection ReplicatedProperties;
+};
+
+/** Maps objects to their replicated properties. */
+USTRUCT()
+struct FObjectReplicationMap
+{
+	GENERATED_BODY()
+
+	/**
+	 * List of replicated objects.
+	 *
+	 * Can be actors or its subobjects.
+	 * Technically, this can also be non-UWorld objects.
+	 */
+	UPROPERTY()
+	TMap<FSoftObjectPath, FReplicatedObjectInfo> ReplicatedObjects;
+};
\ No newline at end of file
diff --git a/Engine/Plugins/Developer/Concert/ConcertSync/ConcertSyncCore/Source/ConcertSyncCore/Public/Replication/Data/ReplicationClientDescription.h b/Engine/Plugins/Developer/Concert/ConcertSync/ConcertSyncCore/Source/ConcertSyncCore/Public/Replication/Data/ReplicationClientDescription.h
new file mode 100644
--- /dev/null
+++ b/Engine/Plugins/Developer/Concert/ConcertSync/ConcertSyncCore/Source/ConcertSyncCore/Public/Replication/Data/ReplicationClientDescription.h
@@ -0,0 +1,14 @@
+// Copyright Epic Games, Inc. All Rights Reserved.
+
+#pragma once
+
+#include "ReplicationClientDescription.generated.h"
+
+/** Describes a client to the server, e.g. the kind of data it wants to receive. */
+USTRUCT()
+struct FReplicationClientDescription
+{
+	GENERATED_BODY()
+
+	// TODO: Add rules for receiving
+};
\ No newline at end of file
diff --git a/Engine/Plugins/Developer/Concert/ConcertSync/ConcertSyncCore/Source/ConcertSyncCore/Public/Replication/Data/ReplicationStreamDescription.h b/Engine/Plugins/Developer/Concert/ConcertSync/ConcertSyncCore/Source/ConcertSyncCore/Public/Replication/Data/ReplicationStreamDescription.h
new file mode 100644
--- /dev/null
+++ b/Engine/Plugins/Developer/Concert/ConcertSync/ConcertSyncCore/Source/ConcertSyncCore/Public/Replication/Data/ReplicationStreamDescription.h
@@ -0,0 +1,61 @@
+// Copyright Epic Games, Inc. All Rights Reserved.
+
+#pragma once
+
+#include "ObjectReplicationMap.h"
+#include "ReplicationStreamDescription.generated.h"
+
+template<typename OptionalType>
+struct TOptional;
+
+/** Shared non-UObject data, which can be serialized by Concert without wrapping (UObjects need wrapping). */
+USTRUCT()
+struct FSharedReplicationStreamDescription
+{
+	GENERATED_BODY()
+	
+	/** Unique id for this stream. */
+	UPROPERTY()
+	FGuid Identifier = FGuid::NewGuid();
+
+	/** Identifies the data that this stream will send. */
+	UPROPERTY()
+	FObjectReplicationMap ReplicationMap;
+};
+
+/**
+ * Contains a description of a replication stream.
+ * A replication stream is a map of objects to their replicated properties, and a set of stream attributes (for rules).
+ */
+USTRUCT()
+struct FReplicationStreamDescription
+{
+	GENERATED_BODY()
+
+	UPROPERTY()
+	FSharedReplicationStreamDescription BaseDescription;
+
+	// TODO: Add TArray<TObjectPtr<UReplicationStreamAttribute>> property
+
+	/** Packs the UObjects in this description (Concert does not allow sending UObjects directly) */
+	CONCERTSYNCCORE_API struct FReplicationStreamDescription_NetPacked Pack() const;
+};
+
+/**
+ * This data is packed such that it can be sent via concert messages; packed here means that instanced UObjects are wrapped
+ * with sufficient data to reconstruct them.
+ */
+USTRUCT()
+struct FReplicationStreamDescription_NetPacked
+{
+	GENERATED_BODY()
+
+	
+	UPROPERTY()
+	FSharedReplicationStreamDescription BaseDescription;
+
+	// TODO: Add TArray<FConcertSessionSerializedPayload> where each corresponds to a UReplicationStreamAttribute; those will be used by receive rules to decide whether a client can receive certain data.
+
+	/** Unpacks this data. This function can fail, e.g if the UObject class does not exist locally. */
+	CONCERTSYNCCORE_API TOptional<FReplicationStreamDescription> Unpack(FString* OutErrorMessage = nullptr) const;
+};
\ No newline at end of file
diff --git a/Engine/Plugins/Developer/Concert/ConcertSync/ConcertSyncCore/Source/ConcertSyncCore/Public/Replication/PropertyChainUtils.h b/Engine/Plugins/Developer/Concert/ConcertSync/ConcertSyncCore/Source/ConcertSyncCore/Public/Replication/PropertyChainUtils.h
new file mode 100644
--- /dev/null
+++ b/Engine/Plugins/Developer/Concert/ConcertSync/ConcertSyncCore/Source/ConcertSyncCore/Public/Replication/PropertyChainUtils.h
@@ -0,0 +1,81 @@
+// Copyright Epic Games, Inc. All Rights Reserved.
+
+#pragma once
+
+#include "Data/ConcertPropertySelection.h"
+#include "Misc/EBreakBehavior.h"
+#include "Templates/Function.h"
+
+class FProperty;
+class UStruct;
+struct FArchiveSerializedPropertyChain;
+
+namespace UE::ConcertSyncCore::PropertyChain
+{
+	/**
+	 * Iterates all properties that are supported for replication.
+	 * 
+	 * This iterates every property as described by FConcertPropertyChain::PathToProperty, which means inner properties
+	 * of containers, like TArray, TSet, and TMap, are skipped.
+	 *
+	 * @param Class The class whose properties to look at
+	 * @param ProcessProperty Callback
+	 * @param OnlyChildrenOf Instead of iterating all properties, you can restrict to only receive the child properties of OnlyChildrenOf.
+	 */
+	CONCERTSYNCCORE_API void ForEachReplicatableProperty(
+		UStruct& Class,
+		TFunctionRef<EBreakBehavior(const FArchiveSerializedPropertyChain& Chain, const FProperty& LeafProperty)> ProcessProperty
+		);
+
+	/**
+	 * Iterates all properties that are supported for replication and converts them to FConcertPropertyChain.
+	 *
+	 * This iterates every property as described by FConcertPropertyChain::PathToProperty, which means inner properties
+	 * of containers, like TArray, TSet, and TMap, are skipped.
+	 *
+	 * @param Class The class whose properties to look at
+	 * @param ProcessProperty Callback
+	 * @param OnlyChildrenOf Instead of iterating all properties, you can restrict to only receive the child properties of OnlyChildrenOf.
+	 */
+	CONCERTSYNCCORE_API void ForEachReplicatableConcertProperty(
+		UStruct& Class,
+		TFunctionRef<EBreakBehavior(FConcertPropertyChain&& PropertyChain)> ProcessProperty
+		);
+
+	
+	/**
+	 * Efficiently create a batch of paths from the same class hierarchy.
+	 *
+	 * @param Class The class for which the paths are created
+	 * @param NumPaths The number of paths we're looking to create. Once MatchesPath has returned true NumPaths times, this function terminates.
+	 * @param MatchesPath Determines whether you are interested in this path. If so, do something with it and return true.
+	 */
+	CONCERTSYNCCORE_API void BulkConstructConcertChainsFromPaths(
+		UStruct& Class,
+		uint32 NumPaths,
+		TFunctionRef<bool(const FArchiveSerializedPropertyChain& Chain, const FProperty& LeafProperty)> MatchesPath
+		);
+
+	/** Util that can be used with BulkConstructConcertChainsFromPaths. Checks whether Path coincides with Chain & LeafProperty.*/
+	CONCERTSYNCCORE_API bool DoPathAndChainsMatch(const TArray<FName>& Path, const FArchiveSerializedPropertyChain& Chain, const FProperty& LeafProperty);
+
+
+	/**
+	 * Decides whether the property is allowed for replication.
+	 * @param LeafProperty The property being considered.
+	 */
+	CONCERTSYNCCORE_API bool IsReplicatableProperty(const FProperty& LeafProperty);
+
+	/** Util for determining whether this property is the inner property of a container. */
+	CONCERTSYNCCORE_API bool IsInnerContainerProperty(const FProperty& Property);
+	/**
+	 * Util for determining whether this property is considered primitive (numeric, bool, enum).
+	 * Primitive properties in containers are indicated as FConcertPropertyChain::InternalContainerPropertyValueName, which is "Value".
+	 */
+	CONCERTSYNCCORE_API bool IsPrimitiveProperty(const FProperty& Property);
+	/**
+	 * Util for determining whether this struct is considered native.
+	 * Native struct properties in containers are indicated as FConcertPropertyChain::InternalContainerPropertyValueName, which is "Value".
+	 */
+	CONCERTSYNCCORE_API bool IsNativeStructProperty(const FProperty& Property);
+}
diff --git a/Engine/Plugins/Developer/Concert/ConcertSync/ConcertSyncTest/Source/ConcertSyncTest/Private/Replication/ConcertPropertyTests.cpp b/Engine/Plugins/Developer/Concert/ConcertSync/ConcertSyncTest/Source/ConcertSyncTest/Private/Replication/ConcertPropertyTests.cpp
new file mode 100644
--- /dev/null
+++ b/Engine/Plugins/Developer/Concert/ConcertSync/ConcertSyncTest/Source/ConcertSyncTest/Private/Replication/ConcertPropertyTests.cpp
@@ -0,0 +1,190 @@
+// Copyright Epic Games, Inc. All Rights Reserved.
+
+#include "Replication/Data/ConcertPropertySelection.h"
+#include "Replication/PropertyChainUtils.h"
+#include "Misc/AutomationTest.h"
+
+#include "TestReflectionObject.h"
+#include "Serialization/ArchiveSerializedPropertyChain.h"
+
+/**
+ * Tests that the FConcertPropertyChain constructor works as the documentation dictates:
+ *  - Root properties
+ *  - Array, set, map of
+ *		- "normal" structs, like FVector > lists sub-properties
+ *		- primitives and native structs > Path contains FConcertPropertyChain::InternalContainerPropertyValueName ("Value") at the end
+ */
+IMPLEMENT_SIMPLE_AUTOMATION_TEST(FConcertPropertyChainTests, "Concert.Replication.PropertyChainTests", EAutomationTestFlags::EditorContext | EAutomationTestFlags::EngineFilter)
+bool FConcertPropertyChainTests::RunTest(const FString& Parameters)
+{
+	// The test is intentionally set up to make it it easy to set breakpoints albeit at the expense of increasing the amount of code
+
+	// 1. Input
+	UStruct* TestReplicationStruct = FTestReplicationStruct::StaticStruct();
+	FProperty* ValueProp = TestReplicationStruct->FindPropertyByName(GET_MEMBER_NAME_CHECKED(FTestReplicationStruct, Value));
+	FProperty* VectorProp = TestReplicationStruct->FindPropertyByName(GET_MEMBER_NAME_CHECKED(FTestReplicationStruct, Vector));
+	FProperty* VectorXProp = CastField<FStructProperty>(VectorProp)->Struct->FindPropertyByName(GET_MEMBER_NAME_CHECKED(FVector, X));
+	
+	FProperty* FloatArrayProp = TestReplicationStruct->FindPropertyByName(GET_MEMBER_NAME_CHECKED(FTestReplicationStruct, FloatArray));
+	FProperty* FloatArrayInnerProp = CastField<FArrayProperty>(FloatArrayProp)->Inner;
+	FProperty* StringArrayProp = TestReplicationStruct->FindPropertyByName(GET_MEMBER_NAME_CHECKED(FTestReplicationStruct, StringArray));
+	FProperty* NativeStructArrayProp = TestReplicationStruct->FindPropertyByName(GET_MEMBER_NAME_CHECKED(FTestReplicationStruct, NativeStructArray));
+	FProperty* NativeStructArrayInnerProp = CastField<FArrayProperty>(NativeStructArrayProp)->Inner;
+	
+	FProperty* FloatSetProp = TestReplicationStruct->FindPropertyByName(GET_MEMBER_NAME_CHECKED(FTestReplicationStruct, FloatSet));
+	FProperty* FloatSetInnerProp = CastField<FSetProperty>(FloatSetProp)->ElementProp;
+	FProperty* StringSetProp = TestReplicationStruct->FindPropertyByName(GET_MEMBER_NAME_CHECKED(FTestReplicationStruct, StringSet));
+	FProperty* NativeStructSetProp = TestReplicationStruct->FindPropertyByName(GET_MEMBER_NAME_CHECKED(FTestReplicationStruct, NativeStructSet));
+	FProperty* NativeStructSetInnerProp = CastField<FSetProperty>(NativeStructSetProp)->ElementProp;
+	
+	FProperty* StringToFloatProp = TestReplicationStruct->FindPropertyByName(GET_MEMBER_NAME_CHECKED(FTestReplicationStruct, StringToFloat));
+	FProperty* StringToFloatInnerProp = CastField<FMapProperty>(StringToFloatProp)->ValueProp;
+	FProperty* StringToVectorProp = TestReplicationStruct->FindPropertyByName(GET_MEMBER_NAME_CHECKED(FTestReplicationStruct, StringToVector));
+	FProperty* StringToNativeStructProp = TestReplicationStruct->FindPropertyByName(GET_MEMBER_NAME_CHECKED(FTestReplicationStruct, StringToNativeStruct));
+	FProperty* StringToNativeStructInnerProp = CastField<FMapProperty>(StringToNativeStructProp)->ValueProp;
+
+	FArchiveSerializedPropertyChain VectorXChain;
+	FArchiveSerializedPropertyChain FloatArrayChain;
+	FArchiveSerializedPropertyChain FloatSetArrayChain;
+	FArchiveSerializedPropertyChain StringToFloatChain;
+	FArchiveSerializedPropertyChain NativeStructArrayChain;
+	FArchiveSerializedPropertyChain NativeStructSetChain;
+	FArchiveSerializedPropertyChain StringToNativeStructChain;
+	VectorXChain.PushProperty(VectorProp, false);
+	FloatArrayChain.PushProperty(FloatArrayProp, false);
+	FloatSetArrayChain.PushProperty(FloatSetProp, false);
+	StringToFloatChain.PushProperty(StringToFloatProp, false);
+	NativeStructArrayChain.PushProperty(NativeStructArrayProp, false);
+	NativeStructSetChain.PushProperty(NativeStructSetProp, false);
+	StringToNativeStructChain.PushProperty(StringToNativeStructProp, false);
+	
+
+	
+	// 2. Run
+	const FConcertPropertyChain Value(nullptr, *ValueProp);
+	const FConcertPropertyChain Vector(nullptr, *VectorProp);
+	const FConcertPropertyChain VectorX(&VectorXChain, *VectorXProp);
+	
+	const FConcertPropertyChain FloatArray(nullptr, *FloatArrayProp);
+	const FConcertPropertyChain FloatArrayInner(&FloatArrayChain, *FloatArrayInnerProp);
+	const FConcertPropertyChain StringArray(nullptr, *StringArrayProp);
+	const FConcertPropertyChain NativeStructArray(nullptr, *NativeStructArrayProp);
+	const FConcertPropertyChain NativeStructArrayInner(&NativeStructArrayChain, *NativeStructArrayInnerProp);
+	
+	const FConcertPropertyChain FloatSet(nullptr, *FloatSetProp);
+	const FConcertPropertyChain FloatSetInner(&FloatSetArrayChain, *FloatSetInnerProp);
+	const FConcertPropertyChain StringSet(nullptr, *StringSetProp);
+	const FConcertPropertyChain NativeStructSet(nullptr, *NativeStructSetProp);
+	const FConcertPropertyChain NativeStructSetInner(&NativeStructSetChain, *NativeStructSetInnerProp);
+	
+	const FConcertPropertyChain StringToFloat(nullptr, *StringToFloatProp);
+	const FConcertPropertyChain StringToFloatInner(&StringToFloatChain, *StringToFloatInnerProp);
+	const FConcertPropertyChain StringToVector(nullptr, *StringToVectorProp);
+	const FConcertPropertyChain StringToNativeStruct(nullptr, *StringToNativeStructProp);
+	const FConcertPropertyChain StringToNativeStructInner(&StringToNativeStructChain, *StringToNativeStructInnerProp);
+
+
+	
+	// 3. Test
+	TestTrue(TEXT("Value"), Value == TArray<FName>{ TEXT("Value") });
+	TestTrue(TEXT("Vector"), Vector == TArray<FName>{ TEXT("Vector") });
+	TestTrue(TEXT("Vector.X"), VectorX == TArray<FName>{ FName(TEXT("Vector")), FName(TEXT("X")) });
+	
+	TestTrue(TEXT("FloatArray"), FloatArray == TArray<FName>{ TEXT("FloatArray") });
+	TestTrue(TEXT("FloatArrayInner"), FloatArrayInner == TArray<FName>{ FName(TEXT("FloatArray")), FConcertPropertyChain::InternalContainerPropertyValueName });
+	TestTrue(TEXT("StringArray"), StringArray == TArray<FName>{ TEXT("StringArray") });
+	TestTrue(TEXT("NativeStructArray"), NativeStructArray == TArray<FName>{ TEXT("NativeStructArray") });
+	TestTrue(TEXT("NativeStructArrayInner"), NativeStructArrayInner == TArray<FName>{ FName(TEXT("NativeStructArray")), FConcertPropertyChain::InternalContainerPropertyValueName });
+	
+	TestTrue(TEXT("FloatSet"), FloatSet == TArray<FName>{ TEXT("FloatSet") });
+	TestTrue(TEXT("FloatSetInner"), FloatSetInner == TArray<FName>{ FName(TEXT("FloatSet")), FConcertPropertyChain::InternalContainerPropertyValueName });
+	TestTrue(TEXT("StringSet"), StringSet == TArray<FName>{ TEXT("StringSet") });
+	TestTrue(TEXT("NativeStructSet"), NativeStructSet == TArray<FName>{ TEXT("NativeStructSet") });
+	TestTrue(TEXT("NativeStructSetInner"), NativeStructSetInner == TArray<FName>{ FName(TEXT("NativeStructSet")), FConcertPropertyChain::InternalContainerPropertyValueName });
+	
+	TestTrue(TEXT("StringToFloat"), StringToFloat == TArray<FName>{ TEXT("StringToFloat") });
+	TestTrue(TEXT("StringToFloatInner"), StringToFloatInner == TArray<FName>{ FName(TEXT("StringToFloat")), FConcertPropertyChain::InternalContainerPropertyValueName });
+	TestTrue(TEXT("StringToVector"), StringToVector == TArray<FName>{ TEXT("StringToVector") });
+	TestTrue(TEXT("StringToNativeStruct"), StringToNativeStruct == TArray<FName>{ TEXT("StringToNativeStruct") });
+	TestTrue(TEXT("StringToNativeStructInner"), StringToNativeStructInner == TArray<FName>{ FName(TEXT("StringToNativeStruct")), FConcertPropertyChain::InternalContainerPropertyValueName });
+	
+	return true;
+}
+
+/**
+ * Tests that all property cases on the specially constructed example class UTestReflectionObject.
+ */
+IMPLEMENT_SIMPLE_AUTOMATION_TEST(FConcertForEachReplictableConcertPropertyTests, "Concert.Replication.ForEachReplicatableConcertProperty", EAutomationTestFlags::EditorContext | EAutomationTestFlags::EngineFilter)
+bool FConcertForEachReplictableConcertPropertyTests::RunTest(const FString& Parameters)
+{
+	UClass* Class = UTestReflectionObject::StaticClass();
+	TArray<FConcertPropertyChain> Properties;
+	UE::ConcertSyncCore::PropertyChain::ForEachReplicatableConcertProperty(*Class, [&Properties](FConcertPropertyChain&& Chain)
+	{
+		Properties.Add(MoveTemp(Chain));
+		return EBreakBehavior::Continue;
+	});
+
+	using FExpectedPropertyPath = TArray<FName>;
+	auto ToString = [](const FExpectedPropertyPath& Path)
+	{
+		FString Result;
+		for (int32 i = 0; i < Path.Num(); ++i)
+		{
+			if (i == 0)
+			{
+				Result += Path[i].ToString();
+			}
+			else
+			{
+				Result += TEXT(".") + Path[i].ToString();
+			}
+		}
+		return Result;
+	};
+	const TArray<FExpectedPropertyPath> ExpectedProperties {
+		{ TEXT("Float") },
+		{ TEXT("Vector") },
+		{ TEXT("Vector"), TEXT("X") },
+		{ TEXT("Vector"), TEXT("Y") },
+		{ TEXT("Vector"), TEXT("Z") },
+		{ TEXT("TestStruct") },
+		{ TEXT("TestStruct"), TEXT("Nested") },
+		{ TEXT("TestStruct"), TEXT("Nested"), TEXT("Value") },
+		{ TEXT("TestStruct"), TEXT("Nested"), TEXT("Vector") },
+		{ TEXT("TestStruct"), TEXT("Nested"), TEXT("Vector"), TEXT("X")  },
+		{ TEXT("TestStruct"), TEXT("Nested"), TEXT("Vector"), TEXT("Y")  },
+		{ TEXT("TestStruct"), TEXT("Nested"), TEXT("Vector"), TEXT("Z")  },
+		{ TEXT("TestStruct"), TEXT("Nested"), TEXT("NativeStruct") },
+		{ TEXT("TestStruct"), TEXT("Nested"), TEXT("FloatArray") },
+		{ TEXT("TestStruct"), TEXT("Nested"), TEXT("FloatArray"), FConcertPropertyChain::InternalContainerPropertyValueName },
+		{ TEXT("TestStruct"), TEXT("Nested"), TEXT("StringArray") },
+		{ TEXT("TestStruct"), TEXT("Nested"), TEXT("NativeStructArray") },
+		{ TEXT("TestStruct"), TEXT("Nested"), TEXT("NativeStructArray"), FConcertPropertyChain::InternalContainerPropertyValueName },
+		{ TEXT("TestStruct"), TEXT("Nested"), TEXT("FloatSet") },
+		{ TEXT("TestStruct"), TEXT("Nested"), TEXT("FloatSet"), FConcertPropertyChain::InternalContainerPropertyValueName },
+		{ TEXT("TestStruct"), TEXT("Nested"), TEXT("StringSet") },
+		{ TEXT("TestStruct"), TEXT("Nested"), TEXT("NativeStructSet") },
+		{ TEXT("TestStruct"), TEXT("Nested"), TEXT("NativeStructSet"), FConcertPropertyChain::InternalContainerPropertyValueName },
+		{ TEXT("TestStruct"), TEXT("Nested"), TEXT("StringToFloat") },
+		{ TEXT("TestStruct"), TEXT("Nested"), TEXT("StringToFloat") },
+		{ TEXT("TestStruct"), TEXT("Nested"), TEXT("StringToVector") },
+		{ TEXT("TestStruct"), TEXT("Nested"), TEXT("StringToVector"), TEXT("X") },
+		{ TEXT("TestStruct"), TEXT("Nested"), TEXT("StringToVector"), TEXT("Y") },
+		{ TEXT("TestStruct"), TEXT("Nested"), TEXT("StringToVector"), TEXT("Z") },
+		{ TEXT("TestStruct"), TEXT("Nested"), TEXT("StringToNativeStruct") },
+		{ TEXT("TestStruct"), TEXT("Nested"), TEXT("StringToNativeStruct"), FConcertPropertyChain::InternalContainerPropertyValueName },
+	};
+
+	for (const FExpectedPropertyPath& Expected : ExpectedProperties)
+	{
+		// This structure is easier for setting breakpoints than using TestTrue
+		if (!Properties.Contains(Expected))
+		{
+			AddError(ToString(Expected));
+		}
+	}
+	TestEqual(TEXT("Does not contain disallowed properties"), Properties.Num(), ExpectedProperties.Num());
+	
+	return true;
+}
\ No newline at end of file
diff --git a/Engine/Plugins/Developer/Concert/ConcertSync/ConcertSyncTest/Source/ConcertSyncTest/Private/Replication/TestReflectionObject.h b/Engine/Plugins/Developer/Concert/ConcertSync/ConcertSyncTest/Source/ConcertSyncTest/Private/Replication/TestReflectionObject.h
new file mode 100644
--- /dev/null
+++ b/Engine/Plugins/Developer/Concert/ConcertSync/ConcertSyncTest/Source/ConcertSyncTest/Private/Replication/TestReflectionObject.h
@@ -0,0 +1,111 @@
+// Fill out your copyright notice in the Description page of Project Settings.
+
+#pragma once
+
+#include "GameFramework/Actor.h"
+#include "UObject/Class.h"
+#include "TestReflectionObject.generated.h"
+
+USTRUCT()
+struct FNativeStruct
+{
+	GENERATED_BODY()
+
+	UPROPERTY()
+	float Float{};
+	
+	bool Serialize(FArchive& Ar)
+	{
+		Ar << Float;
+		return true;
+	}
+};
+
+inline uint32 GetTypeHash(const FNativeStruct& Struct){ return GetTypeHash(Struct.Float); }
+
+template<>
+struct TStructOpsTypeTraits<FNativeStruct> : public TStructOpsTypeTraitsBase2<FNativeStruct>
+{
+	enum 
+	{
+		WithSerializer = true
+	};
+};
+
+USTRUCT()
+struct FTestReplicationStruct
+{
+	GENERATED_BODY()
+
+	UPROPERTY()
+	float Value{};
+	UPROPERTY()
+	FVector Vector;
+	UPROPERTY()
+	FNativeStruct NativeStruct;
+
+	UPROPERTY()
+	TArray<float> FloatArray;
+	UPROPERTY()
+	TArray<FString> StringArray;
+	UPROPERTY()
+	TArray<FNativeStruct> NativeStructArray;
+
+	UPROPERTY()
+	TSet<float> FloatSet;
+	UPROPERTY()
+	TSet<FString> StringSet;
+	UPROPERTY()
+	TSet<FNativeStruct> NativeStructSet;
+	
+	UPROPERTY()
+	TMap<FString, float> StringToFloat;
+	UPROPERTY()
+	TMap<FString, FVector> StringToVector;
+	UPROPERTY()
+	TMap<FString, FNativeStruct> StringToNativeStruct;
+};
+
+USTRUCT()
+struct FTestNestedReplicationStruct
+{
+	GENERATED_BODY()
+
+	UPROPERTY()
+	FTestReplicationStruct Nested;
+
+	UPROPERTY()
+	TArray<FTestReplicationStruct> NestedArray;
+};
+
+DECLARE_DYNAMIC_MULTICAST_DELEGATE(FTestReflectionDelegate);
+
+UCLASS()
+class UTestReflectionObject : public UObject
+{
+	GENERATED_BODY()
+public:
+	
+	UPROPERTY(EditAnywhere, Category = "Default")
+	float Float{};
+	
+	UPROPERTY()
+	FVector Vector;
+
+	UPROPERTY()
+	FTestNestedReplicationStruct TestStruct;
+
+	// These properties are invalid to replicate
+
+	UPROPERTY(BlueprintAssignable, Category = "Dummy")
+	FTestReflectionDelegate DelegateTest;
+
+	UPROPERTY(Instanced)
+	TObjectPtr<UObject> InstancedSubobject;
+	UPROPERTY(Instanced)
+	TArray<TObjectPtr<UObject>> InstancedSubobjectArray;
+	UPROPERTY(Instanced)
+	TSet<TObjectPtr<UObject>> InstancedSubobjectSet;
+	UPROPERTY(Instanced)
+	TMap<FName, TObjectPtr<UObject>> InstancedSubobjectMap;
+};
\ No newline at end of file
diff --git a/Engine/Plugins/Developer/Concert/ConcertUI/ConcertSharedSlate/Source/ConcertSharedSlate/Public/Model/Item/IItemSourceModel.h b/Engine/Plugins/Developer/Concert/ConcertUI/ConcertSharedSlate/Source/ConcertSharedSlate/Public/Model/Item/IItemSourceModel.h
new file mode 100644
--- /dev/null
+++ b/Engine/Plugins/Developer/Concert/ConcertUI/ConcertSharedSlate/Source/ConcertSharedSlate/Public/Model/Item/IItemSourceModel.h
@@ -0,0 +1,81 @@
+// Copyright Epic Games, Inc. All Rights Reserved.
+
+#pragma once
+
+#include "Containers/Array.h"
+#include "Templates/Function.h"
+#include "Templates/SharedPointer.h"
+#include "Textures/SlateIcon.h"
+#include "Misc/EBreakBehavior.h"
+
+namespace UE::ConcertSharedSlate
+{
+	enum class ESourceType : uint8
+	{
+		/** Display EnumerateSelectableItems in a drop-down or menu list. Example: User wants to search from a list of actors. */
+		ShowAsList,
+		/** When this option is clicked, instantly add all of EnumerateSelectableItems. Example: User wants to add all actors selected in the world outliner. */
+		AddOnClick
+	};
+
+	struct FBaseDisplayInfo
+	{
+		/** What the model should be displayed as in UI. */
+		FText Label;
+		/** Optional tooltip when the model is hovered. */
+		FText ToolTip;
+
+		/** Optional icon to display entry as. */
+		FSlateIcon Icon;
+	};
+	
+	struct FSourceDisplayInfo : FBaseDisplayInfo
+	{
+		/** Determines what happens when the option is clicked / hovered. */
+		ESourceType SourceType;
+	};
+	
+	/**
+	 * Defines the source for of item, e.g. objects.
+	 *
+	 * Imagine you have a list view displaying a list of objects. Over it is a "Add Object" combo button.
+	 * In the button's menu, you'd find e.g. "Add from World", "Add from Selection", "Import from File", etc.
+	 * Each of of these sub-menu entries is an item source and corresponds to an IItemSourceModel.
+	 * @see IObjectSelectionSourceModel and FEditorObjectSelectionSourceModel for an example how sources may be used.
+	 *
+	 * Models may also be useful for building context menus entries. When you right-click an item, you may want to add
+	 * more items based on the object clicked, e.g. add the components of a right-clicked actor.
+	 * @see FEditorObjectSelectionSourceModel::FContextMenuSources and FEditorObjectSelectionSourceModel::GetContextMenuOptions
+	 *
+	 * @see FSourceModelBuilders for building UI from models, e.g. the above mentioned combo button or context menus.
+	 */
+	template<typename TItemType>
+	class IItemSourceModel : public TSharedFromThis<IItemSourceModel<TItemType>>
+	{
+	public:
+
+		/** Some data for displaying this source in UI. */
+		virtual FSourceDisplayInfo GetDisplayInfo() const = 0;
+
+		/** @return The number of selectable objects. */
+		virtual uint32 GetNumSelectableItems() const = 0;
+		
+		/** Enumerates ALL objects from this source. This can be a LOT of objects (e.g. 2000 actors). Intended to be placed in a searchable drop-down menu. */
+		virtual void EnumerateSelectableItems(TFunctionRef<EBreakBehavior(const TItemType& SelectableOption)> Delegate) const = 0;
+
+		/** Util that converts EnumerateSelectableItems into an array. */
+		TArray<TItemType> GetSelectableItems()
+		{
+			TArray<TItemType> Result;
+			Result.Reserve(GetNumSelectableItems());
+			EnumerateSelectableItems([&Result](const TItemType& SelectableOption)
+			{
+				Result.Add(SelectableOption);
+				return EBreakBehavior::Continue;
+			});
+			return Result;
+		}
+
+		virtual ~IItemSourceModel() = default;
+	};
+}
diff --git a/Engine/Plugins/Developer/Concert/ConcertUI/ConcertSharedSlate/Source/ConcertSharedSlate/Public/Model/Item/SourceModelBuilders.h b/Engine/Plugins/Developer/Concert/ConcertUI/ConcertSharedSlate/Source/ConcertSharedSlate/Public/Model/Item/SourceModelBuilders.h
new file mode 100644
--- /dev/null
+++ b/Engine/Plugins/Developer/Concert/ConcertUI/ConcertSharedSlate/Source/ConcertSharedSlate/Public/Model/Item/SourceModelBuilders.h
@@ -0,0 +1,275 @@
+// Copyright Epic Games, Inc. All Rights Reserved.
+
+#pragma once
+
+#include "Model/Item/IItemSourceModel.h"
+#include "Model/Item/SourceSelectionCategory.h"
+
+#include "Algo/Sort.h"
+#include "Framework/Commands/UIAction.h"
+#include "Framework/MultiBox/MultiBoxBuilder.h"
+#include "Internationalization/Text.h"
+#include "SPositiveActionButton.h"
+#include "Styling/AppStyle.h"
+#include "Templates/SharedPointer.h"
+
+class FBaseMenuBuilder;
+
+#define LOCTEXT_NAMESPACE "FSourceModelBuilders"
+
+namespace UE::ConcertSharedSlate
+{
+	/**
+	 * Utility for creating widgets of IItemSourceModel and categories thereof (TSourceSelectionCategory).
+	 */
+	template<typename TItemType>
+	struct FSourceModelBuilders
+	{
+		DECLARE_DELEGATE_OneParam(FOnItemsSelected, TArray<TItemType>);
+		DECLARE_DELEGATE_RetVal_OneParam(bool, FIsItemSelected, const TItemType& Item);
+		DECLARE_DELEGATE_RetVal_OneParam(FString, FGetItemDisplayString, const TItemType& Item);
+		DECLARE_DELEGATE_RetVal_OneParam(FSlateIcon, FGetItemIcon, const TItemType& Item);
+
+		struct FItemPickerArgs
+		{
+			/** Required. Called when items are picked by the widget. */
+			FOnItemsSelected OnItemsSelected;
+			
+			/** Required. Gets the display string for the item. */
+			FGetItemDisplayString GetItemDisplayString;
+
+			/** Optional. Gets the icon to display for the item. */
+			FGetItemIcon GetItemIcon;
+
+			/** Optional. Used primarily to filter out objects that are already selected. */
+			FIsItemSelected IsItemSelected;
+
+			FItemPickerArgs(FOnItemsSelected OnObjectsSelected, FGetItemDisplayString GetItemDisplayString, FGetItemIcon GetItemIcon = {}, FIsItemSelected IsItemSelected = {})
+				: OnItemsSelected(MoveTemp(OnObjectsSelected))
+				, GetItemDisplayString(MoveTemp(GetItemDisplayString))
+				, GetItemIcon(MoveTemp(GetItemIcon))
+				, IsItemSelected(MoveTemp(IsItemSelected))
+			{
+				check(this->OnItemsSelected.IsBound() && this->GetItemDisplayString.IsBound());
+			}
+		};
+		
+		/** Builds a combo button where each entry corresponds to one IItemSourceModel. If there is only one option, the option is inlined by calling MakeStandaloneWidget. */
+		static TSharedRef<SWidget> BuildCategory(const TSourceSelectionCategory<TItemType>& Category, const FItemPickerArgs& Args);
+	
+		/**
+		 * Creates a combo button inlining the source model. 
+		 *
+		 * This is the inlined case of when a category only contains one item source.
+		 * Used by BuildCategory.
+		 */
+		static TSharedRef<SWidget> MakeStandaloneWidget(TSharedRef<IItemSourceModel<TItemType>> Source, FItemPickerArgs Args, FSlateIcon Icon);
+
+		/** Adds the option to a context-menu */
+		static void AddOptionToMenu(TSharedRef<IItemSourceModel<TItemType>> Option, FItemPickerArgs Args, FBaseMenuBuilder& MenuBuilder);
+
+	private:
+		
+		static void FillSearchableSubmenu(const TSharedRef<IItemSourceModel<TItemType>>& Source, const FItemPickerArgs& Args, FBaseMenuBuilder& MenuBuilder);
+		static TSharedRef<SWidget> BuildMenu(const TArray<TAttribute<TSharedPtr<IItemSourceModel<TItemType>>>>& Options, TArray<TSourceSelectionCategory<TItemType>> SubCategories, const FItemPickerArgs& Args);
+		static void AddSubCategoryToMenu(const TSourceSelectionCategory<TItemType>& Category, const FItemPickerArgs& Args, FBaseMenuBuilder& MenuBuilder);
+	};
+}
+
+/************************* Implementation *************************/
+
+namespace UE::ConcertSharedSlate
+{
+	/** Builds a combo button where each entry corresponds to one IItemSourceModel. If there is only one option, the option is inlined by calling MakeStandaloneWidget. */
+	template<typename TItemType>
+	TSharedRef<SWidget> FSourceModelBuilders<TItemType>::BuildCategory(const TSourceSelectionCategory<TItemType>& Category, const FItemPickerArgs& Args)
+	{
+		if (!ensure(!Category.Options.IsEmpty()) || !ensure(Args.OnItemsSelected.IsBound()))
+		{
+			return SNullWidget::NullWidget;
+		}
+		
+		if (Category.Options.Num() == 1 && Category.SubCategories.Num() == 0)
+		{
+			const TSharedPtr<IItemSourceModel<TItemType>> Model = Category.Options[0].Get();
+			return ensure(Model) ? MakeStandaloneWidget(Model.ToSharedRef(), Args, Category.DisplayInfo.Icon) : SNullWidget::NullWidget;
+		}
+
+		return SNew(SPositiveActionButton)
+			.Text(Category.DisplayInfo.Label)
+			.ToolTipText(Category.DisplayInfo.ToolTip)
+			.Icon(Category.DisplayInfo.Icon.IsSet() ? Category.DisplayInfo.Icon.GetOptionalIcon() : FAppStyle::Get().GetBrush("Icons.Plus"))
+			.OnGetMenuContent_Lambda([Options = Category.Options, SubCategories = Category.SubCategories, Args]()
+			{
+				return BuildMenu(Options, SubCategories, Args);
+			});
+	}
+	
+	/** Creates a widget that displays the IItemSourceModel according to its specifications. Mainly to be used by BuildCategory. */
+	template<typename TItemType>
+	TSharedRef<SWidget> FSourceModelBuilders<TItemType>::MakeStandaloneWidget(TSharedRef<IItemSourceModel<TItemType>> Source, FItemPickerArgs Args, FSlateIcon Icon)
+	{
+		const FSourceDisplayInfo DisplayInfo = Source->GetDisplayInfo();
+		const FSlateBrush* IconBrush = DisplayInfo.Icon.IsSet() ? DisplayInfo.Icon.GetOptionalIcon() : FAppStyle::Get().GetBrush("Icons.Plus");
+		switch (DisplayInfo.SourceType)
+		{
+		case ESourceType::ShowAsList:
+			return SNew(SPositiveActionButton)
+				.Text(DisplayInfo.Label)
+				.ToolTipText(DisplayInfo.ToolTip)
+				.Icon(IconBrush)
+				.OnGetMenuContent_Lambda([Source = MoveTemp(Source), Args = MoveTemp(Args)]()
+				{
+					FMenuBuilder MenuBuilder(true, nullptr);
+					FillSearchableSubmenu(Source, Args, MenuBuilder);
+					return MenuBuilder.MakeWidget();
+				});
+		case ESourceType::AddOnClick: 
+			return SNew(SPositiveActionButton)
+				.Text(DisplayInfo.Label)
+				.ToolTipText(DisplayInfo.ToolTip)
+				.Icon(IconBrush)
+				.OnClicked_Lambda([Source = MoveTemp(Source), AddObjectsDelegate = Args.OnItemsSelected]()
+				{
+					AddObjectsDelegate.Execute(Source->GetSelectableItems());
+					return FReply::Handled();
+				});
+		default:
+			checkNoEntry();
+			return SNullWidget::NullWidget;
+		}
+	}
+
+	/** Adds the option to a context-menu */
+	template<typename TItemType>
+	void FSourceModelBuilders<TItemType>::AddOptionToMenu(TSharedRef<IItemSourceModel<TItemType>> Option, FItemPickerArgs Args, FBaseMenuBuilder& MenuBuilder)
+	{
+		const FSourceDisplayInfo DisplayInfo = Option->GetDisplayInfo();
+		switch (DisplayInfo.SourceType)
+		{
+		case ESourceType::ShowAsList:
+			{
+				FMenuEntryParams ShowAsListParams;
+				ShowAsListParams.LabelOverride = DisplayInfo.Label;
+				ShowAsListParams.ToolTipOverride = DisplayInfo.ToolTip;
+				ShowAsListParams.bIsSubMenu = true;
+				// Dummy is needed to avoid assert
+				ShowAsListParams.DirectActions = {
+					FExecuteAction::CreateLambda([](){})
+					};
+				ShowAsListParams.MenuBuilder.BindLambda([Source = MoveTemp(Option), Args = MoveTemp(Args)]()
+				{
+					FMenuBuilder MenuBuilder(true, nullptr);
+					FillSearchableSubmenu(Source, Args, MenuBuilder);
+					return MenuBuilder.MakeWidget();
+				});
+				
+				MenuBuilder.AddMenuEntry(ShowAsListParams);
+				break;
+			}
+		
+		case ESourceType::AddOnClick:
+			{
+				FMenuEntryParams AddOnClickParams;
+				AddOnClickParams.LabelOverride = DisplayInfo.Label;
+				AddOnClickParams.ToolTipOverride = DisplayInfo.ToolTip;
+				AddOnClickParams.UserInterfaceActionType = EUserInterfaceActionType::Button; 
+				AddOnClickParams.DirectActions = {
+					FExecuteAction::CreateLambda([Option, Args = MoveTemp(Args)](){ Args.OnItemsSelected.Execute(Option->GetSelectableItems()); }),
+					FCanExecuteAction::CreateLambda([Option](){ return Option->GetNumSelectableItems() > 0; })
+				};
+				
+				MenuBuilder.AddMenuEntry(AddOnClickParams);
+				break;
+			}
+		default:
+			checkNoEntry();
+		}
+	}
+	
+	template<typename TItemType>
+	void FSourceModelBuilders<TItemType>::FillSearchableSubmenu(const TSharedRef<IItemSourceModel<TItemType>>& Source, const FItemPickerArgs& Args, FBaseMenuBuilder& MenuBuilder)
+	{
+		TArray<TItemType> SortedItems = Source->GetSelectableItems();
+		Algo::Sort(SortedItems, [&Args](const TItemType& Left, const TItemType& Right)
+		{
+			return Args.GetItemDisplayString.Execute(Left) < Args.GetItemDisplayString.Execute(Right); 
+		});
+
+		// Add the objects to the MenuBuilder
+		for (const TItemType& Item : SortedItems)
+		{
+			// Gray out entries that will not do anything.
+			// Remember: we're not creating a context menu but a combo button next to the search bar.
+			// If it was just gone, users might be confused why suddenly there is a missing entry!
+			const FCanExecuteAction CanExecuteAction = FCanExecuteAction::CreateLambda([Item, IsItemSelected = Args.IsItemSelected]()
+			{
+				return !IsItemSelected.IsBound() || !IsItemSelected.Execute(Item);
+			});
+			MenuBuilder.AddMenuEntry(
+				FText::FromString(Args.GetItemDisplayString.Execute(Item)),
+				TAttribute<FText>(TAttribute<FText>::CreateLambda([CanExecuteAction](){ return CanExecuteAction.Execute() ? LOCTEXT("ItemSelected", "Item is already selected") : FText::GetEmpty(); })),
+				Args.GetItemIcon.IsBound() ? Args.GetItemIcon.Execute(Item) : FSlateIcon(),
+				FUIAction(
+					FExecuteAction::CreateLambda([Item, SelectItemDelegate = Args.OnItemsSelected]()
+					{
+						SelectItemDelegate.Execute({ Item });
+					}),
+					CanExecuteAction),
+				NAME_None,
+				EUserInterfaceActionType::Button
+				);
+		}
+	}
+
+	template<typename TItemType>
+	TSharedRef<SWidget> FSourceModelBuilders<TItemType>::BuildMenu(const TArray<TAttribute<TSharedPtr<IItemSourceModel<TItemType>>>>& Options, TArray<TSourceSelectionCategory<TItemType>> SubCategories, const FItemPickerArgs& Args)
+	{
+		FMenuBuilder MenuBuilder(true, nullptr);
+
+		// Sources first...
+		for (const TAttribute<TSharedPtr<IItemSourceModel<TItemType>>>& Source : Options)
+		{
+			const TSharedPtr<IItemSourceModel<TItemType>> Model = Source.Get();
+			if (ensure(Model))
+			{
+				AddOptionToMenu(Model.ToSharedRef(), Args, MenuBuilder);
+			}
+		}
+			
+		// ... then the sub-categories sorted by their names
+		SubCategories.Sort([](const TSourceSelectionCategory<TItemType>& Left, const TSourceSelectionCategory<TItemType>& Right)
+		{
+			// Is this slow? May want to cache ToString(). The number of items is so low it will hopefully have no noticeable effect.
+			return Left.DisplayInfo.Label.ToString() < Right.DisplayInfo.Label.ToString();
+		});
+		for (const TSourceSelectionCategory<TItemType>& SubCategory : SubCategories)
+		{
+			AddSubCategoryToMenu(SubCategory, Args, MenuBuilder);
+		}
+			
+		return MenuBuilder.MakeWidget();
+	}
+
+	template<typename TItemType>
+	void FSourceModelBuilders<TItemType>::AddSubCategoryToMenu(const TSourceSelectionCategory<TItemType>& Category, const FItemPickerArgs& Args, FBaseMenuBuilder& MenuBuilder)
+	{
+		FMenuEntryParams MenuEntry;
+		MenuEntry.LabelOverride = Category.DisplayInfo.Label;
+		MenuEntry.ToolTipOverride = Category.DisplayInfo.ToolTip;
+		MenuEntry.bIsSubMenu = true;
+		// Dummy is needed to avoid assert
+		MenuEntry.DirectActions = {
+			FExecuteAction::CreateLambda([](){})
+			};
+		MenuEntry.IconOverride = Category.DisplayInfo.Icon;
+		MenuEntry.MenuBuilder.BindLambda([Options = Category.Options, SubCategories = Category.SubCategories, Args]() mutable
+		{
+			return BuildMenu(Options, MoveTemp(SubCategories), Args);
+		});
+
+		MenuBuilder.AddMenuEntry(MenuEntry);
+	}
+}
+
+#undef LOCTEXT_NAMESPACE
\ No newline at end of file
diff --git a/Engine/Plugins/Developer/Concert/ConcertUI/ConcertSharedSlate/Source/ConcertSharedSlate/Public/Model/Item/SourceSelectionCategory.h b/Engine/Plugins/Developer/Concert/ConcertUI/ConcertSharedSlate/Source/ConcertSharedSlate/Public/Model/Item/SourceSelectionCategory.h
new file mode 100644
--- /dev/null
+++ b/Engine/Plugins/Developer/Concert/ConcertUI/ConcertSharedSlate/Source/ConcertSharedSlate/Public/Model/Item/SourceSelectionCategory.h
@@ -0,0 +1,31 @@
+// Copyright Epic Games, Inc. All Rights Reserved.
+
+#pragma once
+
+#include "IItemSourceModel.h"
+#include "Misc/Attribute.h"
+
+namespace UE::ConcertSharedSlate
+{
+	/**
+	 * A category groups IItemSourceModel together, e.g. a category "Add actor" could have options like
+	 *	- selected in outliner
+	 *	- drop-down list of actors in the world
+	 *
+	 *	Usually a category is displayed by a combo button with DisplayInfo being applied to it.
+	 *	However, if Options.Num() == 1, then for better UX the button can be "inlined" with a widget visualizing that option,
+	 *	e.g. a button that just invokes the single option and its IObjectSourceModel::GetDisplayInfo is applied to the button.
+	 */
+	template<typename TSelectionType>
+	struct TSourceSelectionCategory
+	{
+		/** Data to display about this category. */
+		FBaseDisplayInfo DisplayInfo;
+		
+		/** The options of this category. TAttribute to allow lazy instantiation (e.g. only when sub-menu is opened by user). */
+		TArray<TAttribute<TSharedPtr<IItemSourceModel<TSelectionType>>>> Options;
+
+		/** Sub-categories to show in the context menu. Example: "Add component" is subcategory with the items being FComponentFromActorSource_Root. */
+		TArray<TSourceSelectionCategory> SubCategories;
+	};
+}
