Remove feature activation CVars

diff --git a/Engine/Plugins/Developer/Concert/ConcertApp/MultiUserClient/Source/MultiUserReplicationEditor/Private/Assets/ReplicationClientProfileAssetFactory.cpp b/Engine/Plugins/Developer/Concert/ConcertApp/MultiUserClient/Source/MultiUserReplicationEditor/Private/Assets/ReplicationClientProfileAssetFactory.cpp
--- a/Engine/Plugins/Developer/Concert/ConcertApp/MultiUserClient/Source/MultiUserReplicationEditor/Private/Assets/ReplicationClientProfileAssetFactory.cpp
+++ b/Engine/Plugins/Developer/Concert/ConcertApp/MultiUserClient/Source/MultiUserReplicationEditor/Private/Assets/ReplicationClientProfileAssetFactory.cpp
@@ -2,7 +2,6 @@
 
 #include "ReplicationClientProfileAssetFactory.h"
 
-#include "CVarMultiUserReplicationEditor.h"
 #include "MultiUserReplicationClientProfileAsset.h"
 
 #define LOCTEXT_NAMESPACE "UReplicationClientProfileAssetFactory"
@@ -17,7 +16,7 @@ UReplicationClientProfileAssetFactory::UReplicationClientProfileAssetFactory(con
 
 bool UReplicationClientProfileAssetFactory::CanCreateNew() const
 {
-	return UE::MultiUserReplicationEditor::ConsoleVariables::CVarEnableReplication.GetValueOnAnyThread();
+	return true;
 }
 
 FText UReplicationClientProfileAssetFactory::GetDisplayName() const
diff --git a/Engine/Plugins/Developer/Concert/ConcertApp/MultiUserClient/Source/MultiUserReplicationEditor/Private/Assets/ReplicationStreamAssetFactory.cpp b/Engine/Plugins/Developer/Concert/ConcertApp/MultiUserClient/Source/MultiUserReplicationEditor/Private/Assets/ReplicationStreamAssetFactory.cpp
--- a/Engine/Plugins/Developer/Concert/ConcertApp/MultiUserClient/Source/MultiUserReplicationEditor/Private/Assets/ReplicationStreamAssetFactory.cpp
+++ b/Engine/Plugins/Developer/Concert/ConcertApp/MultiUserClient/Source/MultiUserReplicationEditor/Private/Assets/ReplicationStreamAssetFactory.cpp
@@ -2,7 +2,6 @@
 
 #include "ReplicationStreamAssetFactory.h"
 
-#include "CVarMultiUserReplicationEditor.h"
 #include "MultiUserReplicationStreamAsset.h"
 
 #define LOCTEXT_NAMESPACE "ReplicationStreamAssetFactory"
@@ -17,7 +16,7 @@ UReplicationStreamAssetFactory::UReplicationStreamAssetFactory(const FObjectInit
 
 bool UReplicationStreamAssetFactory::CanCreateNew() const
 {
-	return UE::MultiUserReplicationEditor::ConsoleVariables::CVarEnableReplication.GetValueOnAnyThread();;
+	return true;
 }
 
 FText UReplicationStreamAssetFactory::GetDisplayName() const
diff --git a/Engine/Plugins/Developer/Concert/ConcertApp/MultiUserClient/Source/MultiUserReplicationEditor/Private/CVarMultiUserReplicationEditor.cpp b/Engine/Plugins/Developer/Concert/ConcertApp/MultiUserClient/Source/MultiUserReplicationEditor/Private/CVarMultiUserReplicationEditor.cpp
deleted file mode 100644
diff --git a/Engine/Plugins/Developer/Concert/ConcertApp/MultiUserClient/Source/MultiUserReplicationEditor/Private/CVarMultiUserReplicationEditor.h b/Engine/Plugins/Developer/Concert/ConcertApp/MultiUserClient/Source/MultiUserReplicationEditor/Private/CVarMultiUserReplicationEditor.h
deleted file mode 100644
diff --git a/Engine/Plugins/Developer/Concert/ConcertApp/MultiUserClient/Source/MultiUserReplicationEditor/Private/MultiUserReplicationEditorModule.cpp b/Engine/Plugins/Developer/Concert/ConcertApp/MultiUserClient/Source/MultiUserReplicationEditor/Private/MultiUserReplicationEditorModule.cpp
--- a/Engine/Plugins/Developer/Concert/ConcertApp/MultiUserClient/Source/MultiUserReplicationEditor/Private/MultiUserReplicationEditorModule.cpp
+++ b/Engine/Plugins/Developer/Concert/ConcertApp/MultiUserClient/Source/MultiUserReplicationEditor/Private/MultiUserReplicationEditorModule.cpp
@@ -3,7 +3,6 @@
 #include "MultiUserReplicationEditorModule.h"
 
 #include "ConcertSyncSessionFlags.h"
-#include "CVarMultiUserReplicationEditor.h"
 #include "IMultiUserClientModule.h"
 
 #include "ISettingsModule.h"
@@ -22,45 +21,35 @@ namespace UE::MultiUserReplicationEditor
 	void FMultiUserReplicationEditorModule::StartupModule()
 	{
 		FMultiUserReplicationEditorStyle::Initialize();
-
-		UpdateSettingsRegistrationBasedOnCVar();
-		ConsoleVariables::CVarEnableReplication->OnChangedDelegate().AddLambda([this](IConsoleVariable*)
-		{
-			UpdateSettingsRegistrationBasedOnCVar();
-		});
+		RegisterSettings();
 	}
 
 	void FMultiUserReplicationEditorModule::ShutdownModule()
 	{
 		FMultiUserReplicationEditorStyle::Shutdown();
-	}
-
-	void FMultiUserReplicationEditorModule::UpdateSettingsRegistrationBasedOnCVar()
-	{
-		if (ConsoleVariables::CVarEnableReplication.GetValueOnAnyThread())
-		{
-			RegisterSettings();
-		}
-		else
-		{
-			UnregisterSettings();
-		}
+		UnregisterSettings();
 	}
 
 	void FMultiUserReplicationEditorModule::RegisterSettings()
 	{
 		ISettingsModule* SettingsModule = FModuleManager::GetModulePtr<ISettingsModule>("Settings");
-		SettingsModule->RegisterSettings("Project", "Plugins", "Multi-User Replication",
-			LOCTEXT("MultiUserReplicationSettingsName", "Multi-User Replication"),
-			LOCTEXT("MultiUserReplicationSettingsDescription", "Configure the Multi-User Replication settings."),
-			UMultiUserReplicationSettings::Get()
-		);
+		if (ensure(SettingsModule))
+		{
+			SettingsModule->RegisterSettings("Project", "Plugins", "Multi-User Replication",
+				LOCTEXT("MultiUserReplicationSettingsName", "Multi-User Replication"),
+				LOCTEXT("MultiUserReplicationSettingsDescription", "Configure the Multi-User Replication settings."),
+				UMultiUserReplicationSettings::Get()
+			);
+		}
 	}
 
 	void FMultiUserReplicationEditorModule::UnregisterSettings()
 	{
 		ISettingsModule* SettingsModule = FModuleManager::GetModulePtr<ISettingsModule>("Settings");
-		SettingsModule->UnregisterSettings("Project", "Plugins", "Multi-User Replication");
+		if (SettingsModule)
+		{
+			SettingsModule->UnregisterSettings("Project", "Plugins", "Multi-User Replication");
+		}
 	}
 };
 
diff --git a/Engine/Plugins/Developer/Concert/ConcertApp/MultiUserClient/Source/MultiUserReplicationEditor/Private/MultiUserReplicationEditorModule.h b/Engine/Plugins/Developer/Concert/ConcertApp/MultiUserClient/Source/MultiUserReplicationEditor/Private/MultiUserReplicationEditorModule.h
--- a/Engine/Plugins/Developer/Concert/ConcertApp/MultiUserClient/Source/MultiUserReplicationEditor/Private/MultiUserReplicationEditorModule.h
+++ b/Engine/Plugins/Developer/Concert/ConcertApp/MultiUserClient/Source/MultiUserReplicationEditor/Private/MultiUserReplicationEditorModule.h
@@ -26,7 +26,6 @@ namespace UE::MultiUserReplicationEditor
 
 		FAssetCategoryPath MultiUserReplicationCategory;
 
-		void UpdateSettingsRegistrationBasedOnCVar();
 		void RegisterSettings();
 		void UnregisterSettings();
 	};
