Split SConcertBrowser.cpp up into multiple .cpp & .h files

diff --git a/Engine/Plugins/Developer/Concert/ConcertApp/MultiUserClient/Source/MultiUserClient/Private/Widgets/Disconnected/ConcertClientSessionBrowserController.cpp b/Engine/Plugins/Developer/Concert/ConcertApp/MultiUserClient/Source/MultiUserClient/Private/Widgets/Disconnected/ConcertClientSessionBrowserController.cpp
new file mode 100644
--- /dev/null
+++ b/Engine/Plugins/Developer/Concert/ConcertApp/MultiUserClient/Source/MultiUserClient/Private/Widgets/Disconnected/ConcertClientSessionBrowserController.cpp
@@ -0,0 +1,654 @@
+// Copyright Epic Games, Inc. All Rights Reserved.
+
+#include "ConcertClientSessionBrowserController.h"
+
+#include "ConcertActivityStream.h"
+#include "ConcertLogGlobal.h"
+#include "IMultiUserClientModule.h"
+#include "IConcertClient.h"
+#include "IConcertSyncClient.h"
+#include "SConcertSessionRecovery.h"
+
+#include "Algo/Transform.h"
+#include "Framework/Application/SlateApplication.h"
+#include "Misc/AsyncTaskNotification.h"
+#include "Widgets/SWindow.h"
+
+#define LOCTEXT_NAMESPACE "FConcertClientSessionBrowserController"
+
+FConcertClientSessionBrowserController::FConcertClientSessionBrowserController(IConcertClientPtr InConcertClient)
+{
+	check(InConcertClient.IsValid()); // Don't expect this class to be instantiated if the concert client is not available.
+	check(InConcertClient->IsConfigured()); // Expected to be done by higher level code.
+	ConcertClient = InConcertClient;
+
+	// When others servers are running, add them to the list if you want to test the UI displayed if no servers/no sessions exists.
+	//IgnoredServers.Add(TEXT("wksyul10355")); // TODO: COMMENT BEFORE SUBMIT
+
+	// Start server discovery to find the available Concert servers.
+	ConcertClient->StartDiscovery();
+
+	// Populate the session cache.
+	TickServersAndSessionsDiscovery();
+}
+
+FConcertClientSessionBrowserController::~FConcertClientSessionBrowserController()
+{
+	if (ConcertClient.IsValid())
+	{
+		ConcertClient->StopDiscovery();
+	}
+}
+
+const TArray<FConcertSessionClientInfo>& FConcertClientSessionBrowserController::GetClients(const FGuid& AdminEndpoint, const FGuid& SessionId) const
+{
+	// If a session clients are monitored and the list is cached
+	if (ClientMonitoredSession.IsSet() && ClientMonitoredSession->Session->ServerInfo.AdminEndpointId == AdminEndpoint && ClientMonitoredSession->Session->SessionInfo.SessionId == SessionId)
+	{
+		return ClientMonitoredSession->Session->Clients; // Returns the list retrieved by the last TickClientsDiscovery().
+	}
+
+	// Returns an empty list for now. We expect the caller to call TickClientsDiscovery() periodically to maintain the session client list.
+	static TArray<FConcertSessionClientInfo> EmptyClientList;
+	return EmptyClientList;
+}
+
+void FConcertClientSessionBrowserController::JoinSession(const FGuid& ServerAdminEndpointId, const FGuid& SessionId)
+{
+	// On success: The client joins the session and SConcertBrowser::HandleSessionConnectionChanged() will transit the UI to the SActiveSession.
+	// On failure: An async notification banner will be displayer to the user.
+	ConcertClient->JoinSession(ServerAdminEndpointId, SessionId);
+}
+
+TOptional<FConcertSessionInfo> FConcertClientSessionBrowserController::GetActiveSessionInfo(const FGuid& AdminEndpoint, const FGuid& SessionId) const
+{
+	const TSharedPtr<FClientActiveSessionInfo>* SessionInfo = ActiveSessions.FindByPredicate([&AdminEndpoint, &SessionId](const TSharedPtr<FActiveSessionInfo>& MatchCandidate)
+	{
+		return MatchCandidate->ServerInfo.AdminEndpointId == AdminEndpoint && MatchCandidate->SessionInfo.SessionId == SessionId;
+	});
+
+	return SessionInfo != nullptr ? (*SessionInfo)->SessionInfo : TOptional<FConcertSessionInfo>{};
+}
+
+TOptional<FConcertSessionInfo> FConcertClientSessionBrowserController::GetArchivedSessionInfo(const FGuid& AdminEndpoint, const FGuid& SessionId) const
+{
+	const TSharedPtr<FClientArchivedSessionInfo>* SessionInfo = ArchivedSessions.FindByPredicate([&AdminEndpoint, &SessionId](const TSharedPtr<FArchivedSessionInfo>& MatchCandidate)
+	{
+		return MatchCandidate->ServerInfo.AdminEndpointId == AdminEndpoint && MatchCandidate->SessionInfo.SessionId == SessionId;
+	});
+
+	return SessionInfo != nullptr ? (*SessionInfo)->SessionInfo : TOptional<FConcertSessionInfo>{};
+}
+
+void FConcertClientSessionBrowserController::CreateSession(const FGuid& ServerAdminEndpointId, const FString& SessionName, const FString& ProjectName)
+{
+	FAsyncRequest& CreateRequest = CreateSessionRequests.AddDefaulted_GetRef();
+	TWeakPtr<uint8, ESPMode::ThreadSafe> CreateRequestExecutionToken = CreateRequest.ResetExecutionToken();
+
+	// On success, the client automatically joins the new session and SConcertBrowser::HandleSessionConnectionChanged() will transit the UI to the SActiveSession.
+	// On failure: An async notification banner will be displayer to the user.
+	FConcertCreateSessionArgs CreateSessionArgs;
+	CreateSessionArgs.SessionName = SessionName;
+	// Project Name is not used an override for the client session create.
+	ConcertClient->CreateSession(ServerAdminEndpointId, CreateSessionArgs).Next([this, CreateRequestExecutionToken, ServerAdminEndpointId, SessionName](EConcertResponseCode ResponseCode)
+	{
+		if (TSharedPtr<uint8, ESPMode::ThreadSafe> ExecutionToken = CreateRequestExecutionToken.Pin())
+		{
+			if (ResponseCode == EConcertResponseCode::Success)
+			{
+				// Expect to find those session at some point in the future.
+				ExpectedSessionsToDiscover.Add(FPendingSessionDiscovery{ FDateTime::UtcNow(), ServerAdminEndpointId, SessionName });
+			}
+
+			// Stop tracking the request.
+			CreateSessionRequests.RemoveAll([&ExecutionToken](const FAsyncRequest& DiscardCandidate) { return DiscardCandidate.FutureExecutionToken == ExecutionToken; });
+		}
+	});
+}
+
+void FConcertClientSessionBrowserController::ArchiveSession(const FGuid& ServerAdminEndpointId, const FGuid& SessionId, const FString& ArchiveName, const FConcertSessionFilter& SessionFilter)
+{
+	// On success, an archived is created and TickServersAndSessionsDiscovery() will eventually discover it.
+	// On failure: An async notification banner will be displayer to the user.
+	FConcertArchiveSessionArgs ArchiveSessionArgs;
+	ArchiveSessionArgs.SessionId = SessionId;
+	ArchiveSessionArgs.ArchiveNameOverride = ArchiveName;
+	ArchiveSessionArgs.SessionFilter = SessionFilter;
+	ConcertClient->ArchiveSession(ServerAdminEndpointId, ArchiveSessionArgs);
+}
+
+void FConcertClientSessionBrowserController::RestoreSession(const FGuid& ServerAdminEndpointId, const FGuid& SessionId, const FString& RestoredName, const FConcertSessionFilter& SessionFilter)
+{
+	FString ArchivedSessionName;
+	if (const TOptional<FConcertSessionInfo> SessionInfo = GetArchivedSessionInfo(ServerAdminEndpointId, SessionId))
+	{
+		ArchivedSessionName = SessionInfo->SessionName;
+	}
+
+	TSharedRef<SWindow> NewWindow = SNew(SWindow)
+		.Title(FText::Format(LOCTEXT("RestoreSessionDialogTitle", "Restoring {0}"), FText::AsCultureInvariant(ArchivedSessionName)))
+		.SizingRule(ESizingRule::UserSized)
+		.ClientSize(FVector2D(1200, 800))
+		.IsTopmostWindow(false) // Consider making it always on top?
+		.SupportsMaximize(true)
+		.SupportsMinimize(false);
+
+	// Ask the stream to pull the activity details (for transaction/package) for inspection.
+	constexpr bool bRequestActivityDetails = true;
+
+	// Create a stream of activities (streaming from the most recent to the oldest).
+	TSharedPtr<FConcertActivityStream> ActivityStream = MakeShared<FConcertActivityStream>(ConcertClient, ServerAdminEndpointId, SessionId, bRequestActivityDetails);
+
+	// The UI uses this function to read and consume the activity stream.
+	auto ReadActivitiesFn = [ActivityStream](TArray<TSharedPtr<FConcertSessionActivity>>& InOutActivities, int32& OutFetchCount, FText& OutErrorMsg)
+	{
+		return ActivityStream->Read(InOutActivities, OutFetchCount, OutErrorMsg);
+	};
+
+	// The UI uses this function to map an activity ID from the stream to a client info.
+	auto GetActivityClientInfoFn = [ActivityStream](FGuid EndpointID)
+	{
+		const FConcertClientInfo* Result = ActivityStream->GetActivityClientInfo(EndpointID);
+		return Result ? TOptional<FConcertClientInfo>{ *Result } : TOptional<FConcertClientInfo>{};
+	};
+
+	// Invoked if the client selects a point in time to recover.
+	TWeakPtr<IConcertClient, ESPMode::ThreadSafe> WeakClient = ConcertClient;
+	auto OnAcceptRestoreFn = [WeakClient, ServerAdminEndpointId, SessionId, RestoredName, SessionFilter](TSharedPtr<FConcertSessionActivity> SelectedRecoveryActivity)
+	{
+		FConcertCopySessionArgs RestoreSessionArgs;
+		RestoreSessionArgs.bAutoConnect = true;
+		RestoreSessionArgs.SessionId = SessionId;
+		RestoreSessionArgs.SessionName = RestoredName;
+		RestoreSessionArgs.SessionFilter = SessionFilter;
+		RestoreSessionArgs.SessionFilter.bOnlyLiveData = false;
+
+		// Set which item was selected to recover through.
+		if (SelectedRecoveryActivity)
+		{
+			RestoreSessionArgs.SessionFilter.ActivityIdUpperBound = SelectedRecoveryActivity->Activity.ActivityId;
+		}
+		// else -> Restore the entire session as it.
+
+		bool bDismissRecoveryWindow = true; // Dismiss the window showing the session activities
+
+		if (TSharedPtr<IConcertClient, ESPMode::ThreadSafe> ConcertClientPin = WeakClient.Pin())
+		{
+			// Prompt the user to persist and leave the session.
+			bool bDisconnected = IMultiUserClientModule::Get().DisconnectSession(/*bAlwaysAskConfirmation*/true);
+			if (bDisconnected)
+			{
+				// On success, a new session is created from the archive, the client automatically disconnects from the current session (if any), joins the restored one and SConcertBrowser::HandleSessionConnectionChanged() will transit the UI to the SActiveSession.
+				// On failure: An async notification banner will be displayer to the user.
+				ConcertClientPin->RestoreSession(ServerAdminEndpointId, RestoreSessionArgs);
+			}
+			else // The user declined disconnection.
+			{
+				bDismissRecoveryWindow = false; // Keep the window open, let the client handle it (close/cancel or just restore later).
+			}
+		}
+		else
+		{
+			FAsyncTaskNotificationConfig NotificationConfig;
+			NotificationConfig.bIsHeadless = false;
+			NotificationConfig.bKeepOpenOnFailure = true;
+			NotificationConfig.LogCategory = &LogConcert;
+
+			FAsyncTaskNotification Notification(NotificationConfig);
+			Notification.SetComplete(LOCTEXT("RecoveryError", "Failed to recover the session"), LOCTEXT("ClientUnavailable", "Concert client unavailable"), /*Success*/ false);
+		}
+		return bDismissRecoveryWindow;
+	};
+
+	TSharedRef<SConcertSessionRecovery> RestoreWidget = SNew(SConcertSessionRecovery)
+		.ParentWindow(NewWindow)
+		.IntroductionText(LOCTEXT("RecoverSessionIntroductionText", "Select the point in time at which the session should be restored"))
+		.OnFetchActivities_Lambda(ReadActivitiesFn)
+		.OnMapActivityToClient_Lambda(GetActivityClientInfoFn)
+		.OnRestore(OnAcceptRestoreFn)
+		.WithClientAvatarColorColumn(true)
+		.WithClientNameColumn(true)
+		.WithOperationColumn(true)
+		.WithPackageColumn(false) // Even tough the column is not present, the tooltips and summary contains the affected package.
+		.DetailsAreaVisibility(bRequestActivityDetails ? EVisibility::Visible : EVisibility::Collapsed) // The activity stream was configured to pull the activity details.
+		.IsConnectionActivityFilteringEnabled(true)
+		.IsLockActivityFilteringEnabled(true);
+
+	NewWindow->SetContent(RestoreWidget);
+	FSlateApplication::Get().AddWindow(NewWindow, true);
+}
+
+void FConcertClientSessionBrowserController::RenameActiveSession(const FGuid& ServerAdminEndpointId, const FGuid& SessionId, const FString& NewName)
+{
+	// Find the currently cached session info.
+	const TSharedPtr<FClientActiveSessionInfo>* SessionInfo = ActiveSessions.FindByPredicate([&ServerAdminEndpointId, &SessionId](const TSharedPtr<FActiveSessionInfo>& MatchCandidate)
+	{
+		return MatchCandidate->ServerInfo.AdminEndpointId == ServerAdminEndpointId && MatchCandidate->SessionInfo.SessionId == SessionId;
+	});
+
+	check(SessionInfo); // If the UI is displaying it, the UI backend should have it.
+	ConcertClient->RenameSession(ServerAdminEndpointId, SessionId, NewName).Next([ActiveSessionInfo = *SessionInfo](EConcertResponseCode Response)
+	{
+		if (Response != EConcertResponseCode::Success)
+		{
+			ActiveSessionInfo->bSessionNameDirty = true; // Renamed failed, the UI may be displaying the wrong name. Force the UI to update against the latest cached values.
+		}
+		// else -> Succeeded -> TickServersAndSessionsDiscovery() will pick up the new name.
+	});
+}
+
+void FConcertClientSessionBrowserController::RenameArchivedSession(const FGuid& ServerAdminEndpointId, const FGuid& SessionId, const FString& NewName)
+{
+	// Find the currently cached session info.
+	const TSharedPtr<FClientArchivedSessionInfo>* SessionInfo = ArchivedSessions.FindByPredicate([&ServerAdminEndpointId, &SessionId](const TSharedPtr<FArchivedSessionInfo>& MatchCandidate)
+	{
+		return MatchCandidate->ServerInfo.AdminEndpointId == ServerAdminEndpointId && MatchCandidate->SessionInfo.SessionId == SessionId;
+	});
+
+	check(SessionInfo); // If the UI is displaying it, the UI backend should have it.
+	ConcertClient->RenameSession(ServerAdminEndpointId, SessionId, NewName).Next([SessionInfo = *SessionInfo](EConcertResponseCode Response)
+	{
+		if (Response != EConcertResponseCode::Success)
+		{
+			SessionInfo->bSessionNameDirty = true; // Renamed failed, the UI may be displaying the wrong name. Force the UI to update against the latest cached values.
+		}
+		// else -> Succeeded -> TickServersAndSessionsDiscovery() will pick up the new name.
+	});
+}
+
+bool FConcertClientSessionBrowserController::CanRenameActiveSession(const FGuid& ServerAdminEndpointId, const FGuid& SessionId) const
+{
+	return CanDeleteActiveSession(ServerAdminEndpointId, SessionId); // Rename requires the same permission than delete.
+}
+
+bool FConcertClientSessionBrowserController::CanRenameArchivedSession(const FGuid& ServerAdminEndpointId, const FGuid& SessionId) const
+{
+	return CanDeleteArchivedSession(ServerAdminEndpointId, SessionId); // Rename requires the same permission than delete.
+}
+
+void FConcertClientSessionBrowserController::DeleteSessions(const FGuid& ServerAdminEndpointId, const TArray<FGuid>& SessionIds)
+{
+	// The difference between the two functions is how error notifications are displayed...
+	if (SessionIds.Num() == 1)
+	{
+		// ... shows a specific notification
+		ConcertClient->DeleteSession(ServerAdminEndpointId, SessionIds[0]);
+	}
+	else
+	{
+		TSet<FGuid> SessionIdSet;
+		Algo::Transform(SessionIds, SessionIdSet, [](const FGuid& SessionId) { return SessionId; });
+		// ... shows compressed notification(s)
+		ConcertClient->BatchDeleteSessions(ServerAdminEndpointId, { SessionIdSet, EBatchSessionDeletionFlags::SkipForbiddenSessions });
+	}
+}
+
+bool FConcertClientSessionBrowserController::CanDeleteActiveSession(const FGuid& ServerAdminEndpointId, const FGuid& SessionId) const
+{
+	const TSharedPtr<FClientActiveSessionInfo>* SessionInfo = ActiveSessions.FindByPredicate([&ServerAdminEndpointId, &SessionId](const TSharedPtr<FActiveSessionInfo>& MatchCandidate)
+	{
+		return MatchCandidate->ServerInfo.AdminEndpointId == ServerAdminEndpointId && MatchCandidate->SessionInfo.SessionId == SessionId;
+	});
+
+	// Can delete the session only if the concert client is the owner.
+	return SessionInfo == nullptr ? false : ConcertClient->IsOwnerOf((*SessionInfo)->SessionInfo);
+}
+
+bool FConcertClientSessionBrowserController::CanDeleteArchivedSession(const FGuid& ServerAdminEndpointId, const FGuid& SessionId) const
+{
+	const TSharedPtr<FClientArchivedSessionInfo>* SessionInfo = ArchivedSessions.FindByPredicate([&ServerAdminEndpointId, &SessionId](const TSharedPtr<FArchivedSessionInfo>& MatchCandidate)
+	{
+		return MatchCandidate->ServerInfo.AdminEndpointId == ServerAdminEndpointId && MatchCandidate->SessionInfo.SessionId == SessionId;
+	});
+
+	// Can delete the session only if the concert client is the owner.
+	return SessionInfo == nullptr ? false : ConcertClient->IsOwnerOf((*SessionInfo)->SessionInfo);
+}
+
+TPair<uint32, uint32> FConcertClientSessionBrowserController::TickServersAndSessionsDiscovery()
+{
+	// Fire new async queries to poll the sessions on all known servers.
+	UpdateSessionsAsync();
+
+	// Returns the versions corresponding to the last async responses received. The version are incremented when response to the asynchronous queries are received.
+	return MakeTuple(ServerListVersion, SessionListVersion);
+}
+
+uint32 FConcertClientSessionBrowserController::TickClientsDiscovery(const FGuid& ServerAdminEndpointId, const FGuid& SessionId)
+{
+	// Fire new async queries to poll clients for the selected server/session pair.
+	UpdateClientsAsync(ServerAdminEndpointId, SessionId);
+
+	// Returns the versions corresponding to the last async responses received. The version is incremented when response to the asynchronous query is received.
+	return ClientListVersion;
+}
+
+bool FConcertClientSessionBrowserController::GetAndClearDiscoveryUpdateFlag()
+{
+	bool bOldCacheChanged = bCacheUpdated; // This flag is raised every time an async response updates the cached data.
+	bCacheUpdated = false;
+	return bOldCacheChanged;
+}
+
+void FConcertClientSessionBrowserController::UpdateSessionsAsync()
+{
+	// Get the list of known servers.
+	TArray<FConcertServerInfo> OnlineServers = ConcertClient->GetKnownServers();
+	if (IgnoredServers.Num())
+	{
+		OnlineServers.RemoveAll([this](const FConcertServerInfo& ServerInfo) { return IgnoredServers.Contains(ServerInfo.ServerName); });
+	}
+
+	bool bServerListVersionUpdated = false;
+
+	// Detects which server(s) went offline since the last update.
+	for (const FConcertServerInfo& ServerInfo : Servers)
+	{
+		// If a server, previously tracked, is not online anymore, remove all its sessions.
+		if (!OnlineServers.ContainsByPredicate([this, &ServerInfo](const FConcertServerInfo& Visited) { return ServerInfo.InstanceInfo.InstanceId == Visited.InstanceInfo.InstanceId; }))
+		{
+			// Remove all active sessions corresponding to this server.
+			ActiveSessions.RemoveAll([this, &ServerInfo](const TSharedPtr<FActiveSessionInfo>& ActiveSessionInfo)
+			{
+				if (ServerInfo.InstanceInfo.InstanceId == ActiveSessionInfo->ServerInfo.InstanceInfo.InstanceId)
+				{
+					OnActiveSessionDiscarded(*ActiveSessionInfo);
+					return true;
+				}
+				return false;
+			});
+
+			// Remove all archived sessions corresponding to this server.
+			ArchivedSessions.RemoveAll([this, &ServerInfo](const TSharedPtr<FArchivedSessionInfo>& ArchivedSessionInfo)
+			{
+				if (ServerInfo.InstanceInfo.InstanceId == ArchivedSessionInfo->ServerInfo.InstanceInfo.InstanceId)
+				{
+					OnArchivedSessionDiscarded(*ArchivedSessionInfo);
+					return true;
+				}
+				return false;
+			});
+
+			// Disarm any active async request future for this server. Removing it from the map effectively disarm the future as the shared pointer used to as token gets released.
+			ActiveSessionRequests.Remove(ServerInfo.InstanceInfo.InstanceId);
+			ArchivedSessionRequests.Remove(ServerInfo.InstanceInfo.InstanceId);
+
+			// The server went offline, update the server list version.
+			ServerListVersion++;
+			bServerListVersionUpdated = true;
+		}
+	}
+
+	// For all online servers.
+	for (const FConcertServerInfo& ServerInfo : OnlineServers)
+	{
+		// Check if this is a new server.
+		if (!bServerListVersionUpdated && !Servers.ContainsByPredicate([&ServerInfo](const FConcertServerInfo& Visited) { return ServerInfo.InstanceInfo.InstanceId == Visited.InstanceInfo.InstanceId; }))
+		{
+			ServerListVersion++;
+			bServerListVersionUpdated = true; // No need to look up further, only one new server is required to update the version.
+		}
+
+		// Poll the sessions for this online server.
+		UpdateActiveSessionsAsync(ServerInfo);
+		UpdateArchivedSessionsAsync(ServerInfo);
+	}
+
+	// Keep the list of online servers.
+	Servers = MoveTemp(OnlineServers);
+}
+
+void FConcertClientSessionBrowserController::UpdateActiveSessionsAsync(const FConcertServerInfo& ServerInfo)
+{
+	// Check if a request is already pulling active session from this server.
+	FAsyncRequest& ListActiveSessionAsyncRequest = ActiveSessionRequests.FindOrAdd(ServerInfo.InstanceInfo.InstanceId);
+	if (ListActiveSessionAsyncRequest.IsOngoing())
+	{
+		// Don't stack another async request on top of the ongoing one, TickServersAndSessionsDiscovery() will eventually catch up any possibly missed info here.
+		return;
+	}
+
+	// Arm the future, enabling its continuation to execute (or not).
+	TWeakPtr<uint8, ESPMode::ThreadSafe> ListActiveSessionExecutionToken = ListActiveSessionAsyncRequest.ResetExecutionToken();
+
+	// Keep track of the request time.
+	FDateTime ListRequestTimestamp = FDateTime::UtcNow();
+
+	// Retrieve all live sessions currently known by this server.
+	ListActiveSessionAsyncRequest.Future = ConcertClient->GetLiveSessions(ServerInfo.AdminEndpointId)
+		.Next([this, ServerInfo, ListActiveSessionExecutionToken, ListRequestTimestamp](const FConcertAdmin_GetSessionsResponse& Response)
+		{
+			// If the future is disarmed.
+			if (!ListActiveSessionExecutionToken.IsValid())
+			{
+				// Don't go further, the future execution was canceled, maybe because this object was deleted, the server removed from the list or it wasn't safe to get this future executing.
+				return;
+			}
+
+			// If the server responded.
+			if (Response.ResponseCode == EConcertResponseCode::Success)
+			{
+				bInitialActiveSessionQueryResponded = true;
+
+				// Remove from the cache any session that were deleted from 'Server' since the last update. Find which one were renamed.
+				ActiveSessions.RemoveAll([this, &ServerInfo, &Response](const TSharedPtr<FClientActiveSessionInfo>& DiscardCandidate)
+				{
+					// If the candidate is owned by another server.
+					if (DiscardCandidate->ServerInfo.InstanceInfo.InstanceId != ServerInfo.InstanceInfo.InstanceId)
+					{
+						return false; // Keep that session, it's owned by another server.
+					}
+					// If the candidate is still active.
+					else if (const FConcertSessionInfo* SessionFromServer = Response.Sessions.FindByPredicate([DiscardCandidate](const FConcertSessionInfo& MatchCandidate) { return DiscardCandidate->SessionInfo.SessionId == MatchCandidate.SessionId; }))
+					{
+						// Don't discard the session, it is still active, but check if it was renamed.
+						if (SessionFromServer->SessionName != DiscardCandidate->SessionInfo.SessionName)
+						{
+							DiscardCandidate->SessionInfo.SessionName = SessionFromServer->SessionName; // Update the session name.
+							OnActiveSessionRenamed(*DiscardCandidate, SessionFromServer->SessionName);
+						}
+						else if (DiscardCandidate->bSessionNameDirty) // Renaming this particular session failed?
+						{
+							OnActiveSessionListDirty(); // UI optimistically updated a session name from a rename, but it failed on the server. The controller caches the latest official server list unmodified. Force the client to resync its list against the cache.
+							DiscardCandidate->bSessionNameDirty = false;
+						}
+						return false; // Don't remove the session, it's still active.
+					}
+
+					// The session is not active anymore on the server, it was discarded.
+					OnActiveSessionDiscarded(*DiscardCandidate);
+					return true;
+				});
+
+				// Add to the cache any session that was added to 'Server' since the last update.
+				for (const FConcertSessionInfo& SessionInfo : Response.Sessions)
+				{
+					// Try to find the session in the list of active session.
+					if (!ActiveSessions.ContainsByPredicate([&SessionInfo](const TSharedPtr<FClientActiveSessionInfo>& MatchCandidate)
+						{ return SessionInfo.ServerInstanceId == MatchCandidate->ServerInfo.InstanceInfo.InstanceId && SessionInfo.SessionId == MatchCandidate->SessionInfo.SessionId; }))
+					{
+						// This is a newly discovered session, add it to the list.
+						ActiveSessions.Add(MakeShared<FClientActiveSessionInfo>(ServerInfo, SessionInfo));
+						OnActiveSessionDiscovered(*ActiveSessions.Last());
+					}
+
+					// Remove sessions that were created and discovered or those that should have shown up by now, but have not.
+					ExpectedSessionsToDiscover.RemoveAll([&SessionInfo, &ListRequestTimestamp](const FPendingSessionDiscovery& DiscardCandidate)
+					{
+						bool bDiscovered = DiscardCandidate.ServerEndpoint == SessionInfo.ServerEndpointId && DiscardCandidate.SessionName == SessionInfo.SessionName; // Session was discovered as expected.
+						bool bDeleted = ListRequestTimestamp > DiscardCandidate.CreateTimestamp; // The session was created successfully, a list request posted, but the session did not turn up. The session was deleted before it could be listed.
+						return bDiscovered || bDeleted;
+					});
+				}
+			}
+			// else -> The concert request failed, possibly because the server went offline. Wait until next TickServersAndSessionsDiscovery() to sync the server again and discover if it went offline.
+		});
+}
+
+void FConcertClientSessionBrowserController::UpdateArchivedSessionsAsync(const FConcertServerInfo& ServerInfo)
+{
+	// Check if a request is already pulling archived sessions from this server.
+	FAsyncRequest& ListArchivedSessionAsyncRequest = ArchivedSessionRequests.FindOrAdd(ServerInfo.InstanceInfo.InstanceId);
+	if (ListArchivedSessionAsyncRequest.IsOngoing())
+	{
+		// Don't stack another async request on top of the ongoing one, TickServersAndSessionsDiscovery() will eventually catch up any possibly missed info here.
+		return;
+	}
+
+	// Arm the future, enabling its continuation to execute.
+	TWeakPtr<uint8, ESPMode::ThreadSafe> ListArchivedSessionExecutionToken = ListArchivedSessionAsyncRequest.ResetExecutionToken();
+
+	// Retrieve the archived sessions.
+	ListArchivedSessionAsyncRequest.Future = ConcertClient->GetArchivedSessions(ServerInfo.AdminEndpointId)
+		.Next([this, ServerInfo, ListArchivedSessionExecutionToken](const FConcertAdmin_GetSessionsResponse& Response)
+		{
+			// If the future is disarmed.
+			if (!ListArchivedSessionExecutionToken.IsValid())
+			{
+				// Don't go further, the future execution was canceled, maybe because this object was deleted, the server removed from the list or it wasn't safe to get this future executing.
+				return;
+			}
+
+			// If the server responded.
+			if (Response.ResponseCode == EConcertResponseCode::Success)
+			{
+				bInitialArchivedSessionQueryResponded = true;
+
+				// Remove from the cache archives that were deleted from 'Server' since the last update. Find which one were renamed.
+				ArchivedSessions.RemoveAll([this, &ServerInfo, &Response](const TSharedPtr<FClientArchivedSessionInfo>& DiscardCandidate)
+				{
+					// If the discard candidate is stored on another server.
+					if (DiscardCandidate->ServerInfo.InstanceInfo.InstanceId != ServerInfo.InstanceInfo.InstanceId)
+					{
+						return false; // Keep that archive, it's stored on another server.
+					}
+					// If the archive is still stored on the server.
+					else if (const FConcertSessionInfo* SessionFromServer = Response.Sessions.FindByPredicate([DiscardCandidate](const FConcertSessionInfo& MatchCandidate) { return DiscardCandidate->SessionInfo.SessionId == MatchCandidate.SessionId; }))
+					{
+						// Don't discard the session, it is still there. Check if it was renamed.
+						if (SessionFromServer->SessionName != DiscardCandidate->SessionInfo.SessionName)
+						{
+							DiscardCandidate->SessionInfo.SessionName = SessionFromServer->SessionName; // Update the session name.
+							OnArchivedSessionRenamed(*DiscardCandidate, SessionFromServer->SessionName);
+						}
+						else if (DiscardCandidate->bSessionNameDirty) // Renaming this particular session failed?
+						{
+							OnArchivedSessionListDirty(); // UI optimistically updated a session name from a rename, but it failed on the server. The controller caches the latest official server list unmodified. Force the client to resync its list against the cache.
+							DiscardCandidate->bSessionNameDirty = false;
+						}
+
+						return false; // Don't remove the archive, it's still stored on the server.
+					}
+
+					OnArchivedSessionDiscarded(*DiscardCandidate);
+					return true; // The session is not archived anymore on 'Server' anymore, remove it from the list.
+				});
+
+				// Add to the cache archives that was stored on 'Server' since the last update.
+				for (const FConcertSessionInfo& SessionInfo : Response.Sessions)
+				{
+					// Try to find the archived in the list.
+					if (!ArchivedSessions.ContainsByPredicate([&SessionInfo, &ServerInfo](const TSharedPtr<FArchivedSessionInfo>& MatchCandidate)
+						{ return ServerInfo.InstanceInfo.InstanceId == MatchCandidate->ServerInfo.InstanceInfo.InstanceId && SessionInfo.SessionId == MatchCandidate->SessionInfo.SessionId; }))
+					{
+						// This is a newly discovered archive, add it to the list.
+						ArchivedSessions.Add(MakeShared<FClientArchivedSessionInfo>(ServerInfo, SessionInfo));
+						OnArchivedSessionDiscovered(*ArchivedSessions.Last());
+					}
+				}
+			}
+			// else -> Request failed, will discovered what happened on next TickServersAndSessionsDiscovery().
+		});
+}
+
+void FConcertClientSessionBrowserController::UpdateClientsAsync(const FGuid& ServerAdminEndpointId, const FGuid& SessionId)
+{
+	// The clients request is for a different session than the one cached by the last client update (The UI only shows 1 at the time).
+	if (ClientMonitoredSession.IsSet() && (ClientMonitoredSession->Session->ServerInfo.AdminEndpointId != ServerAdminEndpointId || ClientMonitoredSession->Session->SessionInfo.SessionId != SessionId))
+	{
+		// Reset the cache.
+		ClientMonitoredSession->Session->Clients.Reset();
+		ClientMonitoredSession->ListClientRequest.Cancel();
+		ClientMonitoredSession.Reset();
+		ClientListVersion = 0;
+	}
+
+	// Find the session requested by the user.
+	if (!ClientMonitoredSession.IsSet())
+	{
+		TSharedPtr<FClientActiveSessionInfo>* MatchEntry = ActiveSessions.FindByPredicate([&ServerAdminEndpointId, &SessionId](const TSharedPtr<FClientActiveSessionInfo>& MatchCandidate)
+		{
+			return MatchCandidate->ServerInfo.AdminEndpointId == ServerAdminEndpointId && MatchCandidate->SessionInfo.SessionId == SessionId;
+		});
+
+		if (MatchEntry != nullptr)
+		{
+			ClientMonitoredSession = FMonitoredSessionInfo { MatchEntry->ToSharedRef() };
+		}
+		else
+		{
+			ClientMonitoredSession.Reset();
+		}
+	}
+
+	if (!ClientMonitoredSession.IsSet() || ClientMonitoredSession->ListClientRequest.IsOngoing())
+	{
+		// Don't stack another async request on top of the on-going one, TickClientsDiscovery() will eventually catch up any possibly missed info here.
+		return;
+	}
+
+	// Arm the future, enabling its continuation to execute when called.
+	TWeakPtr<uint8, ESPMode::ThreadSafe> ListClientExecutionToken = ClientMonitoredSession->ListClientRequest.ResetExecutionToken();
+
+	// Retrieve (asynchronously) the clients corresponding to the selected server/session.
+	ClientMonitoredSession->ListClientRequest.Future = ConcertClient->GetSessionClients(ServerAdminEndpointId, SessionId)
+		.Next([this, ListClientExecutionToken](const FConcertAdmin_GetSessionClientsResponse& Response)
+		{
+			// If the future execution was canceled.
+			if (!ListClientExecutionToken.IsValid())
+			{
+				// Don't go further, the future execution was canceled.
+				return;
+			}
+
+			auto SortClientPredicate = [](const FConcertSessionClientInfo& Lhs, const FConcertSessionClientInfo& Rhs) { return Lhs.ClientEndpointId < Rhs.ClientEndpointId; };
+
+			// If the request succeeded.
+			if (Response.ResponseCode == EConcertResponseCode::Success)
+			{
+				if (ClientMonitoredSession->Session->Clients.Num() != Response.SessionClients.Num())
+				{
+					ClientMonitoredSession->Session->Clients = Response.SessionClients;
+					ClientMonitoredSession->Session->Clients.Sort(SortClientPredicate);
+					OnActiveSessionClientsUpdated(*ClientMonitoredSession->Session);
+					return;
+				}
+				else if (Response.SessionClients.Num() == 0) // All existing clients disconnected.
+				{
+					ClientMonitoredSession->Session->Clients.Reset();
+					OnActiveSessionClientsUpdated(*ClientMonitoredSession->Session);
+				}
+				else // Compare the old and the new list, both sorted by client endpoint id.
+				{
+					TArray<FConcertSessionClientInfo> SortedClients = Response.SessionClients;
+					SortedClients.Sort(SortClientPredicate);
+
+					int Index = 0;
+					for (const FConcertSessionClientInfo& Client : SortedClients)
+					{
+						if (ClientMonitoredSession->Session->Clients[Index].ClientEndpointId != Client.ClientEndpointId || // Not the same client?
+						    ClientMonitoredSession->Session->Clients[Index].ClientInfo != Client.ClientInfo)               // Client info was updated?
+						{
+							// The two lists are not identical, don't bother finding the 'delta', refresh all clients.
+							ClientMonitoredSession->Session->Clients = SortedClients;
+							OnActiveSessionClientsUpdated(*ClientMonitoredSession->Session);
+							break;
+						}
+						++Index;
+					}
+				}
+			}
+		});
+}
+
+#undef LOCTEXT_NAMESPACE
\ No newline at end of file
diff --git a/Engine/Plugins/Developer/Concert/ConcertApp/MultiUserClient/Source/MultiUserClient/Private/Widgets/Disconnected/ConcertClientSessionBrowserController.h b/Engine/Plugins/Developer/Concert/ConcertApp/MultiUserClient/Source/MultiUserClient/Private/Widgets/Disconnected/ConcertClientSessionBrowserController.h
new file mode 100644
--- /dev/null
+++ b/Engine/Plugins/Developer/Concert/ConcertApp/MultiUserClient/Source/MultiUserClient/Private/Widgets/Disconnected/ConcertClientSessionBrowserController.h
@@ -0,0 +1,196 @@
+// Copyright Epic Games, Inc. All Rights Reserved.
+
+#pragma once
+
+#include "IConcertClientModule.h"
+#include "Session/Browser/IConcertSessionBrowserController.h"
+
+#include "Async/Future.h"
+#include "Containers/Array.h"
+#include "Containers/Map.h"
+#include "Containers/Set.h"
+#include "Containers/UnrealString.h"
+#include "HAL/Platform.h"
+#include "Misc/DateTime.h"
+#include "Misc/Guid.h"
+#include "Misc/Optional.h"
+#include "Templates/SharedPointer.h"
+
+/**
+ * Runs and cache network queries for the UI. In the model-view-controller pattern, this class acts like the controller. Its purpose
+ * is to keep the UI code as decoupled as possible from the API used to query it. It encapsulate the asynchronous code and provide a
+ * simpler API to the UI.
+ */
+class FConcertClientSessionBrowserController : public IConcertSessionBrowserController
+{
+public:
+	
+	/** Keeps the state of an active async request and provides a tool to cancel its future continuation execution. */
+	struct FAsyncRequest
+	{
+		/** Returns true if there is a registered async request future and if it hasn't executed yet. */
+		bool IsOngoing() const { return Future.IsValid() && !Future.IsReady(); }
+
+		/** Reset the execution token, canceling previous execution (if any) and setting up the token for a new request. */
+		TWeakPtr<uint8, ESPMode::ThreadSafe> ResetExecutionToken() { FutureExecutionToken = MakeShared<uint8, ESPMode::ThreadSafe>(); return FutureExecutionToken; }
+
+		/** Cancel the execution of request async continuation. */
+		void Cancel() { FutureExecutionToken.Reset(); }
+
+		/** The future provided by an asynchronous request. */
+		TFuture<void> Future;
+
+		/** Determines whether or not the async request continuation code should execute. Reset to disarm execution of an async future continuation. */
+		TSharedPtr<uint8, ESPMode::ThreadSafe> FutureExecutionToken;
+	};
+	
+	FConcertClientSessionBrowserController(IConcertClientPtr InConcertClient);
+	virtual ~FConcertClientSessionBrowserController() override;
+
+	/**
+	 * Fires new requests to retrieve all known server and for each server, their active and archived sessions. The responses are
+	 * received asynchronously and may not be available right now. When a response is received, if the corresponding list cached
+	 * is updated, the list version is incremented.
+	 *
+	 * @return A (serverListVersion, sessionsListVersion) pair, corresponding the the versions currently cached by this object.
+	 */
+	TPair<uint32, uint32> TickServersAndSessionsDiscovery();
+
+	/**
+	 * Fires a new request to retrieve the clients for the selected session. The result is cached and can be retrived by GetClients().
+	 * The class caches clients for a single active session determined by this function. If the specified session changes, the cache
+	 * of the previous session is cleared.
+	 *
+	 * @return The version of the client list currently cached by the object.
+	 */
+	uint32 TickClientsDiscovery(const FGuid& ServerAdminEndpointId, const FGuid& SessionId);
+
+	/** Returns true if the controller received async responses and updated its cache since the last time the function was called, then clear the flag. */
+	bool GetAndClearDiscoveryUpdateFlag();
+	
+	/** Returns the latest list of clients corresponding to the session known to this controller. Ensure to call TickClientsDiscovery() periodically. */
+	const TArray<FConcertSessionClientInfo>& GetClients(const FGuid& ServerAdminEndpointId, const FGuid& SessionId) const;
+	
+	void JoinSession(const FGuid& ServerAdminEndpointId, const FGuid& SessionId);
+	
+	// The 2 functions below are used to prevent fast UI transition with the 'no session' panel flashing when the session list is empty.
+	bool HasReceivedInitialSessionList() const { return bInitialActiveSessionQueryResponded && bInitialArchivedSessionQueryResponded; }
+	bool IsCreatingSession() const { return CreateSessionRequests.Num() > 0 || ExpectedSessionsToDiscover.Num() > 0; }
+	
+	IConcertClientPtr GetConcertClient() const { return ConcertClient; }
+	
+	//~ Begin IConcertSessionBrowserController Interface
+	virtual TArray<FConcertServerInfo> GetServers() const override { return Servers; }
+
+	virtual TArray<FActiveSessionInfo> GetActiveSessions() const override
+	{
+		TArray<FActiveSessionInfo> Result;
+		Algo::Transform(ActiveSessions, Result, [](TSharedPtr<FClientActiveSessionInfo> Info) { return *Info; });
+		return Result;
+	}
+
+	virtual TArray<FArchivedSessionInfo> GetArchivedSessions() const override
+	{
+		TArray<FArchivedSessionInfo> Result;
+		Algo::Transform(ArchivedSessions, Result, [](TSharedPtr<FClientArchivedSessionInfo> Info) { return *Info; });
+		return Result;
+	}
+
+	virtual TOptional<FConcertSessionInfo> GetActiveSessionInfo(const FGuid& AdminEndpoint, const FGuid& SessionId) const override;
+	virtual TOptional<FConcertSessionInfo> GetArchivedSessionInfo(const FGuid& AdminEndpoint, const FGuid& SessionId) const override;
+
+	virtual void CreateSession(const FGuid& ServerAdminEndpointId, const FString& SessionName, const FString& ProjectName) override;
+	virtual void ArchiveSession(const FGuid& ServerAdminEndpointId, const FGuid& SessionId, const FString& ArchiveName, const FConcertSessionFilter& SessionFilter) override;
+	virtual void RestoreSession(const FGuid& ServerAdminEndpointId, const FGuid& SessionId, const FString& RestoredName, const FConcertSessionFilter& SessionFilter) override;
+	virtual void RenameActiveSession(const FGuid& ServerAdminEndpointId, const FGuid& SessionId, const FString& NewName) override;
+	virtual void RenameArchivedSession(const FGuid& ServerAdminEndpointId, const FGuid& SessionId, const FString& NewName) override;
+	virtual bool CanRenameActiveSession(const FGuid& ServerAdminEndpointId, const FGuid& SessionId) const override;
+	virtual bool CanRenameArchivedSession(const FGuid& ServerAdminEndpointId, const FGuid& SessionId) const override;
+	virtual void DeleteSessions(const FGuid& ServerAdminEndpointId, const TArray<FGuid>& SessionIds) override;
+	virtual bool CanDeleteActiveSession(const FGuid& ServerAdminEndpointId, const FGuid& SessionId) const override;
+	virtual bool CanDeleteArchivedSession(const FGuid& ServerAdminEndpointId, const FGuid& SessionId) const override;
+	virtual bool CanEverCreateSessions() const override { return true; }
+	//~ End IConcertSessionBrowserController Interface
+
+private:
+	
+	/** Hold information about a session created by this client, not yet 'discovered' by a 'list session' query, but expected to be soon. */
+	struct FPendingSessionDiscovery
+	{
+		FDateTime CreateTimestamp;
+		FGuid ServerEndpoint;
+		FString SessionName;
+	};
+
+	struct FMonitoredSessionInfo
+	{
+		TSharedRef<FActiveSessionInfo> Session;
+		FAsyncRequest ListClientRequest;
+
+		FMonitoredSessionInfo(TSharedRef<FActiveSessionInfo> Session)
+			: Session(Session)
+		{}
+	};
+
+	struct FClientActiveSessionInfo : public FActiveSessionInfo
+	{
+		/** Raised when the UI and the cache values may be out of sync if a rename failed (UI assumed it succeeded) */
+		bool bSessionNameDirty = false;
+
+		FClientActiveSessionInfo(FConcertServerInfo ServerInfo, FConcertSessionInfo SessionInfo)
+			: FActiveSessionInfo(MoveTemp(ServerInfo), MoveTemp(SessionInfo))
+		{}
+	};
+
+	struct FClientArchivedSessionInfo : public FArchivedSessionInfo
+	{
+		/** Raised when the UI and the cache values may be out of sync if a rename failed (UI assumed it succeeded) */
+		bool bSessionNameDirty = false;
+
+		FClientArchivedSessionInfo(FConcertServerInfo ServerInfo, FConcertSessionInfo SessionInfo)
+			: FArchivedSessionInfo(MoveTemp(ServerInfo), MoveTemp(SessionInfo))
+		{}
+	};
+
+	void UpdateSessionsAsync();
+	void UpdateActiveSessionsAsync(const FConcertServerInfo& ServerInfo);
+	void UpdateArchivedSessionsAsync(const FConcertServerInfo& ServerInfo);
+	void UpdateClientsAsync(const FGuid& ServerAdminEndpointId, const FGuid& SessionId);
+
+	void OnActiveSessionDiscovered(const FActiveSessionInfo&) { ++SessionListVersion; bCacheUpdated = true; }
+	void OnActiveSessionDiscarded(const FActiveSessionInfo&) { ++SessionListVersion; bCacheUpdated = true; }
+	void OnArchivedSessionDiscovered(const FArchivedSessionInfo&) { ++SessionListVersion; bCacheUpdated = true; }
+	void OnArchivedSessionDiscarded(const FArchivedSessionInfo&) { ++SessionListVersion; bCacheUpdated = true; }
+	void OnActiveSessionClientsUpdated(const FActiveSessionInfo&) { ++ClientListVersion; bCacheUpdated = true; }
+	void OnActiveSessionRenamed(const FActiveSessionInfo&, const FString& NewName) { ++SessionListVersion; bCacheUpdated = true; }
+	void OnArchivedSessionRenamed(const FArchivedSessionInfo&, const FString& NewName) { ++SessionListVersion; bCacheUpdated = true; }
+	void OnActiveSessionListDirty() { ++SessionListVersion; bCacheUpdated = true; } // This will force the UI to refresh its list.
+	void OnArchivedSessionListDirty() { ++SessionListVersion; bCacheUpdated = true; } // This will force the UI to refresh its list.
+
+	// Holds a concert client instance.
+	IConcertClientPtr ConcertClient;
+
+	// The list of active/archived async requests (requesting the list of session) per server. There is only one per server as we prevent stacking more than one at the time.
+	TMap<FGuid, FAsyncRequest> ActiveSessionRequests;
+	TMap<FGuid, FAsyncRequest> ArchivedSessionRequests;
+
+	// The cached lists.
+	TArray<FConcertServerInfo> Servers;
+	TArray<TSharedPtr<FClientActiveSessionInfo>> ActiveSessions;
+	TArray<TSharedPtr<FClientArchivedSessionInfo>> ArchivedSessions;
+
+	// The session for which the clients are monitored. UI only monitor client of 1 session at the time.
+	TOptional<FMonitoredSessionInfo> ClientMonitoredSession;
+
+	// Holds the version of data cached by the controller. The version is updated when an async response is received and implies a change in the cached values.
+	uint32 ServerListVersion = 0;
+	uint32 SessionListVersion = 0;
+	uint32 ClientListVersion = 0;
+	bool bCacheUpdated = false;
+	bool bInitialActiveSessionQueryResponded = false;
+	bool bInitialArchivedSessionQueryResponded = false;
+
+	TArray<FAsyncRequest> CreateSessionRequests;
+	TArray<FPendingSessionDiscovery> ExpectedSessionsToDiscover;
+	TSet<FString> IgnoredServers; // List of ignored servers (Useful for testing/debugging)
+};
diff --git a/Engine/Plugins/Developer/Concert/ConcertApp/MultiUserClient/Source/MultiUserClient/Private/Widgets/Disconnected/SConcertClientSessionBrowser.cpp b/Engine/Plugins/Developer/Concert/ConcertApp/MultiUserClient/Source/MultiUserClient/Private/Widgets/Disconnected/SConcertClientSessionBrowser.cpp
new file mode 100644
--- /dev/null
+++ b/Engine/Plugins/Developer/Concert/ConcertApp/MultiUserClient/Source/MultiUserClient/Private/Widgets/Disconnected/SConcertClientSessionBrowser.cpp
@@ -0,0 +1,793 @@
+
+#include "SConcertClientSessionBrowser.h"
+
+#include "ConcertClientSettings.h"
+#include "ConcertFrontendStyle.h"
+#include "ConcertFrontendUtils.h"
+#include "ConcertMessageData.h"
+#include "ConcertSettings.h"
+#include "ConcertVersion.h"
+#include "IMultiUserClientModule.h"
+#include "IConcertClient.h"
+#include "IConcertSyncClient.h"
+#include "SConcertDiscovery.h"
+#include "SConcertNoAvailability.h"
+#include "Session/Browser/ConcertBrowserUtils.h"
+#include "Session/Browser/Items/ConcertSessionTreeItem.h"
+#include "Session/Browser/SConcertSessionBrowser.h"
+
+#include "Algo/Transform.h"
+#include "Algo/TiedTupleOutput.h"
+#include "Framework/Commands/UIAction.h"
+#include "Framework/Commands/UICommandInfo.h"
+#include "Framework/MultiBox/MultiBoxBuilder.h"
+#include "Misc/MessageDialog.h"
+#include "Styling/AppStyle.h"
+#include "Textures/SlateIcon.h"
+#include "Widgets/Disconnected/ConcertClientSessionBrowserController.h"
+#include "Widgets/Images/SImage.h"
+#include "Widgets/Input/SButton.h"
+#include "Widgets/Layout/SBorder.h"
+#include "Widgets/Layout/SBox.h"
+#include "Widgets/Layout/SExpandableArea.h"
+#include "Widgets/Layout/SGridPanel.h"
+#include "Widgets/Layout/SScrollBox.h"
+#include "Widgets/Text/STextBlock.h"
+#include "Widgets/Views/SListView.h"
+#include "Widgets/Views/STableRow.h"
+
+#if WITH_EDITOR
+#include "ISettingsModule.h"
+#endif
+
+#define LOCTEXT_NAMESPACE "SConcertClientSessionBrowser"
+
+void SConcertClientSessionBrowser::Construct(const FArguments& InArgs, IConcertClientPtr InConcertClient, TSharedPtr<FText> InSearchText)
+{
+	if (!InConcertClient.IsValid())
+	{
+		return; // Don't build the UI if ConcertClient is not available.
+	}
+
+	Controller = MakeShared<FConcertClientSessionBrowserController>(InConcertClient);
+
+	// Displayed if no server is available.
+	ServerDiscoveryPanel = SNew(SConcertDiscovery)
+		.Text(LOCTEXT("LookingForServer", "Looking for Multi-User Servers..."))
+		.Visibility_Lambda([this]() { return Controller->GetServers().Num() == 0 ? EVisibility::Visible : EVisibility::Hidden; })
+		.IsButtonEnabled(this, &SConcertClientSessionBrowser::IsLaunchServerButtonEnabled)
+		.ButtonStyle(FAppStyle::Get(), "FlatButton.Default")
+		.ButtonIcon(FConcertFrontendStyle::Get()->GetBrush("Concert.NewServer.Small"))
+		.ButtonText(LOCTEXT("LaunchLocalServer", "Launch a Server"))
+		.ButtonToolTip(LOCTEXT("LaunchServerTooltip", "Launch a Multi-User server on your computer unless one is already running.\nThe editor UDP messaging settings will be passed to the launching server.\nThe server will use the server port found in Multi-User client settings if non 0 or use the editor udp messaging port number + 1, if non 0."))
+		.OnButtonClicked(this, &SConcertClientSessionBrowser::OnLaunchServerButtonClicked);
+
+	// Controls the text displayed in the 'No sessions' panel.
+	auto GetNoSessionText = [this]()
+	{
+		if (!Controller->HasReceivedInitialSessionList())
+		{
+			return LOCTEXT("LookingForSession", "Looking for Multi-User Sessions...");
+		}
+
+		return Controller->GetActiveSessions().Num() == 0 && Controller->GetArchivedSessions().Num() == 0 ?
+			LOCTEXT("NoSessionAvailable", "No Sessions Available") :
+			LOCTEXT("AllSessionsFilteredOut", "No Sessions Match the Filters\nChange Your Filter to View Sessions");
+	};
+
+	// Displayed when discovering session or if no session is available.
+	SessionDiscoveryPanel = SNew(SConcertDiscovery)
+		.Text_Lambda(GetNoSessionText)
+		.Visibility_Lambda([this]() { return Controller->GetServers().Num() > 0 && !SessionBrowser->HasAnySessions() && !Controller->IsCreatingSession() ? EVisibility::Visible : EVisibility::Hidden; })
+		.ThrobberVisibility_Lambda([this]() { return !Controller->HasReceivedInitialSessionList() ? EVisibility::Visible : EVisibility::Collapsed; })
+		.ButtonVisibility_Lambda([this]() { return Controller->HasReceivedInitialSessionList() && Controller->GetActiveSessions().Num() == 0 && Controller->GetArchivedSessions().Num() == 0 ? EVisibility::Visible : EVisibility::Collapsed; })
+		.ButtonStyle(FAppStyle::Get(), "FlatButton.Default")
+		.ButtonIcon(FConcertFrontendStyle::Get()->GetBrush("Concert.NewSession.Small"))
+		.ButtonText(LOCTEXT("CreateSession", "Create Session"))
+		.ButtonToolTip(LOCTEXT("CreateSessionTooltip", "Create a new session"))
+		.OnButtonClicked(FOnClicked::CreateSP(this, &SConcertClientSessionBrowser::OnNewButtonClicked));
+
+	// Displayed when the selected session client view is empty (no client to display).
+	NoClientPanel = SNew(SConcertNoAvailability)
+		.Text(LOCTEXT("NoClientAvailable", "No Connected Clients"))
+		.Visibility_Lambda([this]() { return Clients.Num() == 0 ? EVisibility::Visible : EVisibility::Hidden; });
+
+	// Displayed as details when no session is selected. (No session selected or the selected session doesn't have any)
+	NoSessionSelectedPanel = SNew(SConcertNoAvailability)
+		.Text(LOCTEXT("NoSessionSelected", "Select a Session to View Details"));
+
+	// Displayed as details when the selected session has not specific details to display.
+	NoSessionDetailsPanel = SNew(SConcertNoAvailability)
+		.Text(LOCTEXT("NoSessionDetails", "The Selected Session Has No Details"));
+
+	// List used in details panel to display clients connected to an active session.
+	ClientsView = SNew(SListView<TSharedPtr<FConcertSessionClientInfo>>)
+		.ListItemsSource(&Clients)
+		.OnGenerateRow(this, &SConcertClientSessionBrowser::OnGenerateClientRowWidget)
+		.SelectionMode(ESelectionMode::Single)
+		.AllowOverscroll(EAllowOverscroll::No);
+
+	ChildSlot
+	[
+		MakeBrowserContent(InSearchText)
+	];
+
+	// Create a timer to periodically poll the server for sessions and session clients at a lower frequency than the normal tick.
+	RegisterActiveTimer(1.0f, FWidgetActiveTimerDelegate::CreateSP(this, &SConcertClientSessionBrowser::TickDiscovery));
+
+	bLocalServerRunning = IMultiUserClientModule::Get().IsConcertServerRunning();
+}
+
+TSharedRef<SWidget> SConcertClientSessionBrowser::MakeBrowserContent(TSharedPtr<FText> InSearchText)
+{
+	return SNew(SBox)
+		.HAlign(HAlign_Fill)
+		.VAlign(VAlign_Fill)
+		[
+			// Splitter upper part displaying the available sessions/(server).
+			SNew(SSplitter)
+			.Orientation(Orient_Vertical)
+			.MinimumSlotHeight(80.0f) // Prevent widgets from overlapping.
+			+SSplitter::Slot()
+			.Value(0.6)
+			[
+				SAssignNew(SessionBrowser, SConcertSessionBrowser, Controller.ToSharedRef(), InSearchText)
+					.ExtendSessionTable(this, &SConcertClientSessionBrowser::MakeOverlayedTableView)
+					.ExtendControllButtons(this, &SConcertClientSessionBrowser::ExtendControlButtons)
+					.ExtendSessionContextMenu(this, &SConcertClientSessionBrowser::ExtendSessionContextMenu)
+					.RightOfControlButtons()
+					[
+						MakeUserAndSettings()
+					]
+					.OnSessionClicked(this, &SConcertClientSessionBrowser::OnSessionSelectionChanged)
+					.OnLiveSessionDoubleClicked(this, &SConcertClientSessionBrowser::OnSessionDoubleClicked)
+					.PostRequestedDeleteSession_Lambda([this](auto) { UpdateDiscovery(); /* Don't wait up to 1s, kick discovery right now */ })
+					.AskUserToDeleteSessions(this, &SConcertClientSessionBrowser::ConfirmDeleteSessionWithDialog)
+			]
+
+			// Session details.
+			+SSplitter::Slot()
+			.Value(0.4)
+			[
+				SAssignNew(SessionDetailsView, SBorder)
+				.BorderImage(FAppStyle::Get().GetBrush("ToolPanel.GroupBorder"))
+				[
+					NoSessionSelectedPanel.ToSharedRef()
+				]
+			]
+		];
+}
+
+void SConcertClientSessionBrowser::ExtendControlButtons(FExtender& Extender)
+{
+	// Before separator
+	Extender.AddToolBarExtension(
+		SConcertSessionBrowser::ControlButtonExtensionHooks::BeforeSeparator,
+		EExtensionHook::First,
+		nullptr,
+		FToolBarExtensionDelegate::CreateLambda([this](FToolBarBuilder& Builder)
+		{
+			Builder.AddWidget(
+				SNew(SOverlay)
+				// Launch server.
+				+SOverlay::Slot()
+				[
+					ConcertBrowserUtils::MakeIconButton(
+						FConcertFrontendStyle::Get()->GetBrush("Concert.NewServer"),
+						LOCTEXT("LaunchServerTooltip", "Launch a Multi-User server on your computer unless one is already running.\nThe editor UDP messaging settings will be passed to the launching server.\nThe server will use the server port found in Multi-User client settings if non 0 or use the editor udp messaging port number + 1, if non 0."),
+						TAttribute<bool>(this, &SConcertClientSessionBrowser::IsLaunchServerButtonEnabled),
+						FOnClicked::CreateSP(this, &SConcertClientSessionBrowser::OnLaunchServerButtonClicked),
+						TAttribute<EVisibility>::Create([this]() { return IsLaunchServerButtonEnabled() ? EVisibility::Visible : EVisibility::Collapsed; }))
+				]
+				// Stop server.
+				+SOverlay::Slot()
+				[
+					ConcertBrowserUtils::MakeIconButton(
+						FConcertFrontendStyle::Get()->GetBrush("Concert.CloseServer"),
+						LOCTEXT("ShutdownServerTooltip", "Shutdown the Multi-User server running on this computer."),
+						true, // Always enabled, but might be collapsed.
+						FOnClicked::CreateSP(this, &SConcertClientSessionBrowser::OnShutdownServerButtonClicked),
+						TAttribute<EVisibility>::Create([this]() { return IsLaunchServerButtonEnabled() ? EVisibility::Collapsed : EVisibility::Visible; }))
+				]
+				);
+		}));
+
+	// After separator
+	Extender.AddToolBarExtension(
+		SConcertSessionBrowser::ControlButtonExtensionHooks::AfterSeparator,
+		EExtensionHook::First,
+		nullptr,
+		FToolBarExtensionDelegate::CreateLambda([this](FToolBarBuilder& Builder)
+		{
+			TAttribute<FText> AutoJoinTooltip = TAttribute<FText>::Create([this]()
+			{
+				if (Controller->GetConcertClient()->CanAutoConnect()) // Default session and server are configured?
+				{
+					return FText::Format(LOCTEXT("JoinDefaultSessionTooltip", "Join the default session '{0}' on '{1}'"),
+						FText::FromString(Controller->GetConcertClient()->GetConfiguration()->DefaultSessionName),
+						FText::FromString(Controller->GetConcertClient()->GetConfiguration()->DefaultServerURL));
+				}
+				else
+				{
+					return LOCTEXT("JoinDefaultSessionConfiguredTooltip", "Join the default session configured in the Multi-Users settings");
+				}
+			});
+
+			TAttribute<FText> CancelAutoJoinTooltip = TAttribute<FText>::Create([this]()
+			{
+				return FText::Format(LOCTEXT("CancelJoinDefaultSessionTooltip", "Cancel joining the default session '{0}' on '{1}'"),
+					FText::FromString(Controller->GetConcertClient()->GetConfiguration()->DefaultSessionName),
+					FText::FromString(Controller->GetConcertClient()->GetConfiguration()->DefaultServerURL));
+			});
+			
+			Builder.AddWidget(
+				SNew(SOverlay)
+				// Auto-join
+				+SOverlay::Slot()
+				[
+					ConcertBrowserUtils::MakeIconButton(
+						FConcertFrontendStyle::Get()->GetBrush("Concert.JoinDefaultSession"),
+						AutoJoinTooltip,
+						TAttribute<bool>(this, &SConcertClientSessionBrowser::IsAutoJoinButtonEnabled),
+						FOnClicked::CreateSP(this, &SConcertClientSessionBrowser::OnAutoJoinButtonClicked),
+						// Default button shown if both auto-join/cancel are disabled.
+						TAttribute<EVisibility>::Create([this]() { return !IsCancelAutoJoinButtonEnabled() ? EVisibility::Visible : EVisibility::Collapsed; })
+						) 
+				
+				]
+				// Cancel auto join.
+				+SOverlay::Slot()
+				[
+					ConcertBrowserUtils::MakeIconButton(
+						FConcertFrontendStyle::Get()->GetBrush("Concert.CancelAutoJoin"),
+						CancelAutoJoinTooltip,
+						TAttribute<bool>(this, &SConcertClientSessionBrowser::IsCancelAutoJoinButtonEnabled),
+						FOnClicked::CreateSP(this, &SConcertClientSessionBrowser::OnCancelAutoJoinButtonClicked),
+						TAttribute<EVisibility>::Create([this]() { return IsCancelAutoJoinButtonEnabled() ? EVisibility::Visible : EVisibility::Collapsed; })
+						)
+				]
+			);
+
+			// Join
+			Builder.AddWidget(
+				ConcertBrowserUtils::MakeIconButton(FConcertFrontendStyle::Get()->GetBrush("Concert.JoinSession"), LOCTEXT("JoinButtonTooltip", "Join the selected session"),
+					TAttribute<bool>(this, &SConcertClientSessionBrowser::IsJoinButtonEnabled),
+					FOnClicked::CreateSP(this, &SConcertClientSessionBrowser::OnJoinButtonClicked),
+					TAttribute<EVisibility>::Create([this]() { return !SessionBrowser->IsRestoreButtonEnabled() ? EVisibility::Visible : EVisibility::Collapsed; })) // Default button shown if both join/restore are disabled.
+			);
+		}));
+}
+
+void SConcertClientSessionBrowser::ExtendSessionContextMenu(const TSharedPtr<FConcertSessionTreeItem>& Item, FExtender& Extender)
+{
+	if (Item->Type == FConcertSessionTreeItem::EType::ActiveSession)
+	{
+		Extender.AddMenuExtension(
+			SConcertSessionBrowser::SessionContextMenuExtensionHooks::ManageSession,
+			EExtensionHook::First,
+			nullptr,
+			FMenuExtensionDelegate::CreateLambda([this](FMenuBuilder& Builder)
+			{
+				Builder.AddMenuEntry(
+					LOCTEXT("CtxMenuJoin", "Join"),
+					LOCTEXT("CtxMenuJoin_Tooltip", "Join the Session"),
+					FSlateIcon(),
+					FUIAction(
+						FExecuteAction::CreateLambda([this](){ OnJoinButtonClicked(); }),
+						FCanExecuteAction::CreateLambda([SelectedCount = SessionBrowser->GetSelectedItems().Num()] { return SelectedCount == 1; }),
+						FIsActionChecked::CreateLambda([this] { return false; })),
+					NAME_None,
+					EUserInterfaceActionType::Button
+					);
+			}));
+	}
+	
+}
+
+TSharedRef<SWidget> SConcertClientSessionBrowser::MakeUserAndSettings()
+{
+	return SNew(SHorizontalBox)
+		// The user "Avatar color" displayed as a small square colored by the user avatar color.
+		+SHorizontalBox::Slot()
+		.VAlign(VAlign_Center)
+		.AutoWidth()
+		[
+			SNew(SImage)
+			.ColorAndOpacity_Lambda([this]() { return Controller->GetConcertClient()->GetClientInfo().AvatarColor; })
+			.Image(FAppStyle::GetBrush("Icons.FilledCircle"))
+		]
+				
+		// The user "Display Name".
+		+SHorizontalBox::Slot()
+		.VAlign(VAlign_Center)
+		.AutoWidth()
+		.Padding(3, 0, 2, 0)
+		[
+			SNew(STextBlock)
+				.Text_Lambda([this]() { return FText::FromString(Controller->GetConcertClient()->GetClientInfo().DisplayName);} )
+		]
+
+		// The "Settings" icons.
+		+SHorizontalBox::Slot()
+		.VAlign(VAlign_Fill)
+		.AutoWidth()
+		[
+			SNew(SButton)
+			.ButtonStyle(FAppStyle::Get(), TEXT("SimpleButton"))
+			.OnClicked_Lambda([](){FModuleManager::GetModulePtr<ISettingsModule>("Settings")->ShowViewer("Project", "Plugins", "Concert"); return FReply::Handled(); })
+			.VAlign(VAlign_Center)
+			.HAlign(HAlign_Fill)
+			[
+				SNew(SImage)
+				.Image(FAppStyle::GetBrush("Icons.Settings"))
+			]
+		];
+}
+
+TSharedRef<SWidget> SConcertClientSessionBrowser::MakeOverlayedTableView(const TSharedRef<SWidget>& SessionTable)
+{
+	return SNew(SOverlay)
+		+SOverlay::Slot()
+		[
+			SessionTable
+		]
+		+SOverlay::Slot()
+		.Padding(0, 20, 0, 0) // To ensure the panel is below the row header.
+		[
+			SessionDiscoveryPanel.ToSharedRef()
+		]
+		+SOverlay::Slot()
+		.Padding(0, 20, 0, 0) // To ensure the panel is below the row header.
+		[
+			ServerDiscoveryPanel.ToSharedRef()
+		];
+}
+
+EActiveTimerReturnType SConcertClientSessionBrowser::TickDiscovery(double InCurrentTime, float InDeltaTime)
+{
+	// Cache the result of this function because it is very expensive. It kills the framerate if polled every frame.
+	bLocalServerRunning = IMultiUserClientModule::Get().IsConcertServerRunning();
+
+	UpdateDiscovery();
+	return EActiveTimerReturnType::Continue;
+}
+
+void SConcertClientSessionBrowser::UpdateDiscovery()
+{
+	// Check if the controller has updated data since the last 'tick' and fire new asynchronous requests for future 'tick'.
+	TPair<uint32, uint32> ServerSessionListVersions = Controller->TickServersAndSessionsDiscovery();
+	ServerListVersion = ServerSessionListVersions.Key;
+
+	if (ServerSessionListVersions.Value != DisplayedSessionListVersion) // Need to refresh the list?
+	{
+		SessionBrowser->RefreshSessionList();
+		DisplayedSessionListVersion = ServerSessionListVersions.Value;
+	}
+
+	// If an active session is selected.
+	TArray<TSharedPtr<FConcertSessionTreeItem>> SelectedSession = SessionBrowser->GetSelectedItems();
+	if (SelectedSession.Num() && SelectedSession[0]->Type == FConcertSessionTreeItem::EType::ActiveSession)
+	{
+		// Ensure to poll its clients.
+		uint32 CachedClientListVersion = Controller->TickClientsDiscovery(SelectedSession[0]->ServerAdminEndpointId, SelectedSession[0]->SessionId);
+		if (CachedClientListVersion != DisplayedClientListVersion) // Need to refresh the list?
+		{
+			RefreshClientList(Controller->GetClients(SelectedSession[0]->ServerAdminEndpointId, SelectedSession[0]->SessionId));
+			DisplayedClientListVersion = CachedClientListVersion;
+		}
+	}
+}
+
+void SConcertClientSessionBrowser::RefreshClientList(const TArray<FConcertSessionClientInfo>& LastestClientList)
+{
+	// Remember which client is selected.
+	TArray<TSharedPtr<FConcertSessionClientInfo>> SelectedItems = ClientsView->GetSelectedItems();
+
+	// Copy the list of clients.
+	TArray<TSharedPtr<FConcertSessionClientInfo>> LatestClientPtrs;
+	Algo::Transform(LastestClientList, LatestClientPtrs, [](const FConcertSessionClientInfo& Client) { return MakeShared<FConcertSessionClientInfo>(Client); });
+
+	// Merge the current list with the new list, removing client that disappeared and adding client that appeared.
+	ConcertFrontendUtils::SyncArraysByPredicate(Clients, MoveTemp(LatestClientPtrs), [](const TSharedPtr<FConcertSessionClientInfo>& ClientToFind)
+	{
+		return [ClientToFind](const TSharedPtr<FConcertSessionClientInfo>& PotentialClientMatch)
+		{
+			return PotentialClientMatch->ClientEndpointId == ClientToFind->ClientEndpointId && PotentialClientMatch->ClientInfo == ClientToFind->ClientInfo;
+		};
+	});
+
+	// Sort the list item alphabetically.
+	Clients.StableSort([](const TSharedPtr<FConcertSessionClientInfo>& Lhs, const TSharedPtr<FConcertSessionClientInfo>& Rhs)
+	{
+		return Lhs->ClientInfo.DisplayName < Rhs->ClientInfo.DisplayName;
+	});
+
+	// Preserve previously selected item (if any).
+	if (SelectedItems.Num())
+	{
+		ClientsView->SetSelection(SelectedItems[0]);
+	}
+	ClientsView->RequestListRefresh();
+}
+
+TSharedRef<SWidget> SConcertClientSessionBrowser::MakeSessionDetails(TSharedPtr<FConcertSessionTreeItem> Item)
+{
+	const FConcertSessionInfo* SessionInfo = nullptr;
+
+	if (Item.IsValid())
+	{
+		if (Item->Type == FConcertSessionTreeItem::EType::ActiveSession || Item->Type == FConcertSessionTreeItem::EType::SaveSession)
+		{
+			return MakeActiveSessionDetails(Item);
+		}
+		else if (Item->Type == FConcertSessionTreeItem::EType::ArchivedSession)
+		{
+			return MakeArchivedSessionDetails(Item);
+		}
+	}
+
+	return NoSessionSelectedPanel.ToSharedRef();
+}
+
+TSharedRef<SWidget> SConcertClientSessionBrowser::MakeActiveSessionDetails(TSharedPtr<FConcertSessionTreeItem> Item)
+{
+	const TOptional<FConcertSessionInfo> SessionInfo = Controller->GetActiveSessionInfo(Item->ServerAdminEndpointId, Item->SessionId);
+	if (!SessionInfo)
+	{
+		return NoSessionSelectedPanel.ToSharedRef();
+	}
+
+	TSharedPtr<SGridPanel> Grid;
+
+	// State variables captured and shared by the different lambda functions below.
+	TSharedPtr<bool> bDetailsAreaExpanded = MakeShared<bool>(false);
+	TSharedPtr<bool> bClientsAreaExpanded = MakeShared<bool>(true);
+
+	auto DetailsAreaSizeRule = [this, bDetailsAreaExpanded]()
+	{
+		return *bDetailsAreaExpanded ? SSplitter::ESizeRule::FractionOfParent : SSplitter::ESizeRule::SizeToContent;
+	};
+
+	auto OnDetailsAreaExpansionChanged = [bDetailsAreaExpanded](bool bExpanded)
+	{
+		*bDetailsAreaExpanded = bExpanded;
+	};
+
+	auto ClientsAreaSizeRule = [this, bClientsAreaExpanded]()
+	{
+		return *bClientsAreaExpanded ? SSplitter::ESizeRule::FractionOfParent : SSplitter::ESizeRule::SizeToContent;
+	};
+
+	auto OnClientsAreaExpansionChanged = [bClientsAreaExpanded](bool bExpanded)
+	{
+		*bClientsAreaExpanded = bExpanded;
+	};
+
+	TSharedRef<SSplitter> Widget = SNew(SSplitter)
+		.Orientation(Orient_Vertical)
+
+		// Details.
+		+SSplitter::Slot()
+		.SizeRule(TAttribute<SSplitter::ESizeRule>::Create(DetailsAreaSizeRule))
+		.Value(0.6)
+		[
+			SAssignNew(DetailsArea, SExpandableArea)
+			.BorderBackgroundColor(FLinearColor(0.6f, 0.6f, 0.6f, 1.0f))
+			.BorderImage_Lambda([this]() { return ConcertFrontendUtils::GetExpandableAreaBorderImage(*DetailsArea); })
+			.BodyBorderImage(FAppStyle::Get().GetBrush("ToolPanel.GroupBorder"))
+			.BodyBorderBackgroundColor(FLinearColor::White)
+			.OnAreaExpansionChanged(FOnBooleanValueChanged::CreateLambda(OnDetailsAreaExpansionChanged))
+			.InitiallyCollapsed(!(*bDetailsAreaExpanded))
+			.HeaderContent()
+			[
+				SNew(STextBlock)
+				.Text(FText::Format(LOCTEXT("Details", "Details"), FText::FromString(Item->SessionName)))
+				.Font(FAppStyle::Get().GetFontStyle("DetailsView.CategoryFontStyle"))
+				.ShadowOffset(FVector2D(1.0f, 1.0f))
+			]
+			.BodyContent()
+			[
+				SNew(SScrollBox)
+				.Orientation(Orient_Vertical)
+				+SScrollBox::Slot()
+				[
+					SNew(SBox)
+					.Padding(FMargin(0, 2, 0, 2))
+					[
+						SAssignNew(Grid, SGridPanel)
+					]
+				]
+			]
+		]
+
+		// Clients
+		+SSplitter::Slot()
+		.SizeRule(TAttribute<SSplitter::ESizeRule>::Create(ClientsAreaSizeRule))
+		.Value(0.4)
+		[
+			SAssignNew(ClientsArea, SExpandableArea)
+			.BorderBackgroundColor(FLinearColor(0.6f, 0.6f, 0.6f, 1.0f))
+			.BorderImage_Lambda([this]() { return ConcertFrontendUtils::GetExpandableAreaBorderImage(*ClientsArea); })
+			.BodyBorderImage(FAppStyle::Get().GetBrush("ToolPanel.GroupBorder"))
+			.BodyBorderBackgroundColor(FLinearColor::White)
+			.OnAreaExpansionChanged(FOnBooleanValueChanged::CreateLambda(OnClientsAreaExpansionChanged))
+			.InitiallyCollapsed(!(*bClientsAreaExpanded))
+			.HeaderContent()
+			[
+				SNew(STextBlock)
+				.Text(LOCTEXT("Clients", "Clients"))
+				.Font(FAppStyle::Get().GetFontStyle("DetailsView.CategoryFontStyle"))
+				.ShadowOffset(FVector2D(1.0f, 1.0f))
+			]
+			.BodyContent()
+			[
+				SNew(SOverlay)
+				+SOverlay::Slot()
+				[
+					ClientsView.ToSharedRef()
+				]
+				+SOverlay::Slot()
+				[
+					NoClientPanel.ToSharedRef()
+				]
+			]
+		];
+
+	// Fill the grid.
+	PopulateSessionInfoGrid(*Grid, *SessionInfo);
+
+	// Populate the client list.
+	RefreshClientList(Controller->GetClients(Item->ServerAdminEndpointId, Item->SessionId));
+
+	return Widget;
+}
+
+TSharedRef<SWidget> SConcertClientSessionBrowser::MakeArchivedSessionDetails(TSharedPtr<FConcertSessionTreeItem> Item)
+{
+	const TOptional<FConcertSessionInfo> SessionInfo = Controller->GetArchivedSessionInfo(Item->ServerAdminEndpointId, Item->SessionId);
+	if (!SessionInfo)
+	{
+		return NoSessionSelectedPanel.ToSharedRef();
+	}
+
+	TSharedPtr<SGridPanel> Grid;
+
+	TSharedRef<SExpandableArea> Widget = SAssignNew(DetailsArea, SExpandableArea)
+	.BorderBackgroundColor(FLinearColor(0.6f, 0.6f, 0.6f, 1.0f))
+	.BorderImage_Lambda([this]() { return ConcertFrontendUtils::GetExpandableAreaBorderImage(*DetailsArea); })
+	.BodyBorderImage(FAppStyle::Get().GetBrush("ToolPanel.GroupBorder"))
+	.BodyBorderBackgroundColor(FLinearColor::White)
+	.InitiallyCollapsed(true)
+	.HeaderContent()
+	[
+		SNew(STextBlock)
+		.Text(FText::Format(LOCTEXT("Details", "Details"), FText::FromString(Item->SessionName)))
+		.Font(FAppStyle::Get().GetFontStyle("DetailsView.CategoryFontStyle"))
+		.ShadowOffset(FVector2D(1.0f, 1.0f))
+	]
+	.BodyContent()
+	[
+		SNew(SScrollBox)
+		.Orientation(Orient_Vertical)
+		+SScrollBox::Slot()
+		[
+			SNew(SBox)
+			.Padding(FMargin(0, 2, 0, 2))
+			[
+				SAssignNew(Grid, SGridPanel)
+			]
+		]
+	];
+
+	// Fill the grid.
+	PopulateSessionInfoGrid(*Grid, *SessionInfo);
+
+	return Widget;
+}
+
+void SConcertClientSessionBrowser::PopulateSessionInfoGrid(SGridPanel& Grid, const FConcertSessionInfo& SessionInfo)
+{
+	// Local function to populate the details grid.
+	auto AddDetailRow = [](SGridPanel& Grid, int32 Row, const FText& Label, const FText& Value)
+	{
+		const float RowPadding = (Row == 0) ? 0.0f : 4.0f; // Space between line.
+		const float ColPadding = 4.0f; // Space between columns. (Minimum)
+
+		Grid.AddSlot(0, Row)
+		.Padding(0.0f, RowPadding, ColPadding, 0.0f)
+		[
+			SNew(STextBlock).Text(Label)
+		];
+
+		Grid.AddSlot(1, Row)
+		.Padding(0.0f, RowPadding, 0.0f, 0.0f)
+		[
+			SNew(STextBlock).Text(Value)
+		];
+	};
+
+	int32 Row = 0;
+	AddDetailRow(Grid, Row++, LOCTEXT("SessionId", "Session ID:"), FText::FromString(SessionInfo.SessionId.ToString()));
+	AddDetailRow(Grid, Row++, LOCTEXT("SessionName", "Session Name:"), FText::FromString(SessionInfo.SessionName));
+	AddDetailRow(Grid, Row++, LOCTEXT("Owner", "Owner:"), FText::FromString(SessionInfo.OwnerUserName));
+	AddDetailRow(Grid, Row++, LOCTEXT("Project", "Project:"), FText::FromString(SessionInfo.Settings.ProjectName));
+	if (SessionInfo.VersionInfos.Num() > 0)
+	{
+		const FConcertSessionVersionInfo& VersionInfo = SessionInfo.VersionInfos.Last();
+		AddDetailRow(Grid, Row++, LOCTEXT("EngineVersion", "Engine Version:"), VersionInfo.AsText());
+	}
+	AddDetailRow(Grid, Row++, LOCTEXT("ServerEndPointId", "Server Endpoint ID:"), FText::FromString(SessionInfo.ServerEndpointId.ToString()));
+}
+
+TSharedRef<ITableRow> SConcertClientSessionBrowser::OnGenerateClientRowWidget(TSharedPtr<FConcertSessionClientInfo> Item, const TSharedRef<STableViewBase>& OwnerTable)
+{
+	return SNew(STableRow<TSharedPtr<FConcertSessionClientInfo>>, OwnerTable)
+	[
+		SNew(SHorizontalBox)
+		.ToolTipText(Item->ToDisplayString())
+
+		// The user "Avatar color" displayed as a small square colored by the user avatar color.
+		+SHorizontalBox::Slot()
+		.VAlign(VAlign_Center)
+		.AutoWidth()
+		[
+			SNew(SImage)
+			.ColorAndOpacity(Item->ClientInfo.AvatarColor)
+			.Image(FAppStyle::GetBrush("Icons.FilledCircle"))
+		]
+
+		// The user "Display Name".
+		+SHorizontalBox::Slot()
+		.VAlign(VAlign_Center)
+		.AutoWidth()
+		.Padding(FMargin(4.0, 2.0))
+		[
+			SNew(STextBlock)
+			.Text(FText::FromString(Item->ClientInfo.DisplayName))
+		]
+	];
+}
+
+void SConcertClientSessionBrowser::OnSessionSelectionChanged(const TSharedPtr<FConcertSessionTreeItem>& SelectedSession)
+{
+	// Clear the list of clients (if any)
+	Clients.Reset();
+
+	// Update the details panel.
+	SessionDetailsView->SetContent(MakeSessionDetails(SelectedSession));
+}
+
+void SConcertClientSessionBrowser::OnSessionDoubleClicked(const TSharedPtr<FConcertSessionTreeItem>& SelectedSession)
+{
+	switch (SelectedSession->Type)
+	{
+	case FConcertSessionTreeItem::EType::ActiveSession: 
+		RequestJoinSession(SelectedSession);
+		break;
+	case FConcertSessionTreeItem::EType::ArchivedSession:
+		SessionBrowser->InsertRestoreSessionAsEditableRow(SelectedSession);
+		break;
+	default: 
+		checkNoEntry();
+	}
+}
+
+bool SConcertClientSessionBrowser::ConfirmDeleteSessionWithDialog(const TArray<TSharedPtr<FConcertSessionTreeItem>>& SessionItems) const
+{
+	TSet<FString> SessionNames;
+	TSet<FString> UniqueServers;
+
+	Algo::Transform(SessionItems, Algo::TieTupleAdd(SessionNames, UniqueServers), [](const TSharedPtr<FConcertSessionTreeItem>& Item)
+	{
+		return MakeTuple(Item->SessionName, Item->ServerName);
+	});
+	
+	const FText ConfirmationMessage = FText::Format(
+		LOCTEXT("DeleteSessionConfirmationMessage", "Do you really want to delete {0} {0}|plural(one=session,other=sessions) from {1} {1}|plural(one=server,other=servers)?\n\nThe {0}|plural(one=session,other=sessions) {2} will be deleted from the {1}|plural(one=server,other=servers) {3}."),
+		SessionItems.Num(),
+		UniqueServers.Num(),
+		FText::FromString(FString::JoinBy(SessionNames, TEXT(", "), [](const FString& Name) { return FString("\n\t- ") + Name; }) + FString("\n")),
+		FText::FromString(FString::Join(UniqueServers, TEXT(", ")))
+		);
+	const FText ConfirmationTitle = LOCTEXT("DeleteSessionConfirmationTitle", "Delete Session Confirmation");
+	return FMessageDialog::Open(EAppMsgType::YesNo, ConfirmationMessage, ConfirmationTitle) == EAppReturnType::Yes;
+}
+
+bool SConcertClientSessionBrowser::IsLaunchServerButtonEnabled() const
+{
+	return !bLocalServerRunning;
+}
+
+bool SConcertClientSessionBrowser::IsJoinButtonEnabled() const
+{
+	TArray<TSharedPtr<FConcertSessionTreeItem>> SelectedItems = SessionBrowser->GetSelectedItems();
+	return SelectedItems.Num() == 1 && SelectedItems[0]->Type == FConcertSessionTreeItem::EType::ActiveSession;
+}
+
+bool SConcertClientSessionBrowser::IsAutoJoinButtonEnabled() const
+{
+	return Controller->GetConcertClient()->CanAutoConnect() && !Controller->GetConcertClient()->IsAutoConnecting();
+}
+
+bool SConcertClientSessionBrowser::IsCancelAutoJoinButtonEnabled() const
+{
+	return Controller->GetConcertClient()->IsAutoConnecting();
+}
+
+FReply SConcertClientSessionBrowser::OnNewButtonClicked()
+{
+	SessionBrowser->InsertNewSessionEditableRow();
+	return FReply::Handled();
+}
+
+FReply SConcertClientSessionBrowser::OnLaunchServerButtonClicked()
+{
+	IMultiUserClientModule::Get().LaunchConcertServer();
+	bLocalServerRunning = IMultiUserClientModule::Get().IsConcertServerRunning(); // Immediately update the cache state to avoid showing buttons enabled for a split second.
+	return FReply::Handled();
+}
+
+FReply SConcertClientSessionBrowser::OnShutdownServerButtonClicked()
+{
+	if (bLocalServerRunning)
+	{
+		IMultiUserClientModule::Get().ShutdownConcertServer();
+	}
+	return FReply::Handled();
+}
+
+FReply SConcertClientSessionBrowser::OnAutoJoinButtonClicked()
+{
+	IConcertClientPtr ConcertClient = Controller->GetConcertClient();
+
+	// Start the 'auto connect' routine. It will try until it succeeded or gets canceled. Creating or Joining a session automatically cancels it.
+	ConcertClient->StartAutoConnect();
+	return FReply::Handled();
+}
+
+FReply SConcertClientSessionBrowser::OnCancelAutoJoinButtonClicked()
+{
+	Controller->GetConcertClient()->StopAutoConnect();
+	return FReply::Handled();
+}
+
+FReply SConcertClientSessionBrowser::OnJoinButtonClicked()
+{
+	TArray<TSharedPtr<FConcertSessionTreeItem>> SelectedItems = SessionBrowser->GetSelectedItems();
+	if (SelectedItems.Num() == 1)
+	{
+		RequestJoinSession(SelectedItems[0]);
+	}
+	return FReply::Handled();
+}
+
+void SConcertClientSessionBrowser::RequestJoinSession(const TSharedPtr<FConcertSessionTreeItem>& LiveItem)
+{
+	Controller->JoinSession(LiveItem->ServerAdminEndpointId, LiveItem->SessionId);
+}
+
+void SConcertClientSessionBrowser::Tick(const FGeometry& AllottedGeometry, const double InCurrentTime, const float InDeltaTime)
+{
+	// If the model received updates.
+	if (Controller->GetAndClearDiscoveryUpdateFlag())
+	{
+		// Don't wait next TickDiscovery() running at lower frequency, update immediately.
+		UpdateDiscovery();
+	}
+
+	// Ensure the 'default server' filter is updated when the configuration of the default server changes.
+	if (DefaultServerURL != Controller->GetConcertClient()->GetConfiguration()->DefaultServerURL)
+	{
+		DefaultServerURL = Controller->GetConcertClient()->GetConfiguration()->DefaultServerURL;
+		bRefreshSessionFilter = true;
+	}
+
+	// Should refresh the session filter?
+	if (bRefreshSessionFilter)
+	{
+		SessionBrowser->RefreshSessionList();
+		bRefreshSessionFilter = false;
+	}
+}
+
+#undef LOCTEXT_NAMESPACE
\ No newline at end of file
diff --git a/Engine/Plugins/Developer/Concert/ConcertApp/MultiUserClient/Source/MultiUserClient/Private/Widgets/Disconnected/SConcertClientSessionBrowser.h b/Engine/Plugins/Developer/Concert/ConcertApp/MultiUserClient/Source/MultiUserClient/Private/Widgets/Disconnected/SConcertClientSessionBrowser.h
new file mode 100644
--- /dev/null
+++ b/Engine/Plugins/Developer/Concert/ConcertApp/MultiUserClient/Source/MultiUserClient/Private/Widgets/Disconnected/SConcertClientSessionBrowser.h
@@ -0,0 +1,110 @@
+// Copyright Epic Games, Inc. All Rights Reserved.
+
+#pragma once
+
+#include "IConcertClientModule.h"
+#include "Widgets/DeclarativeSyntaxSupport.h"
+#include "Widgets/SCompoundWidget.h"
+
+class FConcertClientSessionBrowserController;
+class FConcertSessionTreeItem;
+class FExtender;
+class ITableRow;
+class SBorder;
+class SConcertSessionBrowser;
+class SExpandableArea;
+class SGridPanel;
+template<typename T>
+class SListView;
+class STableViewBase;
+
+/**
+ * Enables the user to browse/search/filter/sort active and archived sessions, create new session,
+ * archive active sessions, restore archived sessions, join a session and open the settings dialog.
+ */
+class SConcertClientSessionBrowser : public SCompoundWidget
+{
+public:
+	
+	SLATE_BEGIN_ARGS(SConcertClientSessionBrowser) { }
+	SLATE_END_ARGS();
+
+	/**
+	* Constructs the Browser.
+	* @param InArgs The Slate argument list.
+	* @param InConcertClient The concert client used to list, join, delete the sessions.
+	* @param[in,out] InSearchText The text to set in the search box and to remember (as output). Cannot be null.
+	*/
+	void Construct(const FArguments& InArgs, IConcertClientPtr InConcertClient, TSharedPtr<FText> InSearchText);
+
+	virtual void Tick(const FGeometry& AllottedGeometry, const double InCurrentTime, const float InDeltaTime) override;
+
+private:
+	
+	// Gives access to the concert data (servers, sessions, clients, etc).
+	TSharedPtr<FConcertClientSessionBrowserController> Controller;
+
+	TSharedPtr<SConcertSessionBrowser> SessionBrowser;
+
+	// Filtering.
+	bool bRefreshSessionFilter = true;
+	FString DefaultServerURL;
+
+	// Selected Session Details.
+	TSharedPtr<SBorder> SessionDetailsView;
+	TSharedPtr<SExpandableArea> DetailsArea;
+	TArray<TSharedPtr<FConcertSessionClientInfo>> Clients;
+	TSharedPtr<SExpandableArea> ClientsArea;
+	TSharedPtr<SListView<TSharedPtr<FConcertSessionClientInfo>>> ClientsView;
+
+	// Used to compare the version used by UI versus the version cached in the controller.
+	uint32 DisplayedSessionListVersion = 0;
+	uint32 DisplayedClientListVersion = 0;
+	uint32 ServerListVersion = 0;
+	bool bLocalServerRunning = false;
+
+	TSharedPtr<SWidget> ServerDiscoveryPanel; // Displayed until a server is found.
+	TSharedPtr<SWidget> SessionDiscoveryPanel; // Displayed until a session is found.
+	TSharedPtr<SWidget> NoSessionSelectedPanel; // Displays 'select a session to view details' message in the details section.
+	TSharedPtr<SWidget> NoSessionDetailsPanel; // Displays 'no details available' message in the details section.
+	TSharedPtr<SWidget> NoClientPanel; // Displays the 'no client connected' message in Clients expendable area.
+	
+	// Layout the 'session|details' split view.
+	TSharedRef<SWidget> MakeBrowserContent(TSharedPtr<FText> InSearchText);
+	void ExtendControlButtons(FExtender& Extender);
+	void ExtendSessionContextMenu(const TSharedPtr<FConcertSessionTreeItem>& Item, FExtender& Extender);
+	TSharedRef<SWidget> MakeUserAndSettings();
+	TSharedRef<SWidget> MakeOverlayedTableView(const TSharedRef<SWidget>& TableView);
+	
+	// Layouts the session detail panel.
+	TSharedRef<SWidget> MakeSessionDetails(TSharedPtr<FConcertSessionTreeItem> Item);
+	TSharedRef<SWidget> MakeActiveSessionDetails(TSharedPtr<FConcertSessionTreeItem> Item);
+	TSharedRef<SWidget> MakeArchivedSessionDetails(TSharedPtr<FConcertSessionTreeItem> Item);
+	TSharedRef<ITableRow> OnGenerateClientRowWidget(TSharedPtr<FConcertSessionClientInfo> Item, const TSharedRef<STableViewBase>& OwnerTable);
+	void PopulateSessionInfoGrid(SGridPanel& Grid, const FConcertSessionInfo& SessionInfo);
+
+	// The buttons in some overlays
+	bool IsLaunchServerButtonEnabled() const;
+	bool IsJoinButtonEnabled() const;
+	bool IsAutoJoinButtonEnabled() const;
+	bool IsCancelAutoJoinButtonEnabled() const;
+	FReply OnNewButtonClicked();
+	FReply OnLaunchServerButtonClicked();
+	FReply OnShutdownServerButtonClicked();
+	FReply OnAutoJoinButtonClicked();
+	FReply OnCancelAutoJoinButtonClicked();
+
+	// SConcertSessionBrowser extensions
+	FReply OnJoinButtonClicked();
+	void RequestJoinSession(const TSharedPtr<FConcertSessionTreeItem>& ActiveItem);
+
+	// Manipulates the sessions view (the array and the UI).
+	void OnSessionSelectionChanged(const TSharedPtr<FConcertSessionTreeItem>& SessionItem);
+	void OnSessionDoubleClicked(const TSharedPtr<FConcertSessionTreeItem>& SessionItem);
+	bool ConfirmDeleteSessionWithDialog(const TArray<TSharedPtr<FConcertSessionTreeItem>>& SessionItems) const;
+
+	// Update server/session/clients lists.
+	EActiveTimerReturnType TickDiscovery(double InCurrentTime, float InDeltaTime);
+	void UpdateDiscovery();
+	void RefreshClientList(const TArray<FConcertSessionClientInfo>& LastestClientList);
+};
diff --git a/Engine/Plugins/Developer/Concert/ConcertApp/MultiUserClient/Source/MultiUserClient/Private/Widgets/Disconnected/SConcertDiscovery.cpp b/Engine/Plugins/Developer/Concert/ConcertApp/MultiUserClient/Source/MultiUserClient/Private/Widgets/Disconnected/SConcertDiscovery.cpp
new file mode 100644
--- /dev/null
+++ b/Engine/Plugins/Developer/Concert/ConcertApp/MultiUserClient/Source/MultiUserClient/Private/Widgets/Disconnected/SConcertDiscovery.cpp
@@ -0,0 +1,72 @@
+// Copyright Epic Games, Inc. All Rights Reserved.
+
+#include "SConcertDiscovery.h"
+
+#include "Layout/Clipping.h"
+#include "Widgets/Images/SImage.h"
+#include "Widgets/Images/SThrobber.h"
+#include "Widgets/Input/SButton.h"
+#include "Widgets/Layout/SBox.h"
+#include "Widgets/SBoxPanel.h"
+#include "Widgets/Text/STextBlock.h"
+
+void SConcertDiscovery::Construct(const FArguments& InArgs)
+{
+	ChildSlot
+	[
+		SNew(SBox)
+		.HAlign(HAlign_Center)
+		.VAlign(VAlign_Center)
+		.Clipping(EWidgetClipping::ClipToBounds)
+		[
+			SNew(SVerticalBox)
+			+SVerticalBox::Slot()
+			.AutoHeight()
+			.HAlign(HAlign_Center)
+			[
+				SNew(SCircularThrobber)
+				.Visibility(InArgs._ThrobberVisibility)
+			]
+
+			+SVerticalBox::Slot()
+			.AutoHeight()
+			.HAlign(HAlign_Center)
+			[
+				SNew(STextBlock).Text(InArgs._Text).Justification(ETextJustify::Center)
+			]
+
+			+SVerticalBox::Slot()
+			.Padding(0, 4, 0, 0)
+			.AutoHeight()
+			.HAlign(HAlign_Center)
+			[
+				SNew(SButton)
+				.ForegroundColor(FSlateColor::UseForeground())
+				.ButtonStyle(InArgs._ButtonStyle)
+				.Visibility(InArgs._ButtonVisibility)
+				.HAlign(HAlign_Center)
+				.VAlign(VAlign_Center)
+				.IsEnabled(InArgs._IsButtonEnabled)
+				.OnClicked(InArgs._OnButtonClicked)
+				.ToolTipText(InArgs._ButtonToolTip)
+				.ContentPadding(FMargin(8, 4))
+				[
+					SNew(SHorizontalBox)
+					+SHorizontalBox::Slot()
+					.AutoWidth()
+					.VAlign(VAlign_Center)
+					.Padding(0, 0, 3, 0)
+					[
+						SNew(SImage).Image(InArgs._ButtonIcon)
+					]
+					+SHorizontalBox::Slot()
+					.AutoWidth()
+					.VAlign(VAlign_Bottom)
+					[
+						SNew(STextBlock).Text(InArgs._ButtonText)
+					]
+				]
+			]
+		]
+	];
+}
diff --git a/Engine/Plugins/Developer/Concert/ConcertApp/MultiUserClient/Source/MultiUserClient/Private/Widgets/Disconnected/SConcertDiscovery.h b/Engine/Plugins/Developer/Concert/ConcertApp/MultiUserClient/Source/MultiUserClient/Private/Widgets/Disconnected/SConcertDiscovery.h
new file mode 100644
--- /dev/null
+++ b/Engine/Plugins/Developer/Concert/ConcertApp/MultiUserClient/Source/MultiUserClient/Private/Widgets/Disconnected/SConcertDiscovery.h
@@ -0,0 +1,40 @@
+// Copyright Epic Games, Inc. All Rights Reserved.
+
+#pragma once
+
+#include "Framework/SlateDelegates.h"
+#include "Styling/AppStyle.h"
+#include "Styling/SlateTypes.h"
+#include "Styling/SlateWidgetStyleAsset.h"
+#include "Widgets/DeclarativeSyntaxSupport.h"
+#include "Widgets/SCompoundWidget.h"
+
+/** Widget displayed when discovering multi-user server(s) or session(s).*/
+class SConcertDiscovery : public SCompoundWidget
+{
+public:
+	
+	SLATE_BEGIN_ARGS(SConcertDiscovery)
+		: _Text()
+		, _ThrobberVisibility(EVisibility::Visible)
+		, _ButtonVisibility(EVisibility::Visible)
+		, _IsButtonEnabled(true)
+		, _ButtonStyle(&FAppStyle::Get().GetWidgetStyle<FButtonStyle>("FlatButton"))
+		, _ButtonIcon()
+		, _ButtonText()
+		, _ButtonToolTip()
+		, _OnButtonClicked()
+	{}
+		SLATE_ATTRIBUTE(FText, Text)
+		SLATE_ATTRIBUTE(EVisibility, ThrobberVisibility)
+		SLATE_ATTRIBUTE(EVisibility, ButtonVisibility)
+		SLATE_ATTRIBUTE(bool, IsButtonEnabled)
+		SLATE_STYLE_ARGUMENT(FButtonStyle, ButtonStyle)
+		SLATE_ATTRIBUTE(const FSlateBrush*, ButtonIcon)
+		SLATE_ATTRIBUTE(FText, ButtonText)
+		SLATE_ATTRIBUTE(FText, ButtonToolTip)
+		SLATE_EVENT( FOnClicked, OnButtonClicked)
+	SLATE_END_ARGS();
+
+	void Construct(const FArguments& InArgs);
+};
diff --git a/Engine/Plugins/Developer/Concert/ConcertApp/MultiUserClient/Source/MultiUserClient/Private/Widgets/Disconnected/SConcertNoAvailability.cpp b/Engine/Plugins/Developer/Concert/ConcertApp/MultiUserClient/Source/MultiUserClient/Private/Widgets/Disconnected/SConcertNoAvailability.cpp
new file mode 100644
--- /dev/null
+++ b/Engine/Plugins/Developer/Concert/ConcertApp/MultiUserClient/Source/MultiUserClient/Private/Widgets/Disconnected/SConcertNoAvailability.cpp
@@ -0,0 +1,18 @@
+// Copyright Epic Games, Inc. All Rights Reserved.
+
+#include "SConcertNoAvailability.h"
+
+#include "Layout/Visibility.h"
+#include "SConcertDiscovery.h"
+
+void SConcertNoAvailability::Construct(const FArguments& InArgs)
+{
+	ChildSlot
+	[
+		// Reuse this panel, but only show the message.
+		SNew(SConcertDiscovery) 
+			.Text(InArgs._Text)
+			.ThrobberVisibility(EVisibility::Collapsed)
+			.ButtonVisibility(EVisibility::Collapsed)
+	];
+}
diff --git a/Engine/Plugins/Developer/Concert/ConcertApp/MultiUserClient/Source/MultiUserClient/Private/Widgets/Disconnected/SConcertNoAvailability.h b/Engine/Plugins/Developer/Concert/ConcertApp/MultiUserClient/Source/MultiUserClient/Private/Widgets/Disconnected/SConcertNoAvailability.h
new file mode 100644
--- /dev/null
+++ b/Engine/Plugins/Developer/Concert/ConcertApp/MultiUserClient/Source/MultiUserClient/Private/Widgets/Disconnected/SConcertNoAvailability.h
@@ -0,0 +1,18 @@
+// Copyright Epic Games, Inc. All Rights Reserved.
+
+#pragma once
+
+#include "Widgets/DeclarativeSyntaxSupport.h"
+#include "Widgets/SCompoundWidget.h"
+
+/** Displayed when something is not available. */
+class SConcertNoAvailability : public SCompoundWidget
+{
+public:
+	
+	SLATE_BEGIN_ARGS(SConcertNoAvailability) : _Text() {}
+		SLATE_ATTRIBUTE(FText, Text)
+	SLATE_END_ARGS()
+	
+	void Construct(const FArguments& InArgs);
+};
diff --git a/Engine/Plugins/Developer/Concert/ConcertApp/MultiUserClient/Source/MultiUserClient/Private/Widgets/SConcertBrowser.cpp b/Engine/Plugins/Developer/Concert/ConcertApp/MultiUserClient/Source/MultiUserClient/Private/Widgets/SConcertBrowser.cpp
--- a/Engine/Plugins/Developer/Concert/ConcertApp/MultiUserClient/Source/MultiUserClient/Private/Widgets/SConcertBrowser.cpp
+++ b/Engine/Plugins/Developer/Concert/ConcertApp/MultiUserClient/Source/MultiUserClient/Private/Widgets/SConcertBrowser.cpp
@@ -2,1834 +2,21 @@
 
 #include "SConcertBrowser.h"
 
-#include "ConcertActivityStream.h"
-#include "ConcertFrontendStyle.h"
-#include "ConcertFrontendUtils.h"
-#include "ConcertLogGlobal.h"
-#include "ConcertSettings.h"
-#include "ConcertClientSettings.h"
-#include "IMultiUserClientModule.h"
 #include "IConcertClient.h"
-#include "IConcertClientWorkspace.h"
 #include "IConcertSyncClient.h"
 #include "MultiUserClientUtils.h"
 #include "SActiveSession.h"
-#include "SConcertSessionRecovery.h"
+#include "Widgets/Disconnected/SConcertClientSessionBrowser.h"
+#include "Widgets/Disconnected/SConcertNoAvailability.h"
 
-#include "Algo/TiedTupleOutput.h"
-
-#include "Session/Browser/ConcertBrowserUtils.h"
-#include "Session/Browser/Items/ConcertSessionTreeItem.h"
-#include "Session/Browser/IConcertSessionBrowserController.h"
-#include "Session/Browser/SConcertSessionBrowser.h"
-
-#include "Algo/Transform.h"
-#include "Framework/Application/SlateApplication.h"
-#include "Framework/Commands/UIAction.h"
-#include "Framework/Commands/UICommandInfo.h"
-#include "Framework/MultiBox/MultiBoxBuilder.h"
 #include "Internationalization/Regex.h"
-#include "Misc/AsyncTaskNotification.h"
-#include "Misc/TextFilter.h"
-#include "Misc/MessageDialog.h"
 #include "Styling/AppStyle.h"
 #include "Textures/SlateIcon.h"
 #include "Widgets/Images/SImage.h"
-#include "Widgets/Images/SThrobber.h"
-#include "Widgets/Input/SButton.h"
-#include "Widgets/Input/SSearchBox.h"
-#include "Widgets/Layout/SBorder.h"
-#include "Widgets/Layout/SBox.h"
-#include "Widgets/Layout/SExpandableArea.h"
-#include "Widgets/Layout/SGridPanel.h"
-#include "Widgets/Layout/SScrollBox.h"
-#include "Widgets/Layout/SUniformGridPanel.h"
-#include "Widgets/Text/SInlineEditableTextBlock.h"
-#include "Widgets/Text/STextBlock.h"
 #include "Widgets/SWindow.h"
-#include "Widgets/Views/SListView.h"
-#include "Widgets/Views/STableRow.h"
-
-#if WITH_EDITOR
-	#include "ISettingsModule.h"
-#endif
 
 #define LOCTEXT_NAMESPACE "SConcertBrowser"
 
-/**
- * Runs and cache network queries for the UI. In the model-view-controller pattern, this class acts like the controller. Its purpose
- * is to keep the UI code as decoupled as possible from the API used to query it. It encapsulate the asynchronous code and provide a
- * simpler API to the UI.
- */
-class FConcertClientSessionBrowserController : public IConcertSessionBrowserController
-{
-public:
-	
-	/** Keeps the state of an active async request and provides a tool to cancel its future continuation execution. */
-	struct FAsyncRequest
-	{
-		/** Returns true if there is a registered async request future and if it hasn't executed yet. */
-		bool IsOngoing() const { return Future.IsValid() && !Future.IsReady(); }
-
-		/** Reset the execution token, canceling previous execution (if any) and setting up the token for a new request. */
-		TWeakPtr<uint8, ESPMode::ThreadSafe> ResetExecutionToken() { FutureExecutionToken = MakeShared<uint8, ESPMode::ThreadSafe>(); return FutureExecutionToken; }
-
-		/** Cancel the execution of request async continuation. */
-		void Cancel() { FutureExecutionToken.Reset(); }
-
-		/** The future provided by an asynchronous request. */
-		TFuture<void> Future;
-
-		/** Determines whether or not the async request continuation code should execute. Reset to disarm execution of an async future continuation. */
-		TSharedPtr<uint8, ESPMode::ThreadSafe> FutureExecutionToken;
-	};
-	
-	FConcertClientSessionBrowserController(IConcertClientPtr InConcertClient);
-	virtual ~FConcertClientSessionBrowserController() override;
-
-	/**
-	 * Fires new requests to retrieve all known server and for each server, their active and archived sessions. The responses are
-	 * received asynchronously and may not be available right now. When a response is received, if the corresponding list cached
-	 * is updated, the list version is incremented.
-	 *
-	 * @return A (serverListVersion, sessionsListVersion) pair, corresponding the the versions currently cached by this object.
-	 */
-	TPair<uint32, uint32> TickServersAndSessionsDiscovery();
-
-	/**
-	 * Fires a new request to retrieve the clients for the selected session. The result is cached and can be retrived by GetClients().
-	 * The class caches clients for a single active session determined by this function. If the specified session changes, the cache
-	 * of the previous session is cleared.
-	 *
-	 * @return The version of the client list currently cached by the object.
-	 */
-	uint32 TickClientsDiscovery(const FGuid& ServerAdminEndpointId, const FGuid& SessionId);
-
-	/** Returns true if the controller received async responses and updated its cache since the last time the function was called, then clear the flag. */
-	bool GetAndClearDiscoveryUpdateFlag();
-	
-	/** Returns the latest list of clients corresponding to the session known to this controller. Ensure to call TickClientsDiscovery() periodically. */
-	const TArray<FConcertSessionClientInfo>& GetClients(const FGuid& ServerAdminEndpointId, const FGuid& SessionId) const;
-	
-	void JoinSession(const FGuid& ServerAdminEndpointId, const FGuid& SessionId);
-	
-	// The 2 functions below are used to prevent fast UI transition with the 'no session' panel flashing when the session list is empty.
-	bool HasReceivedInitialSessionList() const { return bInitialActiveSessionQueryResponded && bInitialArchivedSessionQueryResponded; }
-	bool IsCreatingSession() const { return CreateSessionRequests.Num() > 0 || ExpectedSessionsToDiscover.Num() > 0; }
-	
-	IConcertClientPtr GetConcertClient() const { return ConcertClient; }
-	
-	//~ Begin IConcertSessionBrowserController Interface
-	virtual TArray<FConcertServerInfo> GetServers() const override { return Servers; }
-
-	virtual TArray<FActiveSessionInfo> GetActiveSessions() const override
-	{
-		TArray<FActiveSessionInfo> Result;
-		Algo::Transform(ActiveSessions, Result, [](TSharedPtr<FClientActiveSessionInfo> Info) { return *Info; });
-		return Result;
-	}
-
-	virtual TArray<FArchivedSessionInfo> GetArchivedSessions() const override
-	{
-		TArray<FArchivedSessionInfo> Result;
-		Algo::Transform(ArchivedSessions, Result, [](TSharedPtr<FClientArchivedSessionInfo> Info) { return *Info; });
-		return Result;
-	}
-
-	virtual TOptional<FConcertSessionInfo> GetActiveSessionInfo(const FGuid& AdminEndpoint, const FGuid& SessionId) const override;
-	virtual TOptional<FConcertSessionInfo> GetArchivedSessionInfo(const FGuid& AdminEndpoint, const FGuid& SessionId) const override;
-
-	virtual void CreateSession(const FGuid& ServerAdminEndpointId, const FString& SessionName, const FString& ProjectName) override;
-	virtual void ArchiveSession(const FGuid& ServerAdminEndpointId, const FGuid& SessionId, const FString& ArchiveName, const FConcertSessionFilter& SessionFilter) override;
-	virtual void RestoreSession(const FGuid& ServerAdminEndpointId, const FGuid& SessionId, const FString& RestoredName, const FConcertSessionFilter& SessionFilter) override;
-	virtual void RenameActiveSession(const FGuid& ServerAdminEndpointId, const FGuid& SessionId, const FString& NewName) override;
-	virtual void RenameArchivedSession(const FGuid& ServerAdminEndpointId, const FGuid& SessionId, const FString& NewName) override;
-	virtual bool CanRenameActiveSession(const FGuid& ServerAdminEndpointId, const FGuid& SessionId) const override;
-	virtual bool CanRenameArchivedSession(const FGuid& ServerAdminEndpointId, const FGuid& SessionId) const override;
-	virtual void DeleteSessions(const FGuid& ServerAdminEndpointId, const TArray<FGuid>& SessionIds) override;
-	virtual bool CanDeleteActiveSession(const FGuid& ServerAdminEndpointId, const FGuid& SessionId) const override;
-	virtual bool CanDeleteArchivedSession(const FGuid& ServerAdminEndpointId, const FGuid& SessionId) const override;
-	virtual bool CanEverCreateSessions() const override { return true; }
-	//~ End IConcertSessionBrowserController Interface
-
-private:
-	
-	/** Hold information about a session created by this client, not yet 'discovered' by a 'list session' query, but expected to be soon. */
-	struct FPendingSessionDiscovery
-	{
-		FDateTime CreateTimestamp;
-		FGuid ServerEndpoint;
-		FString SessionName;
-	};
-
-	struct FMonitoredSessionInfo
-	{
-		TSharedRef<FActiveSessionInfo> Session;
-		FAsyncRequest ListClientRequest;
-
-		FMonitoredSessionInfo(TSharedRef<FActiveSessionInfo> Session)
-			: Session(Session)
-		{}
-	};
-
-	struct FClientActiveSessionInfo : public FActiveSessionInfo
-	{
-		/** Raised when the UI and the cache values may be out of sync if a rename failed (UI assumed it succeeded) */
-		bool bSessionNameDirty = false;
-
-		FClientActiveSessionInfo(FConcertServerInfo ServerInfo, FConcertSessionInfo SessionInfo)
-			: FActiveSessionInfo(MoveTemp(ServerInfo), MoveTemp(SessionInfo))
-		{}
-	};
-
-	struct FClientArchivedSessionInfo : public FArchivedSessionInfo
-	{
-		/** Raised when the UI and the cache values may be out of sync if a rename failed (UI assumed it succeeded) */
-		bool bSessionNameDirty = false;
-
-		FClientArchivedSessionInfo(FConcertServerInfo ServerInfo, FConcertSessionInfo SessionInfo)
-			: FArchivedSessionInfo(MoveTemp(ServerInfo), MoveTemp(SessionInfo))
-		{}
-	};
-
-	void UpdateSessionsAsync();
-	void UpdateActiveSessionsAsync(const FConcertServerInfo& ServerInfo);
-	void UpdateArchivedSessionsAsync(const FConcertServerInfo& ServerInfo);
-	void UpdateClientsAsync(const FGuid& ServerAdminEndpointId, const FGuid& SessionId);
-
-	void OnActiveSessionDiscovered(const FActiveSessionInfo&) { ++SessionListVersion; bCacheUpdated = true; }
-	void OnActiveSessionDiscarded(const FActiveSessionInfo&) { ++SessionListVersion; bCacheUpdated = true; }
-	void OnArchivedSessionDiscovered(const FArchivedSessionInfo&) { ++SessionListVersion; bCacheUpdated = true; }
-	void OnArchivedSessionDiscarded(const FArchivedSessionInfo&) { ++SessionListVersion; bCacheUpdated = true; }
-	void OnActiveSessionClientsUpdated(const FActiveSessionInfo&) { ++ClientListVersion; bCacheUpdated = true; }
-	void OnActiveSessionRenamed(const FActiveSessionInfo&, const FString& NewName) { ++SessionListVersion; bCacheUpdated = true; }
-	void OnArchivedSessionRenamed(const FArchivedSessionInfo&, const FString& NewName) { ++SessionListVersion; bCacheUpdated = true; }
-	void OnActiveSessionListDirty() { ++SessionListVersion; bCacheUpdated = true; } // This will force the UI to refresh its list.
-	void OnArchivedSessionListDirty() { ++SessionListVersion; bCacheUpdated = true; } // This will force the UI to refresh its list.
-
-	// Holds a concert client instance.
-	IConcertClientPtr ConcertClient;
-
-	// The list of active/archived async requests (requesting the list of session) per server. There is only one per server as we prevent stacking more than one at the time.
-	TMap<FGuid, FAsyncRequest> ActiveSessionRequests;
-	TMap<FGuid, FAsyncRequest> ArchivedSessionRequests;
-
-	// The cached lists.
-	TArray<FConcertServerInfo> Servers;
-	TArray<TSharedPtr<FClientActiveSessionInfo>> ActiveSessions;
-	TArray<TSharedPtr<FClientArchivedSessionInfo>> ArchivedSessions;
-
-	// The session for which the clients are monitored. UI only monitor client of 1 session at the time.
-	TOptional<FMonitoredSessionInfo> ClientMonitoredSession;
-
-	// Holds the version of data cached by the controller. The version is updated when an async response is received and implies a change in the cached values.
-	uint32 ServerListVersion = 0;
-	uint32 SessionListVersion = 0;
-	uint32 ClientListVersion = 0;
-	bool bCacheUpdated = false;
-	bool bInitialActiveSessionQueryResponded = false;
-	bool bInitialArchivedSessionQueryResponded = false;
-
-	TArray<FAsyncRequest> CreateSessionRequests;
-	TArray<FPendingSessionDiscovery> ExpectedSessionsToDiscover;
-	TSet<FString> IgnoredServers; // List of ignored servers (Useful for testing/debugging)
-};
-
-FConcertClientSessionBrowserController::FConcertClientSessionBrowserController(IConcertClientPtr InConcertClient)
-{
-	check(InConcertClient.IsValid()); // Don't expect this class to be instantiated if the concert client is not available.
-	check(InConcertClient->IsConfigured()); // Expected to be done by higher level code.
-	ConcertClient = InConcertClient;
-
-	// When others servers are running, add them to the list if you want to test the UI displayed if no servers/no sessions exists.
-	//IgnoredServers.Add(TEXT("wksyul10355")); // TODO: COMMENT BEFORE SUBMIT
-
-	// Start server discovery to find the available Concert servers.
-	ConcertClient->StartDiscovery();
-
-	// Populate the session cache.
-	TickServersAndSessionsDiscovery();
-}
-
-FConcertClientSessionBrowserController::~FConcertClientSessionBrowserController()
-{
-	if (ConcertClient.IsValid())
-	{
-		ConcertClient->StopDiscovery();
-	}
-}
-
-const TArray<FConcertSessionClientInfo>& FConcertClientSessionBrowserController::GetClients(const FGuid& AdminEndpoint, const FGuid& SessionId) const
-{
-	// If a session clients are monitored and the list is cached
-	if (ClientMonitoredSession.IsSet() && ClientMonitoredSession->Session->ServerInfo.AdminEndpointId == AdminEndpoint && ClientMonitoredSession->Session->SessionInfo.SessionId == SessionId)
-	{
-		return ClientMonitoredSession->Session->Clients; // Returns the list retrieved by the last TickClientsDiscovery().
-	}
-
-	// Returns an empty list for now. We expect the caller to call TickClientsDiscovery() periodically to maintain the session client list.
-	static TArray<FConcertSessionClientInfo> EmptyClientList;
-	return EmptyClientList;
-}
-
-void FConcertClientSessionBrowserController::JoinSession(const FGuid& ServerAdminEndpointId, const FGuid& SessionId)
-{
-	// On success: The client joins the session and SConcertBrowser::HandleSessionConnectionChanged() will transit the UI to the SActiveSession.
-	// On failure: An async notification banner will be displayer to the user.
-	ConcertClient->JoinSession(ServerAdminEndpointId, SessionId);
-}
-
-TOptional<FConcertSessionInfo> FConcertClientSessionBrowserController::GetActiveSessionInfo(const FGuid& AdminEndpoint, const FGuid& SessionId) const
-{
-	const TSharedPtr<FClientActiveSessionInfo>* SessionInfo = ActiveSessions.FindByPredicate([&AdminEndpoint, &SessionId](const TSharedPtr<FActiveSessionInfo>& MatchCandidate)
-	{
-		return MatchCandidate->ServerInfo.AdminEndpointId == AdminEndpoint && MatchCandidate->SessionInfo.SessionId == SessionId;
-	});
-
-	return SessionInfo != nullptr ? (*SessionInfo)->SessionInfo : TOptional<FConcertSessionInfo>{};
-}
-
-TOptional<FConcertSessionInfo> FConcertClientSessionBrowserController::GetArchivedSessionInfo(const FGuid& AdminEndpoint, const FGuid& SessionId) const
-{
-	const TSharedPtr<FClientArchivedSessionInfo>* SessionInfo = ArchivedSessions.FindByPredicate([&AdminEndpoint, &SessionId](const TSharedPtr<FArchivedSessionInfo>& MatchCandidate)
-	{
-		return MatchCandidate->ServerInfo.AdminEndpointId == AdminEndpoint && MatchCandidate->SessionInfo.SessionId == SessionId;
-	});
-
-	return SessionInfo != nullptr ? (*SessionInfo)->SessionInfo : TOptional<FConcertSessionInfo>{};
-}
-
-void FConcertClientSessionBrowserController::CreateSession(const FGuid& ServerAdminEndpointId, const FString& SessionName, const FString& ProjectName)
-{
-	FAsyncRequest& CreateRequest = CreateSessionRequests.AddDefaulted_GetRef();
-	TWeakPtr<uint8, ESPMode::ThreadSafe> CreateRequestExecutionToken = CreateRequest.ResetExecutionToken();
-
-	// On success, the client automatically joins the new session and SConcertBrowser::HandleSessionConnectionChanged() will transit the UI to the SActiveSession.
-	// On failure: An async notification banner will be displayer to the user.
-	FConcertCreateSessionArgs CreateSessionArgs;
-	CreateSessionArgs.SessionName = SessionName;
-	// Project Name is not used an override for the client session create.
-	ConcertClient->CreateSession(ServerAdminEndpointId, CreateSessionArgs).Next([this, CreateRequestExecutionToken, ServerAdminEndpointId, SessionName](EConcertResponseCode ResponseCode)
-	{
-		if (TSharedPtr<uint8, ESPMode::ThreadSafe> ExecutionToken = CreateRequestExecutionToken.Pin())
-		{
-			if (ResponseCode == EConcertResponseCode::Success)
-			{
-				// Expect to find those session at some point in the future.
-				ExpectedSessionsToDiscover.Add(FPendingSessionDiscovery{ FDateTime::UtcNow(), ServerAdminEndpointId, SessionName });
-			}
-
-			// Stop tracking the request.
-			CreateSessionRequests.RemoveAll([&ExecutionToken](const FAsyncRequest& DiscardCandidate) { return DiscardCandidate.FutureExecutionToken == ExecutionToken; });
-		}
-	});
-}
-
-void FConcertClientSessionBrowserController::ArchiveSession(const FGuid& ServerAdminEndpointId, const FGuid& SessionId, const FString& ArchiveName, const FConcertSessionFilter& SessionFilter)
-{
-	// On success, an archived is created and TickServersAndSessionsDiscovery() will eventually discover it.
-	// On failure: An async notification banner will be displayer to the user.
-	FConcertArchiveSessionArgs ArchiveSessionArgs;
-	ArchiveSessionArgs.SessionId = SessionId;
-	ArchiveSessionArgs.ArchiveNameOverride = ArchiveName;
-	ArchiveSessionArgs.SessionFilter = SessionFilter;
-	ConcertClient->ArchiveSession(ServerAdminEndpointId, ArchiveSessionArgs);
-}
-
-void FConcertClientSessionBrowserController::RestoreSession(const FGuid& ServerAdminEndpointId, const FGuid& SessionId, const FString& RestoredName, const FConcertSessionFilter& SessionFilter)
-{
-	FString ArchivedSessionName;
-	if (const TOptional<FConcertSessionInfo> SessionInfo = GetArchivedSessionInfo(ServerAdminEndpointId, SessionId))
-	{
-		ArchivedSessionName = SessionInfo->SessionName;
-	}
-
-	TSharedRef<SWindow> NewWindow = SNew(SWindow)
-		.Title(FText::Format(LOCTEXT("RestoreSessionDialogTitle", "Restoring {0}"), FText::AsCultureInvariant(ArchivedSessionName)))
-		.SizingRule(ESizingRule::UserSized)
-		.ClientSize(FVector2D(1200, 800))
-		.IsTopmostWindow(false) // Consider making it always on top?
-		.SupportsMaximize(true)
-		.SupportsMinimize(false);
-
-	// Ask the stream to pull the activity details (for transaction/package) for inspection.
-	constexpr bool bRequestActivityDetails = true;
-
-	// Create a stream of activities (streaming from the most recent to the oldest).
-	TSharedPtr<FConcertActivityStream> ActivityStream = MakeShared<FConcertActivityStream>(ConcertClient, ServerAdminEndpointId, SessionId, bRequestActivityDetails);
-
-	// The UI uses this function to read and consume the activity stream.
-	auto ReadActivitiesFn = [ActivityStream](TArray<TSharedPtr<FConcertSessionActivity>>& InOutActivities, int32& OutFetchCount, FText& OutErrorMsg)
-	{
-		return ActivityStream->Read(InOutActivities, OutFetchCount, OutErrorMsg);
-	};
-
-	// The UI uses this function to map an activity ID from the stream to a client info.
-	auto GetActivityClientInfoFn = [ActivityStream](FGuid EndpointID)
-	{
-		const FConcertClientInfo* Result = ActivityStream->GetActivityClientInfo(EndpointID);
-		return Result ? TOptional<FConcertClientInfo>{ *Result } : TOptional<FConcertClientInfo>{};
-	};
-
-	// Invoked if the client selects a point in time to recover.
-	TWeakPtr<IConcertClient, ESPMode::ThreadSafe> WeakClient = ConcertClient;
-	auto OnAcceptRestoreFn = [WeakClient, ServerAdminEndpointId, SessionId, RestoredName, SessionFilter](TSharedPtr<FConcertSessionActivity> SelectedRecoveryActivity)
-	{
-		FConcertCopySessionArgs RestoreSessionArgs;
-		RestoreSessionArgs.bAutoConnect = true;
-		RestoreSessionArgs.SessionId = SessionId;
-		RestoreSessionArgs.SessionName = RestoredName;
-		RestoreSessionArgs.SessionFilter = SessionFilter;
-		RestoreSessionArgs.SessionFilter.bOnlyLiveData = false;
-
-		// Set which item was selected to recover through.
-		if (SelectedRecoveryActivity)
-		{
-			RestoreSessionArgs.SessionFilter.ActivityIdUpperBound = SelectedRecoveryActivity->Activity.ActivityId;
-		}
-		// else -> Restore the entire session as it.
-
-		bool bDismissRecoveryWindow = true; // Dismiss the window showing the session activities
-
-		if (TSharedPtr<IConcertClient, ESPMode::ThreadSafe> ConcertClientPin = WeakClient.Pin())
-		{
-			// Prompt the user to persist and leave the session.
-			bool bDisconnected = IMultiUserClientModule::Get().DisconnectSession(/*bAlwaysAskConfirmation*/true);
-			if (bDisconnected)
-			{
-				// On success, a new session is created from the archive, the client automatically disconnects from the current session (if any), joins the restored one and SConcertBrowser::HandleSessionConnectionChanged() will transit the UI to the SActiveSession.
-				// On failure: An async notification banner will be displayer to the user.
-				ConcertClientPin->RestoreSession(ServerAdminEndpointId, RestoreSessionArgs);
-			}
-			else // The user declined disconnection.
-			{
-				bDismissRecoveryWindow = false; // Keep the window open, let the client handle it (close/cancel or just restore later).
-			}
-		}
-		else
-		{
-			FAsyncTaskNotificationConfig NotificationConfig;
-			NotificationConfig.bIsHeadless = false;
-			NotificationConfig.bKeepOpenOnFailure = true;
-			NotificationConfig.LogCategory = &LogConcert;
-
-			FAsyncTaskNotification Notification(NotificationConfig);
-			Notification.SetComplete(LOCTEXT("RecoveryError", "Failed to recover the session"), LOCTEXT("ClientUnavailable", "Concert client unavailable"), /*Success*/ false);
-		}
-		return bDismissRecoveryWindow;
-	};
-
-	TSharedRef<SConcertSessionRecovery> RestoreWidget = SNew(SConcertSessionRecovery)
-		.ParentWindow(NewWindow)
-		.IntroductionText(LOCTEXT("RecoverSessionIntroductionText", "Select the point in time at which the session should be restored"))
-		.OnFetchActivities_Lambda(ReadActivitiesFn)
-		.OnMapActivityToClient_Lambda(GetActivityClientInfoFn)
-		.OnRestore(OnAcceptRestoreFn)
-		.WithClientAvatarColorColumn(true)
-		.WithClientNameColumn(true)
-		.WithOperationColumn(true)
-		.WithPackageColumn(false) // Even tough the column is not present, the tooltips and summary contains the affected package.
-		.DetailsAreaVisibility(bRequestActivityDetails ? EVisibility::Visible : EVisibility::Collapsed) // The activity stream was configured to pull the activity details.
-		.IsConnectionActivityFilteringEnabled(true)
-		.IsLockActivityFilteringEnabled(true);
-
-	NewWindow->SetContent(RestoreWidget);
-	FSlateApplication::Get().AddWindow(NewWindow, true);
-}
-
-void FConcertClientSessionBrowserController::RenameActiveSession(const FGuid& ServerAdminEndpointId, const FGuid& SessionId, const FString& NewName)
-{
-	// Find the currently cached session info.
-	const TSharedPtr<FClientActiveSessionInfo>* SessionInfo = ActiveSessions.FindByPredicate([&ServerAdminEndpointId, &SessionId](const TSharedPtr<FActiveSessionInfo>& MatchCandidate)
-	{
-		return MatchCandidate->ServerInfo.AdminEndpointId == ServerAdminEndpointId && MatchCandidate->SessionInfo.SessionId == SessionId;
-	});
-
-	check(SessionInfo); // If the UI is displaying it, the UI backend should have it.
-	ConcertClient->RenameSession(ServerAdminEndpointId, SessionId, NewName).Next([ActiveSessionInfo = *SessionInfo](EConcertResponseCode Response)
-	{
-		if (Response != EConcertResponseCode::Success)
-		{
-			ActiveSessionInfo->bSessionNameDirty = true; // Renamed failed, the UI may be displaying the wrong name. Force the UI to update against the latest cached values.
-		}
-		// else -> Succeeded -> TickServersAndSessionsDiscovery() will pick up the new name.
-	});
-}
-
-void FConcertClientSessionBrowserController::RenameArchivedSession(const FGuid& ServerAdminEndpointId, const FGuid& SessionId, const FString& NewName)
-{
-	// Find the currently cached session info.
-	const TSharedPtr<FClientArchivedSessionInfo>* SessionInfo = ArchivedSessions.FindByPredicate([&ServerAdminEndpointId, &SessionId](const TSharedPtr<FArchivedSessionInfo>& MatchCandidate)
-	{
-		return MatchCandidate->ServerInfo.AdminEndpointId == ServerAdminEndpointId && MatchCandidate->SessionInfo.SessionId == SessionId;
-	});
-
-	check(SessionInfo); // If the UI is displaying it, the UI backend should have it.
-	ConcertClient->RenameSession(ServerAdminEndpointId, SessionId, NewName).Next([SessionInfo = *SessionInfo](EConcertResponseCode Response)
-	{
-		if (Response != EConcertResponseCode::Success)
-		{
-			SessionInfo->bSessionNameDirty = true; // Renamed failed, the UI may be displaying the wrong name. Force the UI to update against the latest cached values.
-		}
-		// else -> Succeeded -> TickServersAndSessionsDiscovery() will pick up the new name.
-	});
-}
-
-bool FConcertClientSessionBrowserController::CanRenameActiveSession(const FGuid& ServerAdminEndpointId, const FGuid& SessionId) const
-{
-	return CanDeleteActiveSession(ServerAdminEndpointId, SessionId); // Rename requires the same permission than delete.
-}
-
-bool FConcertClientSessionBrowserController::CanRenameArchivedSession(const FGuid& ServerAdminEndpointId, const FGuid& SessionId) const
-{
-	return CanDeleteArchivedSession(ServerAdminEndpointId, SessionId); // Rename requires the same permission than delete.
-}
-
-void FConcertClientSessionBrowserController::DeleteSessions(const FGuid& ServerAdminEndpointId, const TArray<FGuid>& SessionIds)
-{
-	// The difference between the two functions is how error notifications are displayed...
-	if (SessionIds.Num() == 1)
-	{
-		// ... shows a specific notification
-		ConcertClient->DeleteSession(ServerAdminEndpointId, SessionIds[0]);
-	}
-	else
-	{
-		TSet<FGuid> SessionIdSet;
-		Algo::Transform(SessionIds, SessionIdSet, [](const FGuid& SessionId) { return SessionId; });
-		// ... shows compressed notification(s)
-		ConcertClient->BatchDeleteSessions(ServerAdminEndpointId, { SessionIdSet, EBatchSessionDeletionFlags::SkipForbiddenSessions });
-	}
-}
-
-bool FConcertClientSessionBrowserController::CanDeleteActiveSession(const FGuid& ServerAdminEndpointId, const FGuid& SessionId) const
-{
-	const TSharedPtr<FClientActiveSessionInfo>* SessionInfo = ActiveSessions.FindByPredicate([&ServerAdminEndpointId, &SessionId](const TSharedPtr<FActiveSessionInfo>& MatchCandidate)
-	{
-		return MatchCandidate->ServerInfo.AdminEndpointId == ServerAdminEndpointId && MatchCandidate->SessionInfo.SessionId == SessionId;
-	});
-
-	// Can delete the session only if the concert client is the owner.
-	return SessionInfo == nullptr ? false : ConcertClient->IsOwnerOf((*SessionInfo)->SessionInfo);
-}
-
-bool FConcertClientSessionBrowserController::CanDeleteArchivedSession(const FGuid& ServerAdminEndpointId, const FGuid& SessionId) const
-{
-	const TSharedPtr<FClientArchivedSessionInfo>* SessionInfo = ArchivedSessions.FindByPredicate([&ServerAdminEndpointId, &SessionId](const TSharedPtr<FArchivedSessionInfo>& MatchCandidate)
-	{
-		return MatchCandidate->ServerInfo.AdminEndpointId == ServerAdminEndpointId && MatchCandidate->SessionInfo.SessionId == SessionId;
-	});
-
-	// Can delete the session only if the concert client is the owner.
-	return SessionInfo == nullptr ? false : ConcertClient->IsOwnerOf((*SessionInfo)->SessionInfo);
-}
-
-TPair<uint32, uint32> FConcertClientSessionBrowserController::TickServersAndSessionsDiscovery()
-{
-	// Fire new async queries to poll the sessions on all known servers.
-	UpdateSessionsAsync();
-
-	// Returns the versions corresponding to the last async responses received. The version are incremented when response to the asynchronous queries are received.
-	return MakeTuple(ServerListVersion, SessionListVersion);
-}
-
-uint32 FConcertClientSessionBrowserController::TickClientsDiscovery(const FGuid& ServerAdminEndpointId, const FGuid& SessionId)
-{
-	// Fire new async queries to poll clients for the selected server/session pair.
-	UpdateClientsAsync(ServerAdminEndpointId, SessionId);
-
-	// Returns the versions corresponding to the last async responses received. The version is incremented when response to the asynchronous query is received.
-	return ClientListVersion;
-}
-
-bool FConcertClientSessionBrowserController::GetAndClearDiscoveryUpdateFlag()
-{
-	bool bOldCacheChanged = bCacheUpdated; // This flag is raised every time an async response updates the cached data.
-	bCacheUpdated = false;
-	return bOldCacheChanged;
-}
-
-void FConcertClientSessionBrowserController::UpdateSessionsAsync()
-{
-	// Get the list of known servers.
-	TArray<FConcertServerInfo> OnlineServers = ConcertClient->GetKnownServers();
-	if (IgnoredServers.Num())
-	{
-		OnlineServers.RemoveAll([this](const FConcertServerInfo& ServerInfo) { return IgnoredServers.Contains(ServerInfo.ServerName); });
-	}
-
-	bool bServerListVersionUpdated = false;
-
-	// Detects which server(s) went offline since the last update.
-	for (const FConcertServerInfo& ServerInfo : Servers)
-	{
-		// If a server, previously tracked, is not online anymore, remove all its sessions.
-		if (!OnlineServers.ContainsByPredicate([this, &ServerInfo](const FConcertServerInfo& Visited) { return ServerInfo.InstanceInfo.InstanceId == Visited.InstanceInfo.InstanceId; }))
-		{
-			// Remove all active sessions corresponding to this server.
-			ActiveSessions.RemoveAll([this, &ServerInfo](const TSharedPtr<FActiveSessionInfo>& ActiveSessionInfo)
-			{
-				if (ServerInfo.InstanceInfo.InstanceId == ActiveSessionInfo->ServerInfo.InstanceInfo.InstanceId)
-				{
-					OnActiveSessionDiscarded(*ActiveSessionInfo);
-					return true;
-				}
-				return false;
-			});
-
-			// Remove all archived sessions corresponding to this server.
-			ArchivedSessions.RemoveAll([this, &ServerInfo](const TSharedPtr<FArchivedSessionInfo>& ArchivedSessionInfo)
-			{
-				if (ServerInfo.InstanceInfo.InstanceId == ArchivedSessionInfo->ServerInfo.InstanceInfo.InstanceId)
-				{
-					OnArchivedSessionDiscarded(*ArchivedSessionInfo);
-					return true;
-				}
-				return false;
-			});
-
-			// Disarm any active async request future for this server. Removing it from the map effectively disarm the future as the shared pointer used to as token gets released.
-			ActiveSessionRequests.Remove(ServerInfo.InstanceInfo.InstanceId);
-			ArchivedSessionRequests.Remove(ServerInfo.InstanceInfo.InstanceId);
-
-			// The server went offline, update the server list version.
-			ServerListVersion++;
-			bServerListVersionUpdated = true;
-		}
-	}
-
-	// For all online servers.
-	for (const FConcertServerInfo& ServerInfo : OnlineServers)
-	{
-		// Check if this is a new server.
-		if (!bServerListVersionUpdated && !Servers.ContainsByPredicate([&ServerInfo](const FConcertServerInfo& Visited) { return ServerInfo.InstanceInfo.InstanceId == Visited.InstanceInfo.InstanceId; }))
-		{
-			ServerListVersion++;
-			bServerListVersionUpdated = true; // No need to look up further, only one new server is required to update the version.
-		}
-
-		// Poll the sessions for this online server.
-		UpdateActiveSessionsAsync(ServerInfo);
-		UpdateArchivedSessionsAsync(ServerInfo);
-	}
-
-	// Keep the list of online servers.
-	Servers = MoveTemp(OnlineServers);
-}
-
-void FConcertClientSessionBrowserController::UpdateActiveSessionsAsync(const FConcertServerInfo& ServerInfo)
-{
-	// Check if a request is already pulling active session from this server.
-	FAsyncRequest& ListActiveSessionAsyncRequest = ActiveSessionRequests.FindOrAdd(ServerInfo.InstanceInfo.InstanceId);
-	if (ListActiveSessionAsyncRequest.IsOngoing())
-	{
-		// Don't stack another async request on top of the ongoing one, TickServersAndSessionsDiscovery() will eventually catch up any possibly missed info here.
-		return;
-	}
-
-	// Arm the future, enabling its continuation to execute (or not).
-	TWeakPtr<uint8, ESPMode::ThreadSafe> ListActiveSessionExecutionToken = ListActiveSessionAsyncRequest.ResetExecutionToken();
-
-	// Keep track of the request time.
-	FDateTime ListRequestTimestamp = FDateTime::UtcNow();
-
-	// Retrieve all live sessions currently known by this server.
-	ListActiveSessionAsyncRequest.Future = ConcertClient->GetLiveSessions(ServerInfo.AdminEndpointId)
-		.Next([this, ServerInfo, ListActiveSessionExecutionToken, ListRequestTimestamp](const FConcertAdmin_GetSessionsResponse& Response)
-		{
-			// If the future is disarmed.
-			if (!ListActiveSessionExecutionToken.IsValid())
-			{
-				// Don't go further, the future execution was canceled, maybe because this object was deleted, the server removed from the list or it wasn't safe to get this future executing.
-				return;
-			}
-
-			// If the server responded.
-			if (Response.ResponseCode == EConcertResponseCode::Success)
-			{
-				bInitialActiveSessionQueryResponded = true;
-
-				// Remove from the cache any session that were deleted from 'Server' since the last update. Find which one were renamed.
-				ActiveSessions.RemoveAll([this, &ServerInfo, &Response](const TSharedPtr<FClientActiveSessionInfo>& DiscardCandidate)
-				{
-					// If the candidate is owned by another server.
-					if (DiscardCandidate->ServerInfo.InstanceInfo.InstanceId != ServerInfo.InstanceInfo.InstanceId)
-					{
-						return false; // Keep that session, it's owned by another server.
-					}
-					// If the candidate is still active.
-					else if (const FConcertSessionInfo* SessionFromServer = Response.Sessions.FindByPredicate([DiscardCandidate](const FConcertSessionInfo& MatchCandidate) { return DiscardCandidate->SessionInfo.SessionId == MatchCandidate.SessionId; }))
-					{
-						// Don't discard the session, it is still active, but check if it was renamed.
-						if (SessionFromServer->SessionName != DiscardCandidate->SessionInfo.SessionName)
-						{
-							DiscardCandidate->SessionInfo.SessionName = SessionFromServer->SessionName; // Update the session name.
-							OnActiveSessionRenamed(*DiscardCandidate, SessionFromServer->SessionName);
-						}
-						else if (DiscardCandidate->bSessionNameDirty) // Renaming this particular session failed?
-						{
-							OnActiveSessionListDirty(); // UI optimistically updated a session name from a rename, but it failed on the server. The controller caches the latest official server list unmodified. Force the client to resync its list against the cache.
-							DiscardCandidate->bSessionNameDirty = false;
-						}
-						return false; // Don't remove the session, it's still active.
-					}
-
-					// The session is not active anymore on the server, it was discarded.
-					OnActiveSessionDiscarded(*DiscardCandidate);
-					return true;
-				});
-
-				// Add to the cache any session that was added to 'Server' since the last update.
-				for (const FConcertSessionInfo& SessionInfo : Response.Sessions)
-				{
-					// Try to find the session in the list of active session.
-					if (!ActiveSessions.ContainsByPredicate([&SessionInfo](const TSharedPtr<FClientActiveSessionInfo>& MatchCandidate)
-						{ return SessionInfo.ServerInstanceId == MatchCandidate->ServerInfo.InstanceInfo.InstanceId && SessionInfo.SessionId == MatchCandidate->SessionInfo.SessionId; }))
-					{
-						// This is a newly discovered session, add it to the list.
-						ActiveSessions.Add(MakeShared<FClientActiveSessionInfo>(ServerInfo, SessionInfo));
-						OnActiveSessionDiscovered(*ActiveSessions.Last());
-					}
-
-					// Remove sessions that were created and discovered or those that should have shown up by now, but have not.
-					ExpectedSessionsToDiscover.RemoveAll([&SessionInfo, &ListRequestTimestamp](const FPendingSessionDiscovery& DiscardCandidate)
-					{
-						bool bDiscovered = DiscardCandidate.ServerEndpoint == SessionInfo.ServerEndpointId && DiscardCandidate.SessionName == SessionInfo.SessionName; // Session was discovered as expected.
-						bool bDeleted = ListRequestTimestamp > DiscardCandidate.CreateTimestamp; // The session was created successfully, a list request posted, but the session did not turn up. The session was deleted before it could be listed.
-						return bDiscovered || bDeleted;
-					});
-				}
-			}
-			// else -> The concert request failed, possibly because the server went offline. Wait until next TickServersAndSessionsDiscovery() to sync the server again and discover if it went offline.
-		});
-}
-
-void FConcertClientSessionBrowserController::UpdateArchivedSessionsAsync(const FConcertServerInfo& ServerInfo)
-{
-	// Check if a request is already pulling archived sessions from this server.
-	FAsyncRequest& ListArchivedSessionAsyncRequest = ArchivedSessionRequests.FindOrAdd(ServerInfo.InstanceInfo.InstanceId);
-	if (ListArchivedSessionAsyncRequest.IsOngoing())
-	{
-		// Don't stack another async request on top of the ongoing one, TickServersAndSessionsDiscovery() will eventually catch up any possibly missed info here.
-		return;
-	}
-
-	// Arm the future, enabling its continuation to execute.
-	TWeakPtr<uint8, ESPMode::ThreadSafe> ListArchivedSessionExecutionToken = ListArchivedSessionAsyncRequest.ResetExecutionToken();
-
-	// Retrieve the archived sessions.
-	ListArchivedSessionAsyncRequest.Future = ConcertClient->GetArchivedSessions(ServerInfo.AdminEndpointId)
-		.Next([this, ServerInfo, ListArchivedSessionExecutionToken](const FConcertAdmin_GetSessionsResponse& Response)
-		{
-			// If the future is disarmed.
-			if (!ListArchivedSessionExecutionToken.IsValid())
-			{
-				// Don't go further, the future execution was canceled, maybe because this object was deleted, the server removed from the list or it wasn't safe to get this future executing.
-				return;
-			}
-
-			// If the server responded.
-			if (Response.ResponseCode == EConcertResponseCode::Success)
-			{
-				bInitialArchivedSessionQueryResponded = true;
-
-				// Remove from the cache archives that were deleted from 'Server' since the last update. Find which one were renamed.
-				ArchivedSessions.RemoveAll([this, &ServerInfo, &Response](const TSharedPtr<FClientArchivedSessionInfo>& DiscardCandidate)
-				{
-					// If the discard candidate is stored on another server.
-					if (DiscardCandidate->ServerInfo.InstanceInfo.InstanceId != ServerInfo.InstanceInfo.InstanceId)
-					{
-						return false; // Keep that archive, it's stored on another server.
-					}
-					// If the archive is still stored on the server.
-					else if (const FConcertSessionInfo* SessionFromServer = Response.Sessions.FindByPredicate([DiscardCandidate](const FConcertSessionInfo& MatchCandidate) { return DiscardCandidate->SessionInfo.SessionId == MatchCandidate.SessionId; }))
-					{
-						// Don't discard the session, it is still there. Check if it was renamed.
-						if (SessionFromServer->SessionName != DiscardCandidate->SessionInfo.SessionName)
-						{
-							DiscardCandidate->SessionInfo.SessionName = SessionFromServer->SessionName; // Update the session name.
-							OnArchivedSessionRenamed(*DiscardCandidate, SessionFromServer->SessionName);
-						}
-						else if (DiscardCandidate->bSessionNameDirty) // Renaming this particular session failed?
-						{
-							OnArchivedSessionListDirty(); // UI optimistically updated a session name from a rename, but it failed on the server. The controller caches the latest official server list unmodified. Force the client to resync its list against the cache.
-							DiscardCandidate->bSessionNameDirty = false;
-						}
-
-						return false; // Don't remove the archive, it's still stored on the server.
-					}
-
-					OnArchivedSessionDiscarded(*DiscardCandidate);
-					return true; // The session is not archived anymore on 'Server' anymore, remove it from the list.
-				});
-
-				// Add to the cache archives that was stored on 'Server' since the last update.
-				for (const FConcertSessionInfo& SessionInfo : Response.Sessions)
-				{
-					// Try to find the archived in the list.
-					if (!ArchivedSessions.ContainsByPredicate([&SessionInfo, &ServerInfo](const TSharedPtr<FArchivedSessionInfo>& MatchCandidate)
-						{ return ServerInfo.InstanceInfo.InstanceId == MatchCandidate->ServerInfo.InstanceInfo.InstanceId && SessionInfo.SessionId == MatchCandidate->SessionInfo.SessionId; }))
-					{
-						// This is a newly discovered archive, add it to the list.
-						ArchivedSessions.Add(MakeShared<FClientArchivedSessionInfo>(ServerInfo, SessionInfo));
-						OnArchivedSessionDiscovered(*ArchivedSessions.Last());
-					}
-				}
-			}
-			// else -> Request failed, will discovered what happened on next TickServersAndSessionsDiscovery().
-		});
-}
-
-void FConcertClientSessionBrowserController::UpdateClientsAsync(const FGuid& ServerAdminEndpointId, const FGuid& SessionId)
-{
-	// The clients request is for a different session than the one cached by the last client update (The UI only shows 1 at the time).
-	if (ClientMonitoredSession.IsSet() && (ClientMonitoredSession->Session->ServerInfo.AdminEndpointId != ServerAdminEndpointId || ClientMonitoredSession->Session->SessionInfo.SessionId != SessionId))
-	{
-		// Reset the cache.
-		ClientMonitoredSession->Session->Clients.Reset();
-		ClientMonitoredSession->ListClientRequest.Cancel();
-		ClientMonitoredSession.Reset();
-		ClientListVersion = 0;
-	}
-
-	// Find the session requested by the user.
-	if (!ClientMonitoredSession.IsSet())
-	{
-		TSharedPtr<FClientActiveSessionInfo>* MatchEntry = ActiveSessions.FindByPredicate([&ServerAdminEndpointId, &SessionId](const TSharedPtr<FClientActiveSessionInfo>& MatchCandidate)
-		{
-			return MatchCandidate->ServerInfo.AdminEndpointId == ServerAdminEndpointId && MatchCandidate->SessionInfo.SessionId == SessionId;
-		});
-
-		if (MatchEntry != nullptr)
-		{
-			ClientMonitoredSession = FMonitoredSessionInfo { MatchEntry->ToSharedRef() };
-		}
-		else
-		{
-			ClientMonitoredSession.Reset();
-		}
-	}
-
-	if (!ClientMonitoredSession.IsSet() || ClientMonitoredSession->ListClientRequest.IsOngoing())
-	{
-		// Don't stack another async request on top of the on-going one, TickClientsDiscovery() will eventually catch up any possibly missed info here.
-		return;
-	}
-
-	// Arm the future, enabling its continuation to execute when called.
-	TWeakPtr<uint8, ESPMode::ThreadSafe> ListClientExecutionToken = ClientMonitoredSession->ListClientRequest.ResetExecutionToken();
-
-	// Retrieve (asynchronously) the clients corresponding to the selected server/session.
-	ClientMonitoredSession->ListClientRequest.Future = ConcertClient->GetSessionClients(ServerAdminEndpointId, SessionId)
-		.Next([this, ListClientExecutionToken](const FConcertAdmin_GetSessionClientsResponse& Response)
-		{
-			// If the future execution was canceled.
-			if (!ListClientExecutionToken.IsValid())
-			{
-				// Don't go further, the future execution was canceled.
-				return;
-			}
-
-			auto SortClientPredicate = [](const FConcertSessionClientInfo& Lhs, const FConcertSessionClientInfo& Rhs) { return Lhs.ClientEndpointId < Rhs.ClientEndpointId; };
-
-			// If the request succeeded.
-			if (Response.ResponseCode == EConcertResponseCode::Success)
-			{
-				if (ClientMonitoredSession->Session->Clients.Num() != Response.SessionClients.Num())
-				{
-					ClientMonitoredSession->Session->Clients = Response.SessionClients;
-					ClientMonitoredSession->Session->Clients.Sort(SortClientPredicate);
-					OnActiveSessionClientsUpdated(*ClientMonitoredSession->Session);
-					return;
-				}
-				else if (Response.SessionClients.Num() == 0) // All existing clients disconnected.
-				{
-					ClientMonitoredSession->Session->Clients.Reset();
-					OnActiveSessionClientsUpdated(*ClientMonitoredSession->Session);
-				}
-				else // Compare the old and the new list, both sorted by client endpoint id.
-				{
-					TArray<FConcertSessionClientInfo> SortedClients = Response.SessionClients;
-					SortedClients.Sort(SortClientPredicate);
-
-					int Index = 0;
-					for (const FConcertSessionClientInfo& Client : SortedClients)
-					{
-						if (ClientMonitoredSession->Session->Clients[Index].ClientEndpointId != Client.ClientEndpointId || // Not the same client?
-						    ClientMonitoredSession->Session->Clients[Index].ClientInfo != Client.ClientInfo)               // Client info was updated?
-						{
-							// The two lists are not identical, don't bother finding the 'delta', refresh all clients.
-							ClientMonitoredSession->Session->Clients = SortedClients;
-							OnActiveSessionClientsUpdated(*ClientMonitoredSession->Session);
-							break;
-						}
-						++Index;
-					}
-				}
-			}
-		});
-}
-
-/**
- * Widget displayed when discovering multi-user server(s) or session(s).
- */
-class SConcertDiscovery : public SCompoundWidget
-{
-public:
-	SLATE_BEGIN_ARGS(SConcertDiscovery)
-		: _Text()
-		, _ThrobberVisibility(EVisibility::Visible)
-		, _ButtonVisibility(EVisibility::Visible)
-		, _IsButtonEnabled(true)
-		, _ButtonStyle(&FAppStyle::Get().GetWidgetStyle<FButtonStyle>("FlatButton"))
-		, _ButtonIcon()
-		, _ButtonText()
-		, _ButtonToolTip()
-		, _OnButtonClicked()
-	{
-	}
-
-	SLATE_ATTRIBUTE(FText, Text)
-	SLATE_ATTRIBUTE(EVisibility, ThrobberVisibility)
-	SLATE_ATTRIBUTE(EVisibility, ButtonVisibility)
-	SLATE_ATTRIBUTE(bool, IsButtonEnabled)
-	SLATE_STYLE_ARGUMENT(FButtonStyle, ButtonStyle)
-	SLATE_ATTRIBUTE(const FSlateBrush*, ButtonIcon)
-	SLATE_ATTRIBUTE(FText, ButtonText)
-	SLATE_ATTRIBUTE(FText, ButtonToolTip)
-	SLATE_EVENT( FOnClicked, OnButtonClicked)
-
-	SLATE_END_ARGS();
-
-public:
-	void Construct(const FArguments& InArgs)
-	{
-		ChildSlot
-		[
-			SNew(SBox)
-			.HAlign(HAlign_Center)
-			.VAlign(VAlign_Center)
-			.Clipping(EWidgetClipping::ClipToBounds)
-			[
-				SNew(SVerticalBox)
-				+SVerticalBox::Slot()
-				.AutoHeight()
-				.HAlign(HAlign_Center)
-				[
-					SNew(SCircularThrobber)
-					.Visibility(InArgs._ThrobberVisibility)
-				]
-
-				+SVerticalBox::Slot()
-				.AutoHeight()
-				.HAlign(HAlign_Center)
-				[
-					SNew(STextBlock).Text(InArgs._Text).Justification(ETextJustify::Center)
-				]
-
-				+SVerticalBox::Slot()
-				.Padding(0, 4, 0, 0)
-				.AutoHeight()
-				.HAlign(HAlign_Center)
-				[
-					SNew(SButton)
-					.ForegroundColor(FSlateColor::UseForeground())
-					.ButtonStyle(InArgs._ButtonStyle)
-					.Visibility(InArgs._ButtonVisibility)
-					.HAlign(HAlign_Center)
-					.VAlign(VAlign_Center)
-					.IsEnabled(InArgs._IsButtonEnabled)
-					.OnClicked(InArgs._OnButtonClicked)
-					.ToolTipText(InArgs._ButtonToolTip)
-					.ContentPadding(FMargin(8, 4))
-					[
-						SNew(SHorizontalBox)
-						+SHorizontalBox::Slot()
-						.AutoWidth()
-						.VAlign(VAlign_Center)
-						.Padding(0, 0, 3, 0)
-						[
-							SNew(SImage).Image(InArgs._ButtonIcon)
-						]
-						+SHorizontalBox::Slot()
-						.AutoWidth()
-						.VAlign(VAlign_Bottom)
-						[
-							SNew(STextBlock).Text(InArgs._ButtonText)
-						]
-					]
-				]
-			]
-		];
-	}
-};
-
-
-/**
- * Displayed when something is not available.
- */
-class SConcertNoAvailability : public SCompoundWidget
-{
-public:
-	SLATE_BEGIN_ARGS(SConcertNoAvailability) : _Text() {}
-	SLATE_ATTRIBUTE(FText, Text)
-	SLATE_END_ARGS()
-
-public:
-	void Construct(const FArguments& InArgs)
-	{
-		ChildSlot
-		[
-			SNew(SConcertDiscovery) // Reuse this panel, but only show the message.
-				.Text(InArgs._Text)
-				.ThrobberVisibility(EVisibility::Collapsed)
-				.ButtonVisibility(EVisibility::Collapsed)
-		];
-	}
-};
-
-
-/**
- * Enables the user to browse/search/filter/sort active and archived sessions, create new session,
- * archive active sessions, restore archived sessions, join a session and open the settings dialog.
- */
-class SConcertClientSessionBrowser : public SCompoundWidget
-{
-public:
-	
-	SLATE_BEGIN_ARGS(SConcertClientSessionBrowser) { }
-	SLATE_END_ARGS();
-
-	/**
-	* Constructs the Browser.
-	* @param InArgs The Slate argument list.
-	* @param InConcertClient The concert client used to list, join, delete the sessions.
-	* @param[in,out] InSearchText The text to set in the search box and to remember (as output). Cannot be null.
-	*/
-	void Construct(const FArguments& InArgs, IConcertClientPtr InConcertClient, TSharedPtr<FText> InSearchText);
-
-	virtual void Tick(const FGeometry& AllottedGeometry, const double InCurrentTime, const float InDeltaTime) override;
-
-private:
-	
-	// Layout the 'session|details' split view.
-	TSharedRef<SWidget> MakeBrowserContent(TSharedPtr<FText> InSearchText);
-	void ExtendControlButtons(FExtender& Extender);
-	void ExtendSessionContextMenu(const TSharedPtr<FConcertSessionTreeItem>& Item, FExtender& Extender);
-	TSharedRef<SWidget> MakeUserAndSettings();
-	TSharedRef<SWidget> MakeOverlayedTableView(const TSharedRef<SWidget>& TableView);
-	
-	// Layouts the session detail panel.
-	TSharedRef<SWidget> MakeSessionDetails(TSharedPtr<FConcertSessionTreeItem> Item);
-	TSharedRef<SWidget> MakeActiveSessionDetails(TSharedPtr<FConcertSessionTreeItem> Item);
-	TSharedRef<SWidget> MakeArchivedSessionDetails(TSharedPtr<FConcertSessionTreeItem> Item);
-	TSharedRef<ITableRow> OnGenerateClientRowWidget(TSharedPtr<FConcertSessionClientInfo> Item, const TSharedRef<STableViewBase>& OwnerTable);
-	void PopulateSessionInfoGrid(SGridPanel& Grid, const FConcertSessionInfo& SessionInfo);
-
-	// The buttons in some overlays
-	bool IsLaunchServerButtonEnabled() const;
-	bool IsJoinButtonEnabled() const;
-	bool IsAutoJoinButtonEnabled() const;
-	bool IsCancelAutoJoinButtonEnabled() const;
-	FReply OnNewButtonClicked();
-	FReply OnLaunchServerButtonClicked();
-	FReply OnShutdownServerButtonClicked();
-	FReply OnAutoJoinButtonClicked();
-	FReply OnCancelAutoJoinButtonClicked();
-
-	// SConcertSessionBrowser extensions
-	FReply OnJoinButtonClicked();
-	void RequestJoinSession(const TSharedPtr<FConcertSessionTreeItem>& ActiveItem);
-
-	// Manipulates the sessions view (the array and the UI).
-	void OnSessionSelectionChanged(const TSharedPtr<FConcertSessionTreeItem>& SessionItem);
-	void OnSessionDoubleClicked(const TSharedPtr<FConcertSessionTreeItem>& SessionItem);
-	bool ConfirmDeleteSessionWithDialog(const TArray<TSharedPtr<FConcertSessionTreeItem>>& SessionItems) const;
-
-	// Update server/session/clients lists.
-	EActiveTimerReturnType TickDiscovery(double InCurrentTime, float InDeltaTime);
-	void UpdateDiscovery();
-	void RefreshClientList(const TArray<FConcertSessionClientInfo>& LastestClientList);
-
-private:
-	
-	// Gives access to the concert data (servers, sessions, clients, etc).
-	TSharedPtr<FConcertClientSessionBrowserController> Controller;
-
-	TSharedPtr<SConcertSessionBrowser> SessionBrowser;
-
-	// Filtering.
-	bool bRefreshSessionFilter = true;
-	FString DefaultServerURL;
-
-	// Selected Session Details.
-	TSharedPtr<SBorder> SessionDetailsView;
-	TSharedPtr<SExpandableArea> DetailsArea;
-	TArray<TSharedPtr<FConcertSessionClientInfo>> Clients;
-	TSharedPtr<SExpandableArea> ClientsArea;
-	TSharedPtr<SListView<TSharedPtr<FConcertSessionClientInfo>>> ClientsView;
-
-	// Used to compare the version used by UI versus the version cached in the controller.
-	uint32 DisplayedSessionListVersion = 0;
-	uint32 DisplayedClientListVersion = 0;
-	uint32 ServerListVersion = 0;
-	bool bLocalServerRunning = false;
-
-	TSharedPtr<SWidget> ServerDiscoveryPanel; // Displayed until a server is found.
-	TSharedPtr<SWidget> SessionDiscoveryPanel; // Displayed until a session is found.
-	TSharedPtr<SWidget> NoSessionSelectedPanel; // Displays 'select a session to view details' message in the details section.
-	TSharedPtr<SWidget> NoSessionDetailsPanel; // Displays 'no details available' message in the details section.
-	TSharedPtr<SWidget> NoClientPanel; // Displays the 'no client connected' message in Clients expendable area.
-};
-
-void SConcertClientSessionBrowser::Construct(const FArguments& InArgs, IConcertClientPtr InConcertClient, TSharedPtr<FText> InSearchText)
-{
-	if (!InConcertClient.IsValid())
-	{
-		return; // Don't build the UI if ConcertClient is not available.
-	}
-
-	Controller = MakeShared<FConcertClientSessionBrowserController>(InConcertClient);
-
-	// Displayed if no server is available.
-	ServerDiscoveryPanel = SNew(SConcertDiscovery)
-		.Text(LOCTEXT("LookingForServer", "Looking for Multi-User Servers..."))
-		.Visibility_Lambda([this]() { return Controller->GetServers().Num() == 0 ? EVisibility::Visible : EVisibility::Hidden; })
-		.IsButtonEnabled(this, &SConcertClientSessionBrowser::IsLaunchServerButtonEnabled)
-		.ButtonStyle(FAppStyle::Get(), "FlatButton.Default")
-		.ButtonIcon(FConcertFrontendStyle::Get()->GetBrush("Concert.NewServer.Small"))
-		.ButtonText(LOCTEXT("LaunchLocalServer", "Launch a Server"))
-		.ButtonToolTip(LOCTEXT("LaunchServerTooltip", "Launch a Multi-User server on your computer unless one is already running.\nThe editor UDP messaging settings will be passed to the launching server.\nThe server will use the server port found in Multi-User client settings if non 0 or use the editor udp messaging port number + 1, if non 0."))
-		.OnButtonClicked(this, &SConcertClientSessionBrowser::OnLaunchServerButtonClicked);
-
-	// Controls the text displayed in the 'No sessions' panel.
-	auto GetNoSessionText = [this]()
-	{
-		if (!Controller->HasReceivedInitialSessionList())
-		{
-			return LOCTEXT("LookingForSession", "Looking for Multi-User Sessions...");
-		}
-
-		return Controller->GetActiveSessions().Num() == 0 && Controller->GetArchivedSessions().Num() == 0 ?
-			LOCTEXT("NoSessionAvailable", "No Sessions Available") :
-			LOCTEXT("AllSessionsFilteredOut", "No Sessions Match the Filters\nChange Your Filter to View Sessions");
-	};
-
-	// Displayed when discovering session or if no session is available.
-	SessionDiscoveryPanel = SNew(SConcertDiscovery)
-		.Text_Lambda(GetNoSessionText)
-		.Visibility_Lambda([this]() { return Controller->GetServers().Num() > 0 && !SessionBrowser->HasAnySessions() && !Controller->IsCreatingSession() ? EVisibility::Visible : EVisibility::Hidden; })
-		.ThrobberVisibility_Lambda([this]() { return !Controller->HasReceivedInitialSessionList() ? EVisibility::Visible : EVisibility::Collapsed; })
-		.ButtonVisibility_Lambda([this]() { return Controller->HasReceivedInitialSessionList() && Controller->GetActiveSessions().Num() == 0 && Controller->GetArchivedSessions().Num() == 0 ? EVisibility::Visible : EVisibility::Collapsed; })
-		.ButtonStyle(FAppStyle::Get(), "FlatButton.Default")
-		.ButtonIcon(FConcertFrontendStyle::Get()->GetBrush("Concert.NewSession.Small"))
-		.ButtonText(LOCTEXT("CreateSession", "Create Session"))
-		.ButtonToolTip(LOCTEXT("CreateSessionTooltip", "Create a new session"))
-		.OnButtonClicked(FOnClicked::CreateSP(this, &SConcertClientSessionBrowser::OnNewButtonClicked));
-
-	// Displayed when the selected session client view is empty (no client to display).
-	NoClientPanel = SNew(SConcertNoAvailability)
-		.Text(LOCTEXT("NoClientAvailable", "No Connected Clients"))
-		.Visibility_Lambda([this]() { return Clients.Num() == 0 ? EVisibility::Visible : EVisibility::Hidden; });
-
-	// Displayed as details when no session is selected. (No session selected or the selected session doesn't have any)
-	NoSessionSelectedPanel = SNew(SConcertNoAvailability)
-		.Text(LOCTEXT("NoSessionSelected", "Select a Session to View Details"));
-
-	// Displayed as details when the selected session has not specific details to display.
-	NoSessionDetailsPanel = SNew(SConcertNoAvailability)
-		.Text(LOCTEXT("NoSessionDetails", "The Selected Session Has No Details"));
-
-	// List used in details panel to display clients connected to an active session.
-	ClientsView = SNew(SListView<TSharedPtr<FConcertSessionClientInfo>>)
-		.ListItemsSource(&Clients)
-		.OnGenerateRow(this, &SConcertClientSessionBrowser::OnGenerateClientRowWidget)
-		.SelectionMode(ESelectionMode::Single)
-		.AllowOverscroll(EAllowOverscroll::No);
-
-	ChildSlot
-	[
-		MakeBrowserContent(InSearchText)
-	];
-
-	// Create a timer to periodically poll the server for sessions and session clients at a lower frequency than the normal tick.
-	RegisterActiveTimer(1.0f, FWidgetActiveTimerDelegate::CreateSP(this, &SConcertClientSessionBrowser::TickDiscovery));
-
-	bLocalServerRunning = IMultiUserClientModule::Get().IsConcertServerRunning();
-}
-
-TSharedRef<SWidget> SConcertClientSessionBrowser::MakeBrowserContent(TSharedPtr<FText> InSearchText)
-{
-	return SNew(SBox)
-		.HAlign(HAlign_Fill)
-		.VAlign(VAlign_Fill)
-		[
-			// Splitter upper part displaying the available sessions/(server).
-			SNew(SSplitter)
-			.Orientation(Orient_Vertical)
-			.MinimumSlotHeight(80.0f) // Prevent widgets from overlapping.
-			+SSplitter::Slot()
-			.Value(0.6)
-			[
-				SAssignNew(SessionBrowser, SConcertSessionBrowser, Controller.ToSharedRef(), InSearchText)
-					.ExtendSessionTable(this, &SConcertClientSessionBrowser::MakeOverlayedTableView)
-					.ExtendControllButtons(this, &SConcertClientSessionBrowser::ExtendControlButtons)
-					.ExtendSessionContextMenu(this, &SConcertClientSessionBrowser::ExtendSessionContextMenu)
-					.RightOfControlButtons()
-					[
-						MakeUserAndSettings()
-					]
-					.OnSessionClicked(this, &SConcertClientSessionBrowser::OnSessionSelectionChanged)
-					.OnLiveSessionDoubleClicked(this, &SConcertClientSessionBrowser::OnSessionDoubleClicked)
-					.PostRequestedDeleteSession_Lambda([this](auto) { UpdateDiscovery(); /* Don't wait up to 1s, kick discovery right now */ })
-					.AskUserToDeleteSessions(this, &SConcertClientSessionBrowser::ConfirmDeleteSessionWithDialog)
-			]
-
-			// Session details.
-			+SSplitter::Slot()
-			.Value(0.4)
-			[
-				SAssignNew(SessionDetailsView, SBorder)
-				.BorderImage(FAppStyle::Get().GetBrush("ToolPanel.GroupBorder"))
-				[
-					NoSessionSelectedPanel.ToSharedRef()
-				]
-			]
-		];
-}
-
-void SConcertClientSessionBrowser::ExtendControlButtons(FExtender& Extender)
-{
-	// Before separator
-	Extender.AddToolBarExtension(
-		SConcertSessionBrowser::ControlButtonExtensionHooks::BeforeSeparator,
-		EExtensionHook::First,
-		nullptr,
-		FToolBarExtensionDelegate::CreateLambda([this](FToolBarBuilder& Builder)
-		{
-			Builder.AddWidget(
-				SNew(SOverlay)
-				// Launch server.
-				+SOverlay::Slot()
-				[
-					ConcertBrowserUtils::MakeIconButton(
-						FConcertFrontendStyle::Get()->GetBrush("Concert.NewServer"),
-						LOCTEXT("LaunchServerTooltip", "Launch a Multi-User server on your computer unless one is already running.\nThe editor UDP messaging settings will be passed to the launching server.\nThe server will use the server port found in Multi-User client settings if non 0 or use the editor udp messaging port number + 1, if non 0."),
-						TAttribute<bool>(this, &SConcertClientSessionBrowser::IsLaunchServerButtonEnabled),
-						FOnClicked::CreateSP(this, &SConcertClientSessionBrowser::OnLaunchServerButtonClicked),
-						TAttribute<EVisibility>::Create([this]() { return IsLaunchServerButtonEnabled() ? EVisibility::Visible : EVisibility::Collapsed; }))
-				]
-				// Stop server.
-				+SOverlay::Slot()
-				[
-					ConcertBrowserUtils::MakeIconButton(
-						FConcertFrontendStyle::Get()->GetBrush("Concert.CloseServer"),
-						LOCTEXT("ShutdownServerTooltip", "Shutdown the Multi-User server running on this computer."),
-						true, // Always enabled, but might be collapsed.
-						FOnClicked::CreateSP(this, &SConcertClientSessionBrowser::OnShutdownServerButtonClicked),
-						TAttribute<EVisibility>::Create([this]() { return IsLaunchServerButtonEnabled() ? EVisibility::Collapsed : EVisibility::Visible; }))
-				]
-				);
-		}));
-
-	// After separator
-	Extender.AddToolBarExtension(
-		SConcertSessionBrowser::ControlButtonExtensionHooks::AfterSeparator,
-		EExtensionHook::First,
-		nullptr,
-		FToolBarExtensionDelegate::CreateLambda([this](FToolBarBuilder& Builder)
-		{
-			TAttribute<FText> AutoJoinTooltip = TAttribute<FText>::Create([this]()
-			{
-				if (Controller->GetConcertClient()->CanAutoConnect()) // Default session and server are configured?
-				{
-					return FText::Format(LOCTEXT("JoinDefaultSessionTooltip", "Join the default session '{0}' on '{1}'"),
-						FText::FromString(Controller->GetConcertClient()->GetConfiguration()->DefaultSessionName),
-						FText::FromString(Controller->GetConcertClient()->GetConfiguration()->DefaultServerURL));
-				}
-				else
-				{
-					return LOCTEXT("JoinDefaultSessionConfiguredTooltip", "Join the default session configured in the Multi-Users settings");
-				}
-			});
-
-			TAttribute<FText> CancelAutoJoinTooltip = TAttribute<FText>::Create([this]()
-			{
-				return FText::Format(LOCTEXT("CancelJoinDefaultSessionTooltip", "Cancel joining the default session '{0}' on '{1}'"),
-					FText::FromString(Controller->GetConcertClient()->GetConfiguration()->DefaultSessionName),
-					FText::FromString(Controller->GetConcertClient()->GetConfiguration()->DefaultServerURL));
-			});
-			
-			Builder.AddWidget(
-				SNew(SOverlay)
-				// Auto-join
-				+SOverlay::Slot()
-				[
-					ConcertBrowserUtils::MakeIconButton(
-						FConcertFrontendStyle::Get()->GetBrush("Concert.JoinDefaultSession"),
-						AutoJoinTooltip,
-						TAttribute<bool>(this, &SConcertClientSessionBrowser::IsAutoJoinButtonEnabled),
-						FOnClicked::CreateSP(this, &SConcertClientSessionBrowser::OnAutoJoinButtonClicked),
-						// Default button shown if both auto-join/cancel are disabled.
-						TAttribute<EVisibility>::Create([this]() { return !IsCancelAutoJoinButtonEnabled() ? EVisibility::Visible : EVisibility::Collapsed; })
-						) 
-				
-				]
-				// Cancel auto join.
-				+SOverlay::Slot()
-				[
-					ConcertBrowserUtils::MakeIconButton(
-						FConcertFrontendStyle::Get()->GetBrush("Concert.CancelAutoJoin"),
-						CancelAutoJoinTooltip,
-						TAttribute<bool>(this, &SConcertClientSessionBrowser::IsCancelAutoJoinButtonEnabled),
-						FOnClicked::CreateSP(this, &SConcertClientSessionBrowser::OnCancelAutoJoinButtonClicked),
-						TAttribute<EVisibility>::Create([this]() { return IsCancelAutoJoinButtonEnabled() ? EVisibility::Visible : EVisibility::Collapsed; })
-						)
-				]
-			);
-
-			// Join
-			Builder.AddWidget(
-				ConcertBrowserUtils::MakeIconButton(FConcertFrontendStyle::Get()->GetBrush("Concert.JoinSession"), LOCTEXT("JoinButtonTooltip", "Join the selected session"),
-					TAttribute<bool>(this, &SConcertClientSessionBrowser::IsJoinButtonEnabled),
-					FOnClicked::CreateSP(this, &SConcertClientSessionBrowser::OnJoinButtonClicked),
-					TAttribute<EVisibility>::Create([this]() { return !SessionBrowser->IsRestoreButtonEnabled() ? EVisibility::Visible : EVisibility::Collapsed; })) // Default button shown if both join/restore are disabled.
-			);
-		}));
-}
-
-void SConcertClientSessionBrowser::ExtendSessionContextMenu(const TSharedPtr<FConcertSessionTreeItem>& Item, FExtender& Extender)
-{
-	if (Item->Type == FConcertSessionTreeItem::EType::ActiveSession)
-	{
-		Extender.AddMenuExtension(
-			SConcertSessionBrowser::SessionContextMenuExtensionHooks::ManageSession,
-			EExtensionHook::First,
-			nullptr,
-			FMenuExtensionDelegate::CreateLambda([this](FMenuBuilder& Builder)
-			{
-				Builder.AddMenuEntry(
-					LOCTEXT("CtxMenuJoin", "Join"),
-					LOCTEXT("CtxMenuJoin_Tooltip", "Join the Session"),
-					FSlateIcon(),
-					FUIAction(
-						FExecuteAction::CreateLambda([this](){ OnJoinButtonClicked(); }),
-						FCanExecuteAction::CreateLambda([SelectedCount = SessionBrowser->GetSelectedItems().Num()] { return SelectedCount == 1; }),
-						FIsActionChecked::CreateLambda([this] { return false; })),
-					NAME_None,
-					EUserInterfaceActionType::Button
-					);
-			}));
-	}
-	
-}
-
-TSharedRef<SWidget> SConcertClientSessionBrowser::MakeUserAndSettings()
-{
-	return SNew(SHorizontalBox)
-		// The user "Avatar color" displayed as a small square colored by the user avatar color.
-		+SHorizontalBox::Slot()
-		.VAlign(VAlign_Center)
-		.AutoWidth()
-		[
-			SNew(SImage)
-			.ColorAndOpacity_Lambda([this]() { return Controller->GetConcertClient()->GetClientInfo().AvatarColor; })
-			.Image(FAppStyle::GetBrush("Icons.FilledCircle"))
-		]
-				
-		// The user "Display Name".
-		+SHorizontalBox::Slot()
-		.VAlign(VAlign_Center)
-		.AutoWidth()
-		.Padding(3, 0, 2, 0)
-		[
-			SNew(STextBlock)
-				.Text_Lambda([this]() { return FText::FromString(Controller->GetConcertClient()->GetClientInfo().DisplayName);} )
-		]
-
-		// The "Settings" icons.
-		+SHorizontalBox::Slot()
-		.VAlign(VAlign_Fill)
-		.AutoWidth()
-		[
-			SNew(SButton)
-			.ButtonStyle(FAppStyle::Get(), TEXT("SimpleButton"))
-			.OnClicked_Lambda([](){FModuleManager::GetModulePtr<ISettingsModule>("Settings")->ShowViewer("Project", "Plugins", "Concert"); return FReply::Handled(); })
-			.VAlign(VAlign_Center)
-			.HAlign(HAlign_Fill)
-			[
-				SNew(SImage)
-				.Image(FAppStyle::GetBrush("Icons.Settings"))
-			]
-		];
-}
-
-TSharedRef<SWidget> SConcertClientSessionBrowser::MakeOverlayedTableView(const TSharedRef<SWidget>& SessionTable)
-{
-	return SNew(SOverlay)
-		+SOverlay::Slot()
-		[
-			SessionTable
-		]
-		+SOverlay::Slot()
-		.Padding(0, 20, 0, 0) // To ensure the panel is below the row header.
-		[
-			SessionDiscoveryPanel.ToSharedRef()
-		]
-		+SOverlay::Slot()
-		.Padding(0, 20, 0, 0) // To ensure the panel is below the row header.
-		[
-			ServerDiscoveryPanel.ToSharedRef()
-		];
-}
-
-EActiveTimerReturnType SConcertClientSessionBrowser::TickDiscovery(double InCurrentTime, float InDeltaTime)
-{
-	// Cache the result of this function because it is very expensive. It kills the framerate if polled every frame.
-	bLocalServerRunning = IMultiUserClientModule::Get().IsConcertServerRunning();
-
-	UpdateDiscovery();
-	return EActiveTimerReturnType::Continue;
-}
-
-void SConcertClientSessionBrowser::UpdateDiscovery()
-{
-	// Check if the controller has updated data since the last 'tick' and fire new asynchronous requests for future 'tick'.
-	TPair<uint32, uint32> ServerSessionListVersions = Controller->TickServersAndSessionsDiscovery();
-	ServerListVersion = ServerSessionListVersions.Key;
-
-	if (ServerSessionListVersions.Value != DisplayedSessionListVersion) // Need to refresh the list?
-	{
-		SessionBrowser->RefreshSessionList();
-		DisplayedSessionListVersion = ServerSessionListVersions.Value;
-	}
-
-	// If an active session is selected.
-	TArray<TSharedPtr<FConcertSessionTreeItem>> SelectedSession = SessionBrowser->GetSelectedItems();
-	if (SelectedSession.Num() && SelectedSession[0]->Type == FConcertSessionTreeItem::EType::ActiveSession)
-	{
-		// Ensure to poll its clients.
-		uint32 CachedClientListVersion = Controller->TickClientsDiscovery(SelectedSession[0]->ServerAdminEndpointId, SelectedSession[0]->SessionId);
-		if (CachedClientListVersion != DisplayedClientListVersion) // Need to refresh the list?
-		{
-			RefreshClientList(Controller->GetClients(SelectedSession[0]->ServerAdminEndpointId, SelectedSession[0]->SessionId));
-			DisplayedClientListVersion = CachedClientListVersion;
-		}
-	}
-}
-
-void SConcertClientSessionBrowser::RefreshClientList(const TArray<FConcertSessionClientInfo>& LastestClientList)
-{
-	// Remember which client is selected.
-	TArray<TSharedPtr<FConcertSessionClientInfo>> SelectedItems = ClientsView->GetSelectedItems();
-
-	// Copy the list of clients.
-	TArray<TSharedPtr<FConcertSessionClientInfo>> LatestClientPtrs;
-	Algo::Transform(LastestClientList, LatestClientPtrs, [](const FConcertSessionClientInfo& Client) { return MakeShared<FConcertSessionClientInfo>(Client); });
-
-	// Merge the current list with the new list, removing client that disappeared and adding client that appeared.
-	ConcertFrontendUtils::SyncArraysByPredicate(Clients, MoveTemp(LatestClientPtrs), [](const TSharedPtr<FConcertSessionClientInfo>& ClientToFind)
-	{
-		return [ClientToFind](const TSharedPtr<FConcertSessionClientInfo>& PotentialClientMatch)
-		{
-			return PotentialClientMatch->ClientEndpointId == ClientToFind->ClientEndpointId && PotentialClientMatch->ClientInfo == ClientToFind->ClientInfo;
-		};
-	});
-
-	// Sort the list item alphabetically.
-	Clients.StableSort([](const TSharedPtr<FConcertSessionClientInfo>& Lhs, const TSharedPtr<FConcertSessionClientInfo>& Rhs)
-	{
-		return Lhs->ClientInfo.DisplayName < Rhs->ClientInfo.DisplayName;
-	});
-
-	// Preserve previously selected item (if any).
-	if (SelectedItems.Num())
-	{
-		ClientsView->SetSelection(SelectedItems[0]);
-	}
-	ClientsView->RequestListRefresh();
-}
-
-TSharedRef<SWidget> SConcertClientSessionBrowser::MakeSessionDetails(TSharedPtr<FConcertSessionTreeItem> Item)
-{
-	const FConcertSessionInfo* SessionInfo = nullptr;
-
-	if (Item.IsValid())
-	{
-		if (Item->Type == FConcertSessionTreeItem::EType::ActiveSession || Item->Type == FConcertSessionTreeItem::EType::SaveSession)
-		{
-			return MakeActiveSessionDetails(Item);
-		}
-		else if (Item->Type == FConcertSessionTreeItem::EType::ArchivedSession)
-		{
-			return MakeArchivedSessionDetails(Item);
-		}
-	}
-
-	return NoSessionSelectedPanel.ToSharedRef();
-}
-
-TSharedRef<SWidget> SConcertClientSessionBrowser::MakeActiveSessionDetails(TSharedPtr<FConcertSessionTreeItem> Item)
-{
-	const TOptional<FConcertSessionInfo> SessionInfo = Controller->GetActiveSessionInfo(Item->ServerAdminEndpointId, Item->SessionId);
-	if (!SessionInfo)
-	{
-		return NoSessionSelectedPanel.ToSharedRef();
-	}
-
-	TSharedPtr<SGridPanel> Grid;
-
-	// State variables captured and shared by the different lambda functions below.
-	TSharedPtr<bool> bDetailsAreaExpanded = MakeShared<bool>(false);
-	TSharedPtr<bool> bClientsAreaExpanded = MakeShared<bool>(true);
-
-	auto DetailsAreaSizeRule = [this, bDetailsAreaExpanded]()
-	{
-		return *bDetailsAreaExpanded ? SSplitter::ESizeRule::FractionOfParent : SSplitter::ESizeRule::SizeToContent;
-	};
-
-	auto OnDetailsAreaExpansionChanged = [bDetailsAreaExpanded](bool bExpanded)
-	{
-		*bDetailsAreaExpanded = bExpanded;
-	};
-
-	auto ClientsAreaSizeRule = [this, bClientsAreaExpanded]()
-	{
-		return *bClientsAreaExpanded ? SSplitter::ESizeRule::FractionOfParent : SSplitter::ESizeRule::SizeToContent;
-	};
-
-	auto OnClientsAreaExpansionChanged = [bClientsAreaExpanded](bool bExpanded)
-	{
-		*bClientsAreaExpanded = bExpanded;
-	};
-
-	TSharedRef<SSplitter> Widget = SNew(SSplitter)
-		.Orientation(Orient_Vertical)
-
-		// Details.
-		+SSplitter::Slot()
-		.SizeRule(TAttribute<SSplitter::ESizeRule>::Create(DetailsAreaSizeRule))
-		.Value(0.6)
-		[
-			SAssignNew(DetailsArea, SExpandableArea)
-			.BorderBackgroundColor(FLinearColor(0.6f, 0.6f, 0.6f, 1.0f))
-			.BorderImage_Lambda([this]() { return ConcertFrontendUtils::GetExpandableAreaBorderImage(*DetailsArea); })
-			.BodyBorderImage(FAppStyle::Get().GetBrush("ToolPanel.GroupBorder"))
-			.BodyBorderBackgroundColor(FLinearColor::White)
-			.OnAreaExpansionChanged(FOnBooleanValueChanged::CreateLambda(OnDetailsAreaExpansionChanged))
-			.InitiallyCollapsed(!(*bDetailsAreaExpanded))
-			.HeaderContent()
-			[
-				SNew(STextBlock)
-				.Text(FText::Format(LOCTEXT("Details", "Details"), FText::FromString(Item->SessionName)))
-				.Font(FAppStyle::Get().GetFontStyle("DetailsView.CategoryFontStyle"))
-				.ShadowOffset(FVector2D(1.0f, 1.0f))
-			]
-			.BodyContent()
-			[
-				SNew(SScrollBox)
-				.Orientation(Orient_Vertical)
-				+SScrollBox::Slot()
-				[
-					SNew(SBox)
-					.Padding(FMargin(0, 2, 0, 2))
-					[
-						SAssignNew(Grid, SGridPanel)
-					]
-				]
-			]
-		]
-
-		// Clients
-		+SSplitter::Slot()
-		.SizeRule(TAttribute<SSplitter::ESizeRule>::Create(ClientsAreaSizeRule))
-		.Value(0.4)
-		[
-			SAssignNew(ClientsArea, SExpandableArea)
-			.BorderBackgroundColor(FLinearColor(0.6f, 0.6f, 0.6f, 1.0f))
-			.BorderImage_Lambda([this]() { return ConcertFrontendUtils::GetExpandableAreaBorderImage(*ClientsArea); })
-			.BodyBorderImage(FAppStyle::Get().GetBrush("ToolPanel.GroupBorder"))
-			.BodyBorderBackgroundColor(FLinearColor::White)
-			.OnAreaExpansionChanged(FOnBooleanValueChanged::CreateLambda(OnClientsAreaExpansionChanged))
-			.InitiallyCollapsed(!(*bClientsAreaExpanded))
-			.HeaderContent()
-			[
-				SNew(STextBlock)
-				.Text(LOCTEXT("Clients", "Clients"))
-				.Font(FAppStyle::Get().GetFontStyle("DetailsView.CategoryFontStyle"))
-				.ShadowOffset(FVector2D(1.0f, 1.0f))
-			]
-			.BodyContent()
-			[
-				SNew(SOverlay)
-				+SOverlay::Slot()
-				[
-					ClientsView.ToSharedRef()
-				]
-				+SOverlay::Slot()
-				[
-					NoClientPanel.ToSharedRef()
-				]
-			]
-		];
-
-	// Fill the grid.
-	PopulateSessionInfoGrid(*Grid, *SessionInfo);
-
-	// Populate the client list.
-	RefreshClientList(Controller->GetClients(Item->ServerAdminEndpointId, Item->SessionId));
-
-	return Widget;
-}
-
-TSharedRef<SWidget> SConcertClientSessionBrowser::MakeArchivedSessionDetails(TSharedPtr<FConcertSessionTreeItem> Item)
-{
-	const TOptional<FConcertSessionInfo> SessionInfo = Controller->GetArchivedSessionInfo(Item->ServerAdminEndpointId, Item->SessionId);
-	if (!SessionInfo)
-	{
-		return NoSessionSelectedPanel.ToSharedRef();
-	}
-
-	TSharedPtr<SGridPanel> Grid;
-
-	TSharedRef<SExpandableArea> Widget = SAssignNew(DetailsArea, SExpandableArea)
-	.BorderBackgroundColor(FLinearColor(0.6f, 0.6f, 0.6f, 1.0f))
-	.BorderImage_Lambda([this]() { return ConcertFrontendUtils::GetExpandableAreaBorderImage(*DetailsArea); })
-	.BodyBorderImage(FAppStyle::Get().GetBrush("ToolPanel.GroupBorder"))
-	.BodyBorderBackgroundColor(FLinearColor::White)
-	.InitiallyCollapsed(true)
-	.HeaderContent()
-	[
-		SNew(STextBlock)
-		.Text(FText::Format(LOCTEXT("Details", "Details"), FText::FromString(Item->SessionName)))
-		.Font(FAppStyle::Get().GetFontStyle("DetailsView.CategoryFontStyle"))
-		.ShadowOffset(FVector2D(1.0f, 1.0f))
-	]
-	.BodyContent()
-	[
-		SNew(SScrollBox)
-		.Orientation(Orient_Vertical)
-		+SScrollBox::Slot()
-		[
-			SNew(SBox)
-			.Padding(FMargin(0, 2, 0, 2))
-			[
-				SAssignNew(Grid, SGridPanel)
-			]
-		]
-	];
-
-	// Fill the grid.
-	PopulateSessionInfoGrid(*Grid, *SessionInfo);
-
-	return Widget;
-}
-
-void SConcertClientSessionBrowser::PopulateSessionInfoGrid(SGridPanel& Grid, const FConcertSessionInfo& SessionInfo)
-{
-	// Local function to populate the details grid.
-	auto AddDetailRow = [](SGridPanel& Grid, int32 Row, const FText& Label, const FText& Value)
-	{
-		const float RowPadding = (Row == 0) ? 0.0f : 4.0f; // Space between line.
-		const float ColPadding = 4.0f; // Space between columns. (Minimum)
-
-		Grid.AddSlot(0, Row)
-		.Padding(0.0f, RowPadding, ColPadding, 0.0f)
-		[
-			SNew(STextBlock).Text(Label)
-		];
-
-		Grid.AddSlot(1, Row)
-		.Padding(0.0f, RowPadding, 0.0f, 0.0f)
-		[
-			SNew(STextBlock).Text(Value)
-		];
-	};
-
-	int32 Row = 0;
-	AddDetailRow(Grid, Row++, LOCTEXT("SessionId", "Session ID:"), FText::FromString(SessionInfo.SessionId.ToString()));
-	AddDetailRow(Grid, Row++, LOCTEXT("SessionName", "Session Name:"), FText::FromString(SessionInfo.SessionName));
-	AddDetailRow(Grid, Row++, LOCTEXT("Owner", "Owner:"), FText::FromString(SessionInfo.OwnerUserName));
-	AddDetailRow(Grid, Row++, LOCTEXT("Project", "Project:"), FText::FromString(SessionInfo.Settings.ProjectName));
-	if (SessionInfo.VersionInfos.Num() > 0)
-	{
-		const FConcertSessionVersionInfo& VersionInfo = SessionInfo.VersionInfos.Last();
-		AddDetailRow(Grid, Row++, LOCTEXT("EngineVersion", "Engine Version:"), VersionInfo.AsText());
-	}
-	AddDetailRow(Grid, Row++, LOCTEXT("ServerEndPointId", "Server Endpoint ID:"), FText::FromString(SessionInfo.ServerEndpointId.ToString()));
-}
-
-TSharedRef<ITableRow> SConcertClientSessionBrowser::OnGenerateClientRowWidget(TSharedPtr<FConcertSessionClientInfo> Item, const TSharedRef<STableViewBase>& OwnerTable)
-{
-	return SNew(STableRow<TSharedPtr<FConcertSessionClientInfo>>, OwnerTable)
-	[
-		SNew(SHorizontalBox)
-		.ToolTipText(Item->ToDisplayString())
-
-		// The user "Avatar color" displayed as a small square colored by the user avatar color.
-		+SHorizontalBox::Slot()
-		.VAlign(VAlign_Center)
-		.AutoWidth()
-		[
-			SNew(SImage)
-			.ColorAndOpacity(Item->ClientInfo.AvatarColor)
-			.Image(FAppStyle::GetBrush("Icons.FilledCircle"))
-		]
-
-		// The user "Display Name".
-		+SHorizontalBox::Slot()
-		.VAlign(VAlign_Center)
-		.AutoWidth()
-		.Padding(FMargin(4.0, 2.0))
-		[
-			SNew(STextBlock)
-			.Text(FText::FromString(Item->ClientInfo.DisplayName))
-		]
-	];
-}
-
-void SConcertClientSessionBrowser::OnSessionSelectionChanged(const TSharedPtr<FConcertSessionTreeItem>& SelectedSession)
-{
-	// Clear the list of clients (if any)
-	Clients.Reset();
-
-	// Update the details panel.
-	SessionDetailsView->SetContent(MakeSessionDetails(SelectedSession));
-}
-
-void SConcertClientSessionBrowser::OnSessionDoubleClicked(const TSharedPtr<FConcertSessionTreeItem>& SelectedSession)
-{
-	switch (SelectedSession->Type)
-	{
-	case FConcertSessionTreeItem::EType::ActiveSession: 
-		RequestJoinSession(SelectedSession);
-		break;
-	case FConcertSessionTreeItem::EType::ArchivedSession:
-		SessionBrowser->InsertRestoreSessionAsEditableRow(SelectedSession);
-		break;
-	default: 
-		checkNoEntry();
-	}
-}
-
-bool SConcertClientSessionBrowser::ConfirmDeleteSessionWithDialog(const TArray<TSharedPtr<FConcertSessionTreeItem>>& SessionItems) const
-{
-	TSet<FString> SessionNames;
-	TSet<FString> UniqueServers;
-
-	Algo::Transform(SessionItems, Algo::TieTupleAdd(SessionNames, UniqueServers), [](const TSharedPtr<FConcertSessionTreeItem>& Item)
-	{
-		return MakeTuple(Item->SessionName, Item->ServerName);
-	});
-	
-	const FText ConfirmationMessage = FText::Format(
-		LOCTEXT("DeleteSessionConfirmationMessage", "Do you really want to delete {0} {0}|plural(one=session,other=sessions) from {1} {1}|plural(one=server,other=servers)?\n\nThe {0}|plural(one=session,other=sessions) {2} will be deleted from the {1}|plural(one=server,other=servers) {3}."),
-		SessionItems.Num(),
-		UniqueServers.Num(),
-		FText::FromString(FString::JoinBy(SessionNames, TEXT(", "), [](const FString& Name) { return FString("\n\t- ") + Name; }) + FString("\n")),
-		FText::FromString(FString::Join(UniqueServers, TEXT(", ")))
-		);
-	const FText ConfirmationTitle = LOCTEXT("DeleteSessionConfirmationTitle", "Delete Session Confirmation");
-	return FMessageDialog::Open(EAppMsgType::YesNo, ConfirmationMessage, ConfirmationTitle) == EAppReturnType::Yes;
-}
-
-bool SConcertClientSessionBrowser::IsLaunchServerButtonEnabled() const
-{
-	return !bLocalServerRunning;
-}
-
-bool SConcertClientSessionBrowser::IsJoinButtonEnabled() const
-{
-	TArray<TSharedPtr<FConcertSessionTreeItem>> SelectedItems = SessionBrowser->GetSelectedItems();
-	return SelectedItems.Num() == 1 && SelectedItems[0]->Type == FConcertSessionTreeItem::EType::ActiveSession;
-}
-
-bool SConcertClientSessionBrowser::IsAutoJoinButtonEnabled() const
-{
-	return Controller->GetConcertClient()->CanAutoConnect() && !Controller->GetConcertClient()->IsAutoConnecting();
-}
-
-bool SConcertClientSessionBrowser::IsCancelAutoJoinButtonEnabled() const
-{
-	return Controller->GetConcertClient()->IsAutoConnecting();
-}
-
-FReply SConcertClientSessionBrowser::OnNewButtonClicked()
-{
-	SessionBrowser->InsertNewSessionEditableRow();
-	return FReply::Handled();
-}
-
-FReply SConcertClientSessionBrowser::OnLaunchServerButtonClicked()
-{
-	IMultiUserClientModule::Get().LaunchConcertServer();
-	bLocalServerRunning = IMultiUserClientModule::Get().IsConcertServerRunning(); // Immediately update the cache state to avoid showing buttons enabled for a split second.
-	return FReply::Handled();
-}
-
-FReply SConcertClientSessionBrowser::OnShutdownServerButtonClicked()
-{
-	if (bLocalServerRunning)
-	{
-		IMultiUserClientModule::Get().ShutdownConcertServer();
-	}
-	return FReply::Handled();
-}
-
-FReply SConcertClientSessionBrowser::OnAutoJoinButtonClicked()
-{
-	IConcertClientPtr ConcertClient = Controller->GetConcertClient();
-
-	// Start the 'auto connect' routine. It will try until it succeeded or gets canceled. Creating or Joining a session automatically cancels it.
-	ConcertClient->StartAutoConnect();
-	return FReply::Handled();
-}
-
-FReply SConcertClientSessionBrowser::OnCancelAutoJoinButtonClicked()
-{
-	Controller->GetConcertClient()->StopAutoConnect();
-	return FReply::Handled();
-}
-
-FReply SConcertClientSessionBrowser::OnJoinButtonClicked()
-{
-	TArray<TSharedPtr<FConcertSessionTreeItem>> SelectedItems = SessionBrowser->GetSelectedItems();
-	if (SelectedItems.Num() == 1)
-	{
-		RequestJoinSession(SelectedItems[0]);
-	}
-	return FReply::Handled();
-}
-
-void SConcertClientSessionBrowser::RequestJoinSession(const TSharedPtr<FConcertSessionTreeItem>& LiveItem)
-{
-	Controller->JoinSession(LiveItem->ServerAdminEndpointId, LiveItem->SessionId);
-}
-
-void SConcertClientSessionBrowser::Tick(const FGeometry& AllottedGeometry, const double InCurrentTime, const float InDeltaTime)
-{
-	// If the model received updates.
-	if (Controller->GetAndClearDiscoveryUpdateFlag())
-	{
-		// Don't wait next TickDiscovery() running at lower frequency, update immediately.
-		UpdateDiscovery();
-	}
-
-	// Ensure the 'default server' filter is updated when the configuration of the default server changes.
-	if (DefaultServerURL != Controller->GetConcertClient()->GetConfiguration()->DefaultServerURL)
-	{
-		DefaultServerURL = Controller->GetConcertClient()->GetConfiguration()->DefaultServerURL;
-		bRefreshSessionFilter = true;
-	}
-
-	// Should refresh the session filter?
-	if (bRefreshSessionFilter)
-	{
-		SessionBrowser->RefreshSessionList();
-		bRefreshSessionFilter = false;
-	}
-}
-
 void SConcertBrowser::Construct(const FArguments& InArgs, const TSharedRef<SDockTab>& ConstructUnderMajorTab, const TSharedPtr<SWindow>& ConstructUnderWindow, TWeakPtr<IConcertSyncClient> InSyncClient)
 {
 	if (!MultiUserClientUtils::HasServerCompatibleCommunicationPluginEnabled())
