Add missing copyright boilerplate

diff --git a/Engine/Plugins/Developer/Concert/ConcertApp/MultiUserClient/Source/MultiUserClient/Private/Widgets/Disconnected/SConcertClientSessionBrowser.cpp b/Engine/Plugins/Developer/Concert/ConcertApp/MultiUserClient/Source/MultiUserClient/Private/Widgets/Disconnected/SConcertClientSessionBrowser.cpp
--- a/Engine/Plugins/Developer/Concert/ConcertApp/MultiUserClient/Source/MultiUserClient/Private/Widgets/Disconnected/SConcertClientSessionBrowser.cpp
+++ b/Engine/Plugins/Developer/Concert/ConcertApp/MultiUserClient/Source/MultiUserClient/Private/Widgets/Disconnected/SConcertClientSessionBrowser.cpp
@@ -1,3 +1,4 @@
+// Copyright Epic Games, Inc. All Rights Reserved.
 
 #include "SConcertClientSessionBrowser.h"
 
