Replcation editor: Auto-add subobjects based on config

diff --git a/Engine/Plugins/Developer/Concert/ConcertApp/MultiUserClient/Source/MultiUserReplicationEditor/Private/Settings/MultiUserReplicationSettings.cpp b/Engine/Plugins/Developer/Concert/ConcertApp/MultiUserClient/Source/MultiUserReplicationEditor/Private/Settings/MultiUserReplicationSettings.cpp
--- a/Engine/Plugins/Developer/Concert/ConcertApp/MultiUserClient/Source/MultiUserReplicationEditor/Private/Settings/MultiUserReplicationSettings.cpp
+++ b/Engine/Plugins/Developer/Concert/ConcertApp/MultiUserClient/Source/MultiUserReplicationEditor/Private/Settings/MultiUserReplicationSettings.cpp
@@ -4,13 +4,16 @@
 
 #include "LogMultiUserReplicationEditor.h"
 #include "Replication/PropertyChainUtils.h"
-#include "Settings/DefaultPropertySelection.h"
+#include "Settings/MultiUserDefaultPropertySelection.h"
+#include "Settings/MultiUserDefaultSubobjectSelection.h"
 
 #include "Algo/IndexOf.h"
+#include "Components/ActorComponent.h"
+#include "Internationalization/Regex.h"
 
-namespace UE::MultiUserReplicationEditor::Private
+namespace UE::MultiUserReplicationEditor::DefaultProperties
 {
-	static void ApplyDefaultPropertySelection(FReplicatedObjectInfo& Info, const FDefaultPropertySelection& Selection, UStruct& Class)
+	static void ApplyDefaultPropertySelection(FReplicatedObjectInfo& Info, const FMultiUserDefaultPropertySelection& Selection, UStruct& Class)
 	{
 		// Preparse the the FStrings into paths
 		TArray<TArray<FName>> Paths;
@@ -54,7 +57,7 @@ void UMultiUserReplicationSettings::AddDefaultPropertiesFromSettings(FReplicated
 {
 	// Find the most specialized class properties
 	UClass* Current = &Class;
-	FDefaultPropertySelection* DefaultProperties = nullptr;
+	FMultiUserDefaultPropertySelection* DefaultProperties = nullptr;
 	for (; Current && !DefaultProperties; Current = Current->GetSuperClass())
 	{
 		DefaultProperties = DefaultPropertySelection.Find(Current);
@@ -63,7 +66,7 @@ void UMultiUserReplicationSettings::AddDefaultPropertiesFromSettings(FReplicated
 			continue;
 		}
 			
-		UE::MultiUserReplicationEditor::Private::ApplyDefaultPropertySelection(Info, *DefaultProperties, *Current);
+		UE::MultiUserReplicationEditor::DefaultProperties::ApplyDefaultPropertySelection(Info, *DefaultProperties, *Current);
 		// Recurse super structs
 		if (UClass* Parent = Current->GetSuperClass()
 			; Parent && DefaultProperties->bInheritFromBase)
@@ -72,3 +75,92 @@ void UMultiUserReplicationSettings::AddDefaultPropertiesFromSettings(FReplicated
 		}
 	}
 }
+
+namespace UE::MultiUserReplicationEditor::DefaultSubobjects
+{
+	static bool MatchesAnyRegex(const FString& Input, const TSet<FString>& AllRegex)
+	{
+		for (const FString& RegexString : AllRegex)
+		{
+			const FRegexPattern Pattern(RegexString);
+			FRegexMatcher Matcher(Pattern, Input);
+			if (Matcher.FindNext())
+			{
+				return true;
+			}
+		}
+
+		return false;
+	}
+	
+	static void ApplyDefaultSubobjectSelection(UObject& AddedObject, const FMultiUserDefaultSubobjectSelection& Selection, TFunctionRef<void(UObject&)> FurtherObjectsCallback)
+	{
+		// Do not search recursively so FurtherObjectsCallback can decide to call AddAdditionalObjectsFromSettings again on the newly added objects.
+		constexpr bool bIncludeNested = false;
+		ForEachObjectWithOuter(&AddedObject, [&Selection, &FurtherObjectsCallback](UObject* Subobject)
+		{
+			// Has exclusion regex?
+			const bool bShouldExclude = MatchesAnyRegex(Subobject->GetName(), Selection.ExcludeSubobjectRegex);
+			if (bShouldExclude)
+			{
+				return;
+			}
+
+			// Was told to add all UActorComponents?
+			const bool bIncludeAllSubobjects = Selection.IncludeAllOption == EMultiUserIncludeAllSubobjectsType::AllSubobjects;
+			const bool bIncludeDueToComponent = Subobject->IsA(UActorComponent::StaticClass())
+				&& Selection.IncludeAllOption == EMultiUserIncludeAllSubobjectsType::AllComponents;
+			if (bIncludeAllSubobjects || bIncludeDueToComponent)
+			{
+				FurtherObjectsCallback(*Subobject);
+				return;
+			}
+
+			// Has configured class?
+			UClass* CurrentClass = Subobject->GetClass();
+			for (; CurrentClass; CurrentClass = CurrentClass->GetSuperClass())
+			{
+				if (Selection.IncludeClasses.Contains(Subobject->GetClass()))
+				{
+					FurtherObjectsCallback(*Subobject);
+					return;
+				}
+			}
+
+			// Has inclusion regex?
+			const bool bShouldIncludeByRegex = MatchesAnyRegex(Subobject->GetName(), Selection.IncludeSubobjectRegex);
+			if (bShouldIncludeByRegex)
+			{
+				FurtherObjectsCallback(*Subobject);
+			}
+		}, bIncludeNested);
+	}
+
+	static void InternalAddAdditionalObjectsFromSettings(UClass& StartClass, UMultiUserReplicationSettings& Settings, UObject& AddedObject, TFunctionRef<void(UObject&)> FurtherObjectsCallback)
+	{
+		// Find the most specialized class properties
+		UClass* Current = &StartClass;
+		FMultiUserDefaultSubobjectSelection* DefaultSubobjects = nullptr;
+		for (; Current && !DefaultSubobjects; Current = Current->GetSuperClass())
+		{
+			DefaultSubobjects = Settings.DefaultSubobjectSelection.Find(Current);
+			if (!DefaultSubobjects)
+			{
+				continue;
+			}
+			
+			ApplyDefaultSubobjectSelection(AddedObject, *DefaultSubobjects, FurtherObjectsCallback);
+			// Recurse super structs
+			if (UClass* Parent = Current->GetSuperClass()
+				; Parent && DefaultSubobjects->bInheritFromBase)
+			{
+				InternalAddAdditionalObjectsFromSettings(*Parent, Settings, AddedObject, FurtherObjectsCallback);
+			}
+		}
+	}
+}
+
+void UMultiUserReplicationSettings::AddAdditionalObjectsFromSettings(UObject& AddedObject, TFunctionRef<void(UObject&)> FurtherObjectsCallback)
+{
+	UE::MultiUserReplicationEditor::DefaultSubobjects::InternalAddAdditionalObjectsFromSettings(*AddedObject.GetClass(), *this, AddedObject, FurtherObjectsCallback);
+}
diff --git a/Engine/Plugins/Developer/Concert/ConcertApp/MultiUserClient/Source/MultiUserReplicationEditor/Private/StreamEditor/Model/PropertySelectionAssetModel.cpp b/Engine/Plugins/Developer/Concert/ConcertApp/MultiUserClient/Source/MultiUserReplicationEditor/Private/StreamEditor/Model/PropertySelectionAssetModel.cpp
--- a/Engine/Plugins/Developer/Concert/ConcertApp/MultiUserClient/Source/MultiUserReplicationEditor/Private/StreamEditor/Model/PropertySelectionAssetModel.cpp
+++ b/Engine/Plugins/Developer/Concert/ConcertApp/MultiUserClient/Source/MultiUserReplicationEditor/Private/StreamEditor/Model/PropertySelectionAssetModel.cpp
@@ -3,7 +3,7 @@
 #include "PropertySelectionAssetModel.h"
 
 #include "Replication/PropertyChainUtils.h"
-#include "Settings/DefaultPropertySelection.h"
+#include "Settings/MultiUserDefaultPropertySelection.h"
 #include "Settings/MultiUserReplicationSettings.h"
 
 #include "Algo/AllOf.h"
@@ -121,6 +121,7 @@ namespace UE::MultiUserReplicationEditor
 		Asset->Modify();
 		
 		TSet<UObject*> ObjectsNotAdded;
+		TArray<UObject*> AdditionalObjectsToAdd;
 		for (UObject* Object : Objects)
 		{
 			const FSoftObjectPath ObjectPath = Object;
@@ -128,7 +129,13 @@ namespace UE::MultiUserReplicationEditor
 			{
 				FReplicatedObjectInfo& ObjectInfo = Asset->ReplicationMap.ReplicatedObjects.Add(ObjectPath);
 				ObjectInfo.ClassPath = Object->GetClass();
-				UMultiUserReplicationSettings::Get()->AddDefaultPropertiesFromSettings(ObjectInfo, *Object->GetClass());
+				
+				UMultiUserReplicationSettings* Settings = UMultiUserReplicationSettings::Get();
+				Settings->AddDefaultPropertiesFromSettings(ObjectInfo, *Object->GetClass());
+				Settings->AddAdditionalObjectsFromSettings(*Object, [&AdditionalObjectsToAdd](UObject& FurtherObject)
+				{
+					AdditionalObjectsToAdd.Add(&FurtherObject);
+				});
 			}
 			else
 			{
@@ -153,6 +160,8 @@ namespace UE::MultiUserReplicationEditor
 			}
 			OnObjectsChangedDelegate.Broadcast(Added, {}, EReplicatedObjectChangeReason::ChangedDirectly);
 		}
+
+		AddObjects(AdditionalObjectsToAdd);
 	}
 
 	void FPropertySelectionAssetModel::RemoveObjects(TArrayView<FSoftObjectPath> Objects)
diff --git a/Engine/Plugins/Developer/Concert/ConcertApp/MultiUserClient/Source/MultiUserReplicationEditor/Public/Settings/DefaultPropertySelection.h b/Engine/Plugins/Developer/Concert/ConcertApp/MultiUserClient/Source/MultiUserReplicationEditor/Public/Settings/DefaultPropertySelection.h
deleted file mode 100644
diff --git a/Engine/Plugins/Developer/Concert/ConcertApp/MultiUserClient/Source/MultiUserReplicationEditor/Public/Settings/MultiUserDefaultPropertySelection.h b/Engine/Plugins/Developer/Concert/ConcertApp/MultiUserClient/Source/MultiUserReplicationEditor/Public/Settings/MultiUserDefaultPropertySelection.h
new file mode 100644
--- /dev/null
+++ b/Engine/Plugins/Developer/Concert/ConcertApp/MultiUserClient/Source/MultiUserReplicationEditor/Public/Settings/MultiUserDefaultPropertySelection.h
@@ -0,0 +1,20 @@
+// Copyright Epic Games, Inc. All Rights Reserved.
+
+#pragma once
+
+#include "MultiUserInheritableClassOption.h"
+#include "MultiUserDefaultPropertySelection.generated.h"
+
+USTRUCT()
+struct FMultiUserDefaultPropertySelection : public FMultiUserInheritableClassOption
+{
+	GENERATED_BODY()
+
+	/**
+	 * A list of properties that should be selected by default when you add a new object type.
+	 * Specify properties by using the names as the appear in the replication editor and each property with ".".
+	 * Example: "RelativeLocation.X"
+	 */
+	UPROPERTY(EditAnywhere, Category = "Config")
+	TArray<FString> DefaultSelectedProperties;
+};
diff --git a/Engine/Plugins/Developer/Concert/ConcertApp/MultiUserClient/Source/MultiUserReplicationEditor/Public/Settings/MultiUserDefaultSubobjectSelection.h b/Engine/Plugins/Developer/Concert/ConcertApp/MultiUserClient/Source/MultiUserReplicationEditor/Public/Settings/MultiUserDefaultSubobjectSelection.h
new file mode 100644
--- /dev/null
+++ b/Engine/Plugins/Developer/Concert/ConcertApp/MultiUserClient/Source/MultiUserReplicationEditor/Public/Settings/MultiUserDefaultSubobjectSelection.h
@@ -0,0 +1,56 @@
+// Copyright Epic Games, Inc. All Rights Reserved.
+
+#pragma once
+
+#include "MultiUserInheritableClassOption.h"
+#include "UObject/SoftObjectPath.h"
+#include "MultiUserDefaultSubobjectSelection.generated.h"
+
+UENUM()
+enum class EMultiUserIncludeAllSubobjectsType : uint8
+{
+	/** Include nothing by default */
+	None,
+	/** Include all components */
+	AllComponents,
+	/** Include all subobjects recursively (includes components, their subobjects, etc. ) */
+	AllSubobjects
+};
+
+/** Settings for components that are supposed to be included  */
+USTRUCT()
+struct FMultiUserDefaultSubobjectSelection : public FMultiUserInheritableClassOption
+{
+	GENERATED_BODY()
+	
+	/**
+	 * A list of subobject classes to add by default.
+	 * For example, when you add an actor, you want to also add certain component classes.
+	 */
+	UPROPERTY(EditAnywhere, Category = "Config")
+	TSet<FSoftClassPath> IncludeClasses;
+
+	/**
+	 * Components that match any of this regex will be excluded from DefaultSelectedComponentClasses.
+	 *
+	 * Tips:
+	 * - If you want to include "Component" but not "ComponentName" you can use boundaries "\bComponent\b".
+	 * . ".*" matches any sequence of characters.
+	 */
+	UPROPERTY(EditAnywhere, Category = "Config")
+	TSet<FString> ExcludeSubobjectRegex;
+
+	/**
+	 * Components that match any of this regex will be included as well.
+	 * 
+	 * Tips:
+	 * - If you want to include "Component" but not "ComponentName" you can use boundaries "\bComponent\b".
+	 * . ".*" matches any sequence of characters. 
+	 */
+	UPROPERTY(EditAnywhere, Category = "Config")
+	TSet<FString> IncludeSubobjectRegex;
+
+	/** Behaviour to configure for including all of a certain type of subobject. */
+	UPROPERTY(EditAnywhere, AdvancedDisplay, Category = "Config")
+	EMultiUserIncludeAllSubobjectsType IncludeAllOption = EMultiUserIncludeAllSubobjectsType::None;
+};
diff --git a/Engine/Plugins/Developer/Concert/ConcertApp/MultiUserClient/Source/MultiUserReplicationEditor/Public/Settings/MultiUserInheritableClassOption.h b/Engine/Plugins/Developer/Concert/ConcertApp/MultiUserClient/Source/MultiUserReplicationEditor/Public/Settings/MultiUserInheritableClassOption.h
new file mode 100644
--- /dev/null
+++ b/Engine/Plugins/Developer/Concert/ConcertApp/MultiUserClient/Source/MultiUserReplicationEditor/Public/Settings/MultiUserInheritableClassOption.h
@@ -0,0 +1,16 @@
+// Copyright Epic Games, Inc. All Rights Reserved.
+
+#pragma once
+
+#include "MultiUserInheritableClassOption.generated.h"
+
+/** Base struct for settings intended for classes */
+USTRUCT()
+struct FMultiUserInheritableClassOption
+{
+	GENERATED_BODY()
+
+	/** Whether to inherit the settings from the base class. */
+	UPROPERTY(EditAnywhere, Category = "Config")
+	bool bInheritFromBase = true;
+};
diff --git a/Engine/Plugins/Developer/Concert/ConcertApp/MultiUserClient/Source/MultiUserReplicationEditor/Public/Settings/MultiUserReplicationSettings.h b/Engine/Plugins/Developer/Concert/ConcertApp/MultiUserClient/Source/MultiUserReplicationEditor/Public/Settings/MultiUserReplicationSettings.h
--- a/Engine/Plugins/Developer/Concert/ConcertApp/MultiUserClient/Source/MultiUserReplicationEditor/Public/Settings/MultiUserReplicationSettings.h
+++ b/Engine/Plugins/Developer/Concert/ConcertApp/MultiUserClient/Source/MultiUserReplicationEditor/Public/Settings/MultiUserReplicationSettings.h
@@ -4,26 +4,63 @@
 
 #include "UObject/Object.h"
 #include "MultiUserReplicationClientProfileAsset.h"
-#include "Settings/DefaultPropertySelection.h"
+#if WITH_EDITORONLY_DATA
+#include "Settings/MultiUserDefaultPropertySelection.h"
+#include "Settings/MultiUserDefaultSubobjectSelection.h"
+#include "Templates/Function.h"
+#endif
 #include "MultiUserReplicationSettings.generated.h"
 
+// Architecturally speaking FDefaultReplicationSessionSettings should be in the MultiUserClient module.
+// However, we want a single "Multi-User Replication" settings section... so all the data just lives here (for now?).
+
+USTRUCT()
+struct FDefaultReplicationSessionSettings
+{
+	GENERATED_BODY()
+
+	/** Whether to automatically start replicating after successfully joining a concert session. */
+	UPROPERTY(EditAnywhere, Config, Category = "Replication")
+	bool bAutoJoin = false;
+	
+	/** The client profile to use when connecting to a session. */
+	UPROPERTY(EditAnywhere, Config, Category = "Replication")
+	TSoftObjectPtr<UMultiUserReplicationClientProfileAsset> DefaultProfile;
+};
+
 /** Asset for users to describe a client to the server. */
 UCLASS(Config=MultiUserClient)
 class MULTIUSERREPLICATIONEDITOR_API UMultiUserReplicationSettings : public UObject
 {
 	GENERATED_BODY()
 public:
-	
-	/** The client profile to use when connecting to a session. */
+
+	/** Default replication settings to use on this editor. */
 	UPROPERTY(EditAnywhere, Config, Category = "Replication")
-	TSoftObjectPtr<UMultiUserReplicationClientProfileAsset> DefaultProfile;
+	FDefaultReplicationSessionSettings DefaultSessionSettings;
 
+#if WITH_EDITORONLY_DATA
 	/** Properties you want selected by default when you add a new replicated object in the editor */
 	UPROPERTY(EditAnywhere, Config, Category = "Replication|Editor")
-	TMap<FSoftClassPath, FDefaultPropertySelection> DefaultPropertySelection;
+	TMap<FSoftClassPath, FMultiUserDefaultPropertySelection> DefaultPropertySelection;
+
+	/**
+	 * Subobjects you want selected by default when you add a new replicated object in the editor.
+	 * For examples:
+	 *  - Whenever you add a StaticMeshActor actor, you its static mesh component to be added automatically, too.
+	 *  - Whenever you add a component, you want to add its child components, too.
+	 */
+	UPROPERTY(EditAnywhere, Config, Category = "Replication|Editor")
+	TMap<FSoftClassPath, FMultiUserDefaultSubobjectSelection> DefaultSubobjectSelection;
+#endif
 
 	static UMultiUserReplicationSettings* Get() { return GetMutableDefault<UMultiUserReplicationSettings>(); }
 
+#if WITH_EDITORONLY_DATA
 	/** Reads DefaultPropertySelection and applies any default property selections to Info based on the Class just added. */
 	void AddDefaultPropertiesFromSettings(FReplicatedObjectInfo& Info, UClass& Class);
+
+	/** Reads DefaultComponentSelection and calls FurtherObjectsCallback on any further objects that should also be added. */
+	void AddAdditionalObjectsFromSettings(UObject& AddedObject, TFunctionRef<void(UObject&)> FurtherObjectsCallback);
+#endif
 };
\ No newline at end of file
