Replication Editor: Automatically add static mesh and camera components when adding actor

diff --git a/Engine/Plugins/Developer/Concert/ConcertApp/MultiUserClient/Config/BaseMultiUserClient.ini b/Engine/Plugins/Developer/Concert/ConcertApp/MultiUserClient/Config/BaseMultiUserClient.ini
--- a/Engine/Plugins/Developer/Concert/ConcertApp/MultiUserClient/Config/BaseMultiUserClient.ini
+++ b/Engine/Plugins/Developer/Concert/ConcertApp/MultiUserClient/Config/BaseMultiUserClient.ini
@@ -1,2 +1,3 @@
 [/Script/MultiUserReplicationEditor.MultiUserReplicationSettings]
-+DefaultPropertySelection=(("/Script/Engine.SceneComponent", (DefaultSelectedProperties=("RelativeLocation","RelativeLocation.X","RelativeLocation.Y","RelativeLocation.Z","RelativeRotation","RelativeRotation.X","RelativeRotation.Y","RelativeRotation.Z","RelativeScale","RelativeScale.X","RelativeScale.Y","RelativeScale.Z"))))
\ No newline at end of file
++DefaultPropertySelection=(("/Script/Engine.SceneComponent", (DefaultSelectedProperties=("RelativeLocation","RelativeLocation.X","RelativeLocation.Y","RelativeLocation.Z","RelativeRotation","RelativeRotation.X","RelativeRotation.Y","RelativeRotation.Z","RelativeScale","RelativeScale.X","RelativeScale.Y","RelativeScale.Z"))))
++DefaultSubobjectSelection=(("/Script/Engine.Actor", (IncludeClasses=("/Script/Engine.StaticMeshComponent","/Script/Engine.CameraComponent"))))
\ No newline at end of file
