Replication editor: Filter by property type

diff --git a/Engine/Plugins/Developer/Concert/ConcertApp/MultiUserClient/Source/MultiUserReplicationEditor/MultiUserReplicationEditor.Build.cs b/Engine/Plugins/Developer/Concert/ConcertApp/MultiUserClient/Source/MultiUserReplicationEditor/MultiUserReplicationEditor.Build.cs
--- a/Engine/Plugins/Developer/Concert/ConcertApp/MultiUserClient/Source/MultiUserReplicationEditor/MultiUserReplicationEditor.Build.cs
+++ b/Engine/Plugins/Developer/Concert/ConcertApp/MultiUserClient/Source/MultiUserReplicationEditor/MultiUserReplicationEditor.Build.cs
@@ -30,6 +30,7 @@ namespace UnrealBuildTool.Rules
 					"Projects",
 					"Slate",
 					"SlateCore",
+					"ToolMenus",
 					"ToolWidgets",
 					"UnrealEd", 
 					
diff --git a/Engine/Plugins/Developer/Concert/ConcertApp/MultiUserClient/Source/MultiUserReplicationEditor/Private/StreamEditor/View/ObjectViewer/Property/PropertyFilterBase.h b/Engine/Plugins/Developer/Concert/ConcertApp/MultiUserClient/Source/MultiUserReplicationEditor/Private/StreamEditor/View/ObjectViewer/Property/PropertyFilterBase.h
new file mode 100644
--- /dev/null
+++ b/Engine/Plugins/Developer/Concert/ConcertApp/MultiUserClient/Source/MultiUserReplicationEditor/Private/StreamEditor/View/ObjectViewer/Property/PropertyFilterBase.h
@@ -0,0 +1,45 @@
+// Copyright Epic Games, Inc. All Rights Reserved.
+
+#pragma once
+
+#include "Containers/Set.h"
+#include "Misc/IFilter.h"
+
+class FFieldClass;
+
+namespace UE::MultiUserReplicationEditor
+{
+	class FReplicatedPropertyData;
+
+	/**
+	 * Base class for property filters.
+	 * 
+	 * This filter inverts the results: all filters are supposed to be inverse filters, i.e. they should be run when
+	 * they are greyed out in the UI. When they are run, they should remove specific properties.
+	 *
+	 * Subclasses simply implement Matches, which figures out whether the property is contained in some FReplicatedPropertyData.
+	 */
+	class FPropertyFilterBase : public IFilter<const FReplicatedPropertyData&>
+	{
+	public:
+
+
+		//~ Begin IFilter Interface
+		virtual bool PassesFilter(const FReplicatedPropertyData& InItem) const final override
+		{
+			// If the property matches our filter criteria, it should NOT be shown in the UI since we're an inverse filter, so negate.
+			return !MatchesFilteredForProperty(InItem);
+		}
+		virtual FChangedEvent& OnChanged() final override { return ChangedEventDelegate; }
+		//~ End IFilter Interface
+
+	private:
+
+		/** Dummy delegate to fulfill IFilter interface: this filter will never change so this delegate will never trigger. */
+		FChangedEvent ChangedEventDelegate;
+
+		/** @return Whether this item contains the property this filter is looking for.*/
+		virtual bool MatchesFilteredForProperty(const FReplicatedPropertyData& InItem) const = 0;
+	};
+}
+
diff --git a/Engine/Plugins/Developer/Concert/ConcertApp/MultiUserClient/Source/MultiUserReplicationEditor/Private/StreamEditor/View/ObjectViewer/Property/PropertyFilter_ByPropertyType.cpp b/Engine/Plugins/Developer/Concert/ConcertApp/MultiUserClient/Source/MultiUserReplicationEditor/Private/StreamEditor/View/ObjectViewer/Property/PropertyFilter_ByPropertyType.cpp
new file mode 100644
--- /dev/null
+++ b/Engine/Plugins/Developer/Concert/ConcertApp/MultiUserClient/Source/MultiUserReplicationEditor/Private/StreamEditor/View/ObjectViewer/Property/PropertyFilter_ByPropertyType.cpp
@@ -0,0 +1,23 @@
+// Copyright Epic Games, Inc. All Rights Reserved.
+
+#include "PropertyFilter_ByPropertyType.h"
+
+#include "Replication/PropertyChainUtils.h"
+#include "StreamEditor/View/ObjectViewer/ReplicatedPropertyData.h"
+
+#include "UObject/UnrealType.h"
+
+namespace UE::MultiUserReplicationEditor
+{
+	bool FPropertyFilter_ByPropertyType::MatchesFilteredForProperty(const FReplicatedPropertyData& InItem) const
+	{
+		UClass* Class = InItem.GetOwningClass().TryLoadClass<UObject>();
+		if (!Class)
+		{
+			return false;
+		}
+
+		const FProperty* Property = ConcertSyncCore::PropertyChain::ResolveProperty(*Class, InItem.GetProperty());
+		return Property && AllowedClasses.Contains(Property->GetClass());
+	}
+}
diff --git a/Engine/Plugins/Developer/Concert/ConcertApp/MultiUserClient/Source/MultiUserReplicationEditor/Private/StreamEditor/View/ObjectViewer/Property/PropertyFilter_ByPropertyType.h b/Engine/Plugins/Developer/Concert/ConcertApp/MultiUserClient/Source/MultiUserReplicationEditor/Private/StreamEditor/View/ObjectViewer/Property/PropertyFilter_ByPropertyType.h
new file mode 100644
--- /dev/null
+++ b/Engine/Plugins/Developer/Concert/ConcertApp/MultiUserClient/Source/MultiUserReplicationEditor/Private/StreamEditor/View/ObjectViewer/Property/PropertyFilter_ByPropertyType.h
@@ -0,0 +1,36 @@
+// Copyright Epic Games, Inc. All Rights Reserved.
+
+#pragma once
+
+#include "PropertyFilterBase.h"
+#include "Containers/Set.h"
+#include "Misc/IFilter.h"
+
+class FFieldClass;
+
+namespace UE::MultiUserReplicationEditor
+{
+	class FReplicatedPropertyData;
+
+	/** Filters a property based on the FFieldClass it has. */
+	class FPropertyFilter_ByPropertyType : public FPropertyFilterBase
+	{
+	public:
+
+		FPropertyFilter_ByPropertyType(TSet<FFieldClass*> AllowedClasses)
+			: AllowedClasses(MoveTemp(AllowedClasses))
+		{}
+
+	private:
+
+		TSet<FFieldClass*> AllowedClasses;
+
+		/** Dummy delegate to fulfill IFilter interface: this filter will never change so this delegate will never trigger. */
+		FChangedEvent ChangedEventDelegate;
+
+		//~ Begin FPropertyFilterBase Interface
+		virtual bool MatchesFilteredForProperty(const FReplicatedPropertyData& InItem) const override;
+		//~ End FPropertyFilterBase Interface
+	};
+}
+
diff --git a/Engine/Plugins/Developer/Concert/ConcertApp/MultiUserClient/Source/MultiUserReplicationEditor/Private/StreamEditor/View/ObjectViewer/Property/PropertyFrontendFilter.h b/Engine/Plugins/Developer/Concert/ConcertApp/MultiUserClient/Source/MultiUserReplicationEditor/Private/StreamEditor/View/ObjectViewer/Property/PropertyFrontendFilter.h
new file mode 100644
--- /dev/null
+++ b/Engine/Plugins/Developer/Concert/ConcertApp/MultiUserClient/Source/MultiUserReplicationEditor/Private/StreamEditor/View/ObjectViewer/Property/PropertyFrontendFilter.h
@@ -0,0 +1,81 @@
+// Copyright Epic Games, Inc. All Rights Reserved.
+
+#pragma once
+
+#include "Filters/FilterBase.h"
+#include "Internationalization/Text.h"
+
+namespace UE::MultiUserReplicationEditor
+{
+	class FReplicatedPropertyData;
+
+	/** Base UI filter for filtering properties. Gets rid of all the extra functionality in FFilterBase we do not need. */
+	class FPropertyFrontendFilter : public FFilterBase<TSharedPtr<FReplicatedPropertyData>>
+	{
+	public:
+
+		FPropertyFrontendFilter(TSharedRef<FFilterCategory> InCategory, FText DisplayName, FText Tooltip = FText::GetEmpty())
+			: FFilterBase(MoveTemp(InCategory))
+			, DisplayName(MoveTemp(DisplayName))
+			, Tooltip(MoveTemp(Tooltip))
+		{}
+
+		//~ Begin FFilterBase Interface
+		virtual FString GetName() const override { return DisplayName.ToString(); }
+		virtual FText GetDisplayName() const override { return DisplayName; }
+		virtual FText GetToolTipText() const override { return Tooltip; }
+		virtual FLinearColor GetColor() const override { return FLinearColor(0.6f, 0.6f, 0.6f, 0.6f); }
+		virtual FName GetIconName() const override { return NAME_None; }
+		virtual bool IsInverseFilter() const override { return true; }
+		virtual void ActiveStateChanged(bool bActive) override {}
+		virtual void ModifyContextMenu(FMenuBuilder& MenuBuilder) override {}
+		virtual void SaveSettings(const FString& IniFilename, const FString& IniSection, const FString& SettingsString) const override {}
+		virtual void LoadSettings(const FString& IniFilename, const FString& IniSection, const FString& SettingsString) override {}
+		//~ End FFilterBase Interface
+
+	private:
+
+		FText DisplayName;
+		FText Tooltip;
+	};
+
+	/** Inlines a IFilter<const FReplicatedPropertyData&> at compile time to avoid an additional TSharedPtr<IFilter<const FReplicatedPropertyData&>>. */
+	template<typename TFilterType>
+	class TPropertyFrontendFilter : public FPropertyFrontendFilter
+	{
+	public:
+
+		template<typename... TArg>
+		TPropertyFrontendFilter(TSharedRef<FFilterCategory> InCategory, FText DisplayName, TArg&&... Arg)
+			: FPropertyFrontendFilter(MoveTemp(InCategory), MoveTemp(DisplayName))
+			, FilterImplementation(Forward<TArg>(Arg)...)
+		{
+			FilterImplementation.OnChanged().AddLambda([this]()
+			{
+				FPropertyFrontendFilter::BroadcastChangedEvent();
+			});
+		}
+
+		template<typename... TArg>
+		TPropertyFrontendFilter(TSharedRef<FFilterCategory> InCategory, FText DisplayName, FText Tooltip, TArg&&... Arg)
+			: FPropertyFrontendFilter(MoveTemp(InCategory), MoveTemp(DisplayName), MoveTemp(Tooltip))
+			, FilterImplementation(Forward<TArg>(Arg)...)
+		{
+			FilterImplementation.OnChanged().AddLambda([this]()
+			{
+				FPropertyFrontendFilter::BroadcastChangedEvent();
+			});
+		}
+		
+		virtual bool PassesFilter(TSharedPtr<FReplicatedPropertyData> InItem) const override
+		{
+			check(InItem);
+			return FilterImplementation.PassesFilter(*InItem.Get());
+		}
+
+	private:
+		
+		TFilterType FilterImplementation;
+	};
+}
+
diff --git a/Engine/Plugins/Developer/Concert/ConcertApp/MultiUserClient/Source/MultiUserReplicationEditor/Private/StreamEditor/View/ObjectViewer/Property/SReplicatedPropertiesView.cpp b/Engine/Plugins/Developer/Concert/ConcertApp/MultiUserClient/Source/MultiUserReplicationEditor/Private/StreamEditor/View/ObjectViewer/Property/SReplicatedPropertiesView.cpp
new file mode 100644
--- /dev/null
+++ b/Engine/Plugins/Developer/Concert/ConcertApp/MultiUserClient/Source/MultiUserReplicationEditor/Private/StreamEditor/View/ObjectViewer/Property/SReplicatedPropertiesView.cpp
@@ -0,0 +1,136 @@
+// Copyright Epic Games, Inc. All Rights Reserved.
+
+#include "SReplicatedPropertiesView.h"
+
+#include "Filters/SBasicFilterBar.h"
+#include "PropertyFilter_ByPropertyType.h"
+#include "PropertyFrontendFilter.h"
+#include "StreamEditor/View/ObjectViewer/ReplicatedPropertyData.h"
+
+#include "UObject/UnrealType.h"
+
+#define LOCTEXT_NAMESPACE "SReplicatedPropertiesView"
+
+namespace UE::MultiUserReplicationEditor
+{
+	void SReplicatedPropertiesView::Construct(const FArguments& InArgs)
+	{
+		const FBuildFilterBarResult Filters = BuildFilterBar();
+		
+		ChildSlot
+		[
+			SAssignNew(ReplicatedProperties, SReplicationTreeView<TSharedPtr<FReplicatedPropertyData>>)
+				.RootItemsSource(InArgs._RootItemsSource)
+				.OnGetChildren(InArgs._OnGetChildren)
+				.Columns(InArgs._Columns)
+				.ExpandableColumnLabel(InArgs._ExpandableColumnLabel)
+				.SelectionMode(InArgs._SelectionMode)
+				.FilterItem(this, &SReplicatedPropertiesView::PassesFilters)
+				.LeftOfSearchBar()
+				[
+					SNew(SHorizontalBox)
+
+					// The combo button for selecting the property filters
+					+SHorizontalBox::Slot()
+					.AutoWidth()
+					[
+						SBasicFilterBar<TSharedPtr<FReplicatedPropertyData>>::MakeAddFilterButton(FilterBar.ToSharedRef())
+					]
+
+					+SHorizontalBox::Slot()
+					.AutoWidth()
+					[
+						InArgs._LeftOfSearchBar.Widget
+					]
+				]
+				.RowBelowSearchBar()
+				[
+					FilterBar.ToSharedRef()
+				]
+		];
+
+		// For better UX, hide certain properties by default (e.g. why would you want to replicate bools?)
+		// Do this AFTER initializing ReplicatedProperties because it triggers the OnItemsChanged callback.
+		for (const FFilterRef& Filter : Filters.DisabledByDefault)
+		{
+			FilterBar->SetFilterCheckState(Filter, ECheckBoxState::Checked);
+		}
+		// Our filters are inverse (FFilterBase::IsInverseFilter) so that means these disabled filters will be active now (so e.g. bool will be filtered out by default now)
+		FilterBar->DisableAllFilters();
+		// Show all the other filters as enabled (not greyed out: blue) - they will not run their logic since they are inverse.
+		for (const FFilterRef& Filter : Filters.EnabledByDefault)
+		{
+			FilterBar->SetFilterCheckState(Filter, ECheckBoxState::Checked);
+		}
+	}
+
+	void SReplicatedPropertiesView::OnItemsChanged() const
+	{
+		ReplicatedProperties->OnItemsChanged();
+	}
+
+	SReplicatedPropertiesView::FBuildFilterBarResult SReplicatedPropertiesView::BuildFilterBar()
+	{
+		TSharedRef<FFilterCategory> GeneralCategory = MakeShared<FFilterCategory>(
+			LOCTEXT("CommonCategory.Name", "General"),
+			LOCTEXT("CommonCategory.ToolTip", "Exclude general properties.")
+		);
+		TSharedRef<FFilterCategory> NumericCategory = MakeShared<FFilterCategory>(
+			LOCTEXT("NumericCategory.Name", "Numeric"),
+			LOCTEXT("NumericCategory.ToolTip", "Excludes numeric properties.")
+		);
+		TSharedRef<FFilterCategory> ContainersCategory = MakeShared<FFilterCategory>(
+			LOCTEXT("ContainersCategory.Name", "Containers"),
+			LOCTEXT("ContainersCategory.ToolTip", "Exclude container properties.")
+		);
+
+		FBuildFilterBarResult Result;
+		using FFrontendFilter = TPropertyFrontendFilter<FPropertyFilter_ByPropertyType>;
+		Result.DisabledByDefault =
+		{
+			MakeShared<FFrontendFilter>(MoveTemp(GeneralCategory), LOCTEXT("Bool", "Boolean"), TSet<FFieldClass*>{ FBoolProperty::StaticClass() }),
+			MakeShared<FFrontendFilter>(MoveTemp(NumericCategory), LOCTEXT("Byte", "Byte"), TSet<FFieldClass*>{ FByteProperty::StaticClass() }),
+			MakeShared<FFrontendFilter>(MoveTemp(GeneralCategory), LOCTEXT("Name", "Name"), TSet<FFieldClass*>{ FNameProperty::StaticClass() }),
+			MakeShared<FFrontendFilter>(MoveTemp(GeneralCategory), LOCTEXT("String", "String"), TSet<FFieldClass*>{ FStrProperty::StaticClass() }),
+			MakeShared<FFrontendFilter>(MoveTemp(GeneralCategory), LOCTEXT("Enum", "Enum"), TSet<FFieldClass*>{ FEnumProperty::StaticClass() })
+		};
+		Result.EnabledByDefault =
+		{
+			MakeShared<FFrontendFilter>(MoveTemp(NumericCategory), LOCTEXT("Ints", "Integers"), LOCTEXT("Ints.Tooltip", "Includes integer types: uint16, uint32, uint64, int16, int32, int64"), TSet<FFieldClass*>{ FUInt16Property::StaticClass(), FUInt32Property::StaticClass(), FUInt64Property::StaticClass(), FInt16Property::StaticClass(), FIntProperty::StaticClass(), FInt64Property::StaticClass() }),
+			MakeShared<FFrontendFilter>(MoveTemp(NumericCategory), LOCTEXT("Float", "Float"), TSet<FFieldClass*>{ FFloatProperty::StaticClass() }),
+			MakeShared<FFrontendFilter>(MoveTemp(NumericCategory), LOCTEXT("Double", "Double"), TSet<FFieldClass*>{ FDoubleProperty::StaticClass() }),
+			MakeShared<FFrontendFilter>(MoveTemp(GeneralCategory), LOCTEXT("Struct", "Struct"), TSet<FFieldClass*>{ FStructProperty::StaticClass() }),
+			MakeShared<FFrontendFilter>(MoveTemp(ContainersCategory), LOCTEXT("Array", "Array"), TSet<FFieldClass*>{ FArrayProperty::StaticClass() }),
+			MakeShared<FFrontendFilter>(MoveTemp(ContainersCategory), LOCTEXT("Set", "Set"), TSet<FFieldClass*>{ FSetProperty::StaticClass() }),
+			MakeShared<FFrontendFilter>(MoveTemp(ContainersCategory), LOCTEXT("Map", "Map"), TSet<FFieldClass*>{ FMapProperty::StaticClass() }),
+		};
+		
+		TArray<FFilterRef> AllFilters =
+		{
+			MakeShared<FFrontendFilter>(MoveTemp(NumericCategory), LOCTEXT("UInt16", "UInt16"), TSet<FFieldClass*>{ FUInt16Property::StaticClass() }),
+			MakeShared<FFrontendFilter>(MoveTemp(NumericCategory), LOCTEXT("UInt32", "UInt32"), TSet<FFieldClass*>{ FUInt32Property::StaticClass() }),
+			MakeShared<FFrontendFilter>(MoveTemp(NumericCategory), LOCTEXT("UInt64", "UInt64"), TSet<FFieldClass*>{ FUInt64Property::StaticClass() }),
+			MakeShared<FFrontendFilter>(MoveTemp(NumericCategory), LOCTEXT("Int16", "Int16"), TSet<FFieldClass*>{ FInt16Property::StaticClass() }),
+			MakeShared<FFrontendFilter>(MoveTemp(NumericCategory), LOCTEXT("Int32", "Int32"), TSet<FFieldClass*>{ FIntProperty::StaticClass() }),
+			MakeShared<FFrontendFilter>(MoveTemp(NumericCategory), LOCTEXT("Int64", "Int64"), TSet<FFieldClass*>{ FInt64Property::StaticClass() }),
+			MakeShared<FFrontendFilter>(MoveTemp(GeneralCategory), LOCTEXT("SoftPtr", "Soft Ptr"), TSet<FFieldClass*>{ FSoftObjectProperty::StaticClass() }),
+		};
+		AllFilters.Append(Result.DisabledByDefault);
+		AllFilters.Append(Result.EnabledByDefault);
+		
+		FilterBar = SNew(SBasicFilterBar<TSharedPtr<FReplicatedPropertyData>>)
+			.FilterPillStyle(EFilterPillStyle::Basic)
+			.CustomFilters(MoveTemp(AllFilters))
+			.OnFilterChanged(this, &SReplicatedPropertiesView::OnItemsChanged)
+			.UseSectionsForCategories(true);
+		
+		return Result;
+	}
+
+	bool SReplicatedPropertiesView::PassesFilters(const TSharedPtr<FReplicatedPropertyData>& ReplicatedPropertyData) const
+	{
+		return FilterBar->GetAllActiveFilters()->PassesAllFilters(ReplicatedPropertyData);
+	}
+}
+
+#undef LOCTEXT_NAMESPACE
\ No newline at end of file
diff --git a/Engine/Plugins/Developer/Concert/ConcertApp/MultiUserClient/Source/MultiUserReplicationEditor/Private/StreamEditor/View/ObjectViewer/Property/SReplicatedPropertiesView.h b/Engine/Plugins/Developer/Concert/ConcertApp/MultiUserClient/Source/MultiUserReplicationEditor/Private/StreamEditor/View/ObjectViewer/Property/SReplicatedPropertiesView.h
new file mode 100644
--- /dev/null
+++ b/Engine/Plugins/Developer/Concert/ConcertApp/MultiUserClient/Source/MultiUserReplicationEditor/Private/StreamEditor/View/ObjectViewer/Property/SReplicatedPropertiesView.h
@@ -0,0 +1,69 @@
+// Copyright Epic Games, Inc. All Rights Reserved.
+
+#pragma once
+
+#include "Filters/SBasicFilterBar.h"
+#include "Widgets/DeclarativeSyntaxSupport.h"
+#include "Widgets/SCompoundWidget.h"
+#include "StreamEditor/View/ObjectViewer/SReplicationTreeView.h"
+
+namespace UE::MultiUserReplicationEditor
+{
+	class FReplicatedPropertyData;
+	
+	/** Displays a searchable tree view of properties (SReplicationTreeView) and decorates it with a SBasicFilterBar. */
+	class SReplicatedPropertiesView : public SCompoundWidget
+	{
+		using FFilterRef = TSharedRef<FFilterBase<TSharedPtr<FReplicatedPropertyData>>>;
+	public:
+		
+		SLATE_BEGIN_ARGS(SReplicatedPropertiesView)
+		{}
+			/*************** Arguments inherited by SReplicationTreeView ***************/
+		
+			/** The items to display */
+			SLATE_ARGUMENT(TArray<TSharedPtr<FReplicatedPropertyData>>*, RootItemsSource)
+		
+			/** Gets an items children for the tree view */
+			SLATE_EVENT(SReplicationTreeView<TSharedPtr<FReplicatedPropertyData>>::FGetItemChildren, OnGetChildren)
+		
+			/** The columns this list should have */
+			SLATE_ARGUMENT(TArray<TReplicationColumn<TSharedPtr<FReplicatedPropertyData>>>, Columns)
+			/** The name of the column that will have the SExpanderArrow for the tree view. */
+			SLATE_ARGUMENT(FName, ExpandableColumnLabel)
+		
+			/** How many items are to allowed to be selected */
+			SLATE_ARGUMENT(ESelectionMode::Type, SelectionMode)
+		
+			/** Optional widget to add to the left of the search bar. */
+			SLATE_NAMED_SLOT(FArguments, LeftOfSearchBar)
+
+			/*************** Own arguments ***************/
+			// Please add new arguments here in the future
+		SLATE_END_ARGS()
+
+		void Construct(const FArguments& InArgs);
+
+		/** Called when the items need to be refiltered due to the item source changing. */
+		void OnItemsChanged() const;
+
+	private:
+
+		/** The tree view displaying the replicated properties */
+		TSharedPtr<SReplicationTreeView<TSharedPtr<FReplicatedPropertyData>>> ReplicatedProperties;
+		/** Displays the active filters*/
+		TSharedPtr<SBasicFilterBar<TSharedPtr<FReplicatedPropertyData>>> FilterBar;
+
+		struct FBuildFilterBarResult
+		{
+			TArray<FFilterRef> EnabledByDefault;
+			TArray<FFilterRef> DisabledByDefault;
+		};
+		/** Build the filter bar and returns the filters that should be active by default. */
+		FBuildFilterBarResult BuildFilterBar();
+
+		/** Runs all filters through this item */
+		bool PassesFilters(const TSharedPtr<FReplicatedPropertyData>& ReplicatedPropertyData) const;
+	};
+}
+
diff --git a/Engine/Plugins/Developer/Concert/ConcertApp/MultiUserClient/Source/MultiUserReplicationEditor/Private/StreamEditor/View/ObjectViewer/SPropertyReplicationSelectionViewer.cpp b/Engine/Plugins/Developer/Concert/ConcertApp/MultiUserClient/Source/MultiUserReplicationEditor/Private/StreamEditor/View/ObjectViewer/SPropertyReplicationSelectionViewer.cpp
--- a/Engine/Plugins/Developer/Concert/ConcertApp/MultiUserClient/Source/MultiUserReplicationEditor/Private/StreamEditor/View/ObjectViewer/SPropertyReplicationSelectionViewer.cpp
+++ b/Engine/Plugins/Developer/Concert/ConcertApp/MultiUserClient/Source/MultiUserReplicationEditor/Private/StreamEditor/View/ObjectViewer/SPropertyReplicationSelectionViewer.cpp
@@ -3,10 +3,11 @@
 #include "SPropertyReplicationSelectionViewer.h"
 
 #include "ConcertFrontendUtils.h"
-#include "StreamEditor/Model/IObjectToPropertiesModel.h"
 #include "ReplicatedObjectData.h"
 #include "ReplicatedPropertyData.h"
 #include "SelectionViewerColumns.h"
+#include "StreamEditor/Model/IObjectToPropertiesModel.h"
+#include "StreamEditor/View/ObjectViewer/Property/SReplicatedPropertiesView.h"
 
 #include "Algo/AllOf.h"
 #include "Algo/ForEach.h"
@@ -230,7 +231,8 @@ namespace UE::MultiUserReplicationEditor
 	{
 		TArray<ReplicationPropertyColumns::FReplicationPropertyColumn> Columns
 		{
-			ReplicationPropertyColumns::LabelColumn()
+			ReplicationPropertyColumns::LabelColumn(),
+			ReplicationPropertyColumns::TypeColumn()
 		};
 		Columns.Append(InArgs._AdditionalPropertyColumns);
 		
@@ -262,7 +264,7 @@ namespace UE::MultiUserReplicationEditor
 					// EReplicatedPropertyContent::Properties
 					+SWidgetSwitcher::Slot()
 					[
-						SAssignNew(ReplicatedProperties, SReplicationTreeView<TSharedPtr<FReplicatedPropertyData>>)
+						SAssignNew(ReplicatedProperties, SReplicatedPropertiesView)
 						.RootItemsSource(&RootPropertyRowData)
 						.OnGetChildren(this, &SPropertyReplicationSelectionViewer::GetPropertyRowChildren)
 						.Columns(Columns)
diff --git a/Engine/Plugins/Developer/Concert/ConcertApp/MultiUserClient/Source/MultiUserReplicationEditor/Private/StreamEditor/View/ObjectViewer/SPropertyReplicationSelectionViewer.h b/Engine/Plugins/Developer/Concert/ConcertApp/MultiUserClient/Source/MultiUserReplicationEditor/Private/StreamEditor/View/ObjectViewer/SPropertyReplicationSelectionViewer.h
--- a/Engine/Plugins/Developer/Concert/ConcertApp/MultiUserClient/Source/MultiUserReplicationEditor/Private/StreamEditor/View/ObjectViewer/SPropertyReplicationSelectionViewer.h
+++ b/Engine/Plugins/Developer/Concert/ConcertApp/MultiUserClient/Source/MultiUserReplicationEditor/Private/StreamEditor/View/ObjectViewer/SPropertyReplicationSelectionViewer.h
@@ -23,6 +23,7 @@ namespace UE::MultiUserReplicationEditor
 	class FReplicatedObjectData;
 	class IEditableObjectToPropertiesModel;
 	class IObjectToPropertiesModel;
+	class SReplicatedPropertiesView;
 
 	DECLARE_DELEGATE(FOnPropertiesRefreshed);
 	DECLARE_DELEGATE_RetVal_TwoParams(bool, FSortPropertyPredicate, const TSharedPtr<FReplicatedPropertyData>& Left, const TSharedPtr<FReplicatedPropertyData>& Right);
@@ -99,7 +100,7 @@ namespace UE::MultiUserReplicationEditor
 		/** Tree view for replicated objects. */
 		TSharedPtr<SReplicationTreeView<TSharedPtr<FReplicatedObjectData>>> ReplicatedObjects;
 		/** Tree view for replicated properties. Content depends on the current object selected. */
-		TSharedPtr<SReplicationTreeView<TSharedPtr<FReplicatedPropertyData>>> ReplicatedProperties;
+		TSharedPtr<SReplicatedPropertiesView> ReplicatedProperties;
 
 		enum class EReplicatedPropertyContent
 		{
diff --git a/Engine/Plugins/Developer/Concert/ConcertApp/MultiUserClient/Source/MultiUserReplicationEditor/Private/StreamEditor/View/ObjectViewer/SReplicationTreeView.h b/Engine/Plugins/Developer/Concert/ConcertApp/MultiUserClient/Source/MultiUserReplicationEditor/Private/StreamEditor/View/ObjectViewer/SReplicationTreeView.h
--- a/Engine/Plugins/Developer/Concert/ConcertApp/MultiUserClient/Source/MultiUserReplicationEditor/Private/StreamEditor/View/ObjectViewer/SReplicationTreeView.h
+++ b/Engine/Plugins/Developer/Concert/ConcertApp/MultiUserClient/Source/MultiUserReplicationEditor/Private/StreamEditor/View/ObjectViewer/SReplicationTreeView.h
@@ -19,7 +19,7 @@ namespace UE::MultiUserReplicationEditor
 {
 	/**
 	 * Shared code for the list view for replicated actors and properties.
-	 * It is a table view that is searchable with a search box.
+	 * It is a table view that is searchable with a search box and exposes slots to add more filter widgets, such as SBasicFilterBar.
 	 */
 	template<typename TItemType>
 	class SReplicationTreeView : public SCompoundWidget
@@ -29,6 +29,7 @@ namespace UE::MultiUserReplicationEditor
 		DECLARE_DELEGATE_OneParam(FDeleteItems, const TArray<TItemType>& SelectedItems);
 		DECLARE_DELEGATE_TwoParams(FGetItemChildren, TItemType Item, TFunctionRef<void(TItemType)> ProcessChild);
 		DECLARE_DELEGATE(FOnSelectionChanged);
+		DECLARE_DELEGATE_RetVal_OneParam(bool, FCustomFilter, const TItemType& Item);
 
 		SLATE_BEGIN_ARGS(SReplicationTreeView<TItemType>)
 			: _SelectionMode(ESelectionMode::Single)
@@ -47,6 +48,9 @@ namespace UE::MultiUserReplicationEditor
 
 			/** Called when the selection changes. Call GetSelectedItems to get the selected items. */
 			SLATE_EVENT(FOnSelectionChanged, OnSelectionChanged)
+
+			/** Optional callback to do even more filtering of items. */
+			SLATE_EVENT(FCustomFilter, FilterItem)
 			
 			/** The columns this list should have */
 			SLATE_ARGUMENT(TArray<TReplicationColumn<TItemType>>, Columns)
@@ -58,6 +62,8 @@ namespace UE::MultiUserReplicationEditor
 
 			/** Optional widget to add to the left of the search bar. */
 			SLATE_NAMED_SLOT(FArguments, LeftOfSearchBar)
+			/** Optional widget to add between the search bar and the table view (e.g. a SBasicFilterBar). */
+			SLATE_NAMED_SLOT(FArguments, RowBelowSearchBar)
 		SLATE_END_ARGS()
 
 		void Construct(const FArguments& InArgs)
@@ -67,6 +73,7 @@ namespace UE::MultiUserReplicationEditor
 
 			OnGetChildrenDelegate = InArgs._OnGetChildren;
 			OnDeleteItemsDelegate = InArgs._OnDeleteItems;
+			CustomFilterDelegate = InArgs._FilterItem;
 			ExpandableColumnId = InArgs._ExpandableColumnLabel;
 			
 			SearchText = MakeShared<FText>();
@@ -103,6 +110,14 @@ namespace UE::MultiUserReplicationEditor
 					]
 				]
 
+				// Optional slot between search bar and table, e.g. for an external SBasicFilterBar
+				+SVerticalBox::Slot()
+				.Padding(1.f)
+				.AutoHeight()
+				[
+					InArgs._RowBelowSearchBar.Widget
+				]
+
 				// Table row
 				+SVerticalBox::Slot()
 				.Padding(1.f)
@@ -163,6 +178,8 @@ namespace UE::MultiUserReplicationEditor
 		FGetItemChildren OnGetChildrenDelegate;
 		/** Optional delegate for responding to pressing the delete button */
 		FDeleteItems OnDeleteItemsDelegate;
+		/** Optional delegate for filtering the items even more. */
+		FCustomFilter CustomFilterDelegate;
 		
 
 		TSharedRef<SWidget> CreateTreeView(const FArguments& InArgs);
@@ -329,7 +346,8 @@ namespace UE::MultiUserReplicationEditor
 	template <typename TItemType>
 	bool SReplicationTreeView<TItemType>::PassesFilters(const TItemType& Item)
 	{
-		return SearchTextFilter->PassesFilter(Item);
+		return SearchTextFilter->PassesFilter(Item)
+			&& (!CustomFilterDelegate.IsBound() || CustomFilterDelegate.Execute(Item));
 	}
 }
 
diff --git a/Engine/Plugins/Developer/Concert/ConcertApp/MultiUserClient/Source/MultiUserReplicationEditor/Private/StreamEditor/View/ObjectViewer/SelectionViewerColumns.cpp b/Engine/Plugins/Developer/Concert/ConcertApp/MultiUserClient/Source/MultiUserReplicationEditor/Private/StreamEditor/View/ObjectViewer/SelectionViewerColumns.cpp
--- a/Engine/Plugins/Developer/Concert/ConcertApp/MultiUserClient/Source/MultiUserReplicationEditor/Private/StreamEditor/View/ObjectViewer/SelectionViewerColumns.cpp
+++ b/Engine/Plugins/Developer/Concert/ConcertApp/MultiUserClient/Source/MultiUserReplicationEditor/Private/StreamEditor/View/ObjectViewer/SelectionViewerColumns.cpp
@@ -2,11 +2,13 @@
 
 #include "SelectionViewerColumns.h"
 
+#include "LogMultiUserReplicationEditor.h"
 #include "ReplicatedPropertyData.h"
 #include "StreamEditor/View/ObjectViewer/ReplicatedObjectData.h"
 
 #include "StreamEditor/Model/DisplayUtils.h"
 #include "Internationalization/Internationalization.h"
+#include "Replication/PropertyChainUtils.h"
 #include "StreamEditor/Model/IEditableObjectToPropertiesModel.h"
 #include "StreamEditor/View/ObjectEditor/SPropertyReplicationSelectionEditor.h"
 #include "Textures/SlateIcon.h"
@@ -95,6 +97,7 @@ namespace UE::MultiUserReplicationEditor::ReplicationPropertyColumns
 {
 	const FName ReplicatesColumnId = TEXT("ReplicatesColumn");
 	const FName LabelColumnId = TEXT("LabelColumn");
+	const FName TypeColumnId = TEXT("TypeColumn");
 
 	FReplicationPropertyColumn ReplicatesColumns(TSharedRef<SPropertyReplicationSelectionEditor> EditorWidget, TSharedRef<IEditableObjectToPropertiesModel> Model)
 	{
@@ -148,6 +151,35 @@ namespace UE::MultiUserReplicationEditor::ReplicationPropertyColumns
 				.ColumnSortOrder(static_cast<int32>(EReplicationPropertyColumnOrder::Label)),
 			SHeaderRow::Column(LabelColumnId)
 				.DefaultLabel(LOCTEXT("LabelColumnLabel", "Label"))
+				.FillSized(500.f)
+			);
+	}
+	
+	FReplicationPropertyColumn TypeColumn()
+	{
+		static auto GetDisplayText = [](const TSharedPtr<FReplicatedPropertyData>& Args)
+		{
+			check(Args);
+			UClass* Class = Args->GetOwningClass().TryLoadClass<UObject>();
+			const FProperty* Property = Class ? ConcertSyncCore::PropertyChain::ResolveProperty(*Class, Args->GetProperty()) : nullptr;
+			return Property ? FText::FromString(Property->GetCPPType()) : LOCTEXT("Unknown", "Unknown");	
+		};
+		
+		return FReplicationPropertyColumn(
+			FReplicationPropertyColumn::FArguments()
+				.GenerateWidgetColumn_Lambda([](const FReplicationPropertyColumn::FBuildArgs& Args)
+				{
+					return SNew(STextBlock)
+						.HighlightText(TAttribute<FText>::CreateLambda([HighlightText = Args.HighlightText](){ return *HighlightText; }))
+						.Text(GetDisplayText(Args.RowData));
+				})
+				.PopulateSearchItems_Lambda([](const TSharedPtr<FReplicatedPropertyData>& ObjectData, TArray<FString>& InOutSearchStrings)
+				{
+					InOutSearchStrings.Add(GetDisplayText(ObjectData).ToString());
+				})
+				.ColumnSortOrder(static_cast<int32>(EReplicationPropertyColumnOrder::Type)),
+			SHeaderRow::Column(TypeColumnId)
+				.DefaultLabel(LOCTEXT("TypeColumnLabel", "Type"))
 				.FillWidth(1.f)
 			);
 	}
diff --git a/Engine/Plugins/Developer/Concert/ConcertApp/MultiUserClient/Source/MultiUserReplicationEditor/Private/StreamEditor/View/ObjectViewer/SelectionViewerColumns.h b/Engine/Plugins/Developer/Concert/ConcertApp/MultiUserClient/Source/MultiUserReplicationEditor/Private/StreamEditor/View/ObjectViewer/SelectionViewerColumns.h
--- a/Engine/Plugins/Developer/Concert/ConcertApp/MultiUserClient/Source/MultiUserReplicationEditor/Private/StreamEditor/View/ObjectViewer/SelectionViewerColumns.h
+++ b/Engine/Plugins/Developer/Concert/ConcertApp/MultiUserClient/Source/MultiUserReplicationEditor/Private/StreamEditor/View/ObjectViewer/SelectionViewerColumns.h
@@ -47,15 +47,19 @@ namespace UE::MultiUserReplicationEditor::ReplicationPropertyColumns
 		/** The checkbox in SPropertyReplicationSelectionEditor determining whether the property is in the selection*/
 		ReplicatesCheckbox = 0,
 		/** Label of the property */
-		Label = 10
+		Label = 10,
+		/** Type of the property */
+		Type = 20
 	};
 	
 	/** The checkbox in SPropertyReplicationSelectionEditor determining whether the property is in the selection*/
 	extern const FName ReplicatesColumnId;
 	extern const FName LabelColumnId;
+	extern const FName TypeColumnId;
 
 	FReplicationPropertyColumn ReplicatesColumns(TSharedRef<SPropertyReplicationSelectionEditor> EditorWidget, TSharedRef<IEditableObjectToPropertiesModel> Model);
 	FReplicationPropertyColumn LabelColumn();
+	FReplicationPropertyColumn TypeColumn();
 
 	/**
 	 * Goes through all selected objects, checks whether the property is checked on it or not, and returns a checkbox state.
diff --git a/Engine/Plugins/Developer/Concert/ConcertSync/ConcertSyncCore/Source/ConcertSyncCore/Private/Replication/Data/ConcertPropertySelection.cpp b/Engine/Plugins/Developer/Concert/ConcertSync/ConcertSyncCore/Source/ConcertSyncCore/Private/Replication/Data/ConcertPropertySelection.cpp
--- a/Engine/Plugins/Developer/Concert/ConcertSync/ConcertSyncCore/Source/ConcertSyncCore/Private/Replication/Data/ConcertPropertySelection.cpp
+++ b/Engine/Plugins/Developer/Concert/ConcertSync/ConcertSyncCore/Source/ConcertSyncCore/Private/Replication/Data/ConcertPropertySelection.cpp
@@ -178,6 +178,13 @@ bool FConcertPropertyChain::MatchesExactly(const FArchiveSerializedPropertyChain
 	return bArePathsEqual && bVisitedEveryProperty;
 }
 
+FProperty* FConcertPropertyChain::ResolveProperty(UStruct& Class, bool bLogOnFail)
+{
+	// FConcertPropertyChain::ResolveProperty exists only for visibility to developers since
+	// the class is the first place one would look and not in the utils namespace.
+	return UE::ConcertSyncCore::PropertyChain::ResolveProperty(Class, *this, bLogOnFail);
+}
+
 FString FConcertPropertyChain::ToString(EToStringMethod Method) const
 {
 	switch (Method)
diff --git a/Engine/Plugins/Developer/Concert/ConcertSync/ConcertSyncCore/Source/ConcertSyncCore/Private/Replication/PropertyChainUtils.cpp b/Engine/Plugins/Developer/Concert/ConcertSync/ConcertSyncCore/Source/ConcertSyncCore/Private/Replication/PropertyChainUtils.cpp
--- a/Engine/Plugins/Developer/Concert/ConcertSync/ConcertSyncCore/Source/ConcertSyncCore/Private/Replication/PropertyChainUtils.cpp
+++ b/Engine/Plugins/Developer/Concert/ConcertSync/ConcertSyncCore/Source/ConcertSyncCore/Private/Replication/PropertyChainUtils.cpp
@@ -2,6 +2,7 @@
 
 #include "Replication/PropertyChainUtils.h"
 
+#include "ConcertLogGlobal.h"
 #include "Misc/ScopeExit.h"
 #include "Serialization/ArchiveSerializedPropertyChain.h"
 #include "UObject/UnrealType.h"
@@ -10,12 +11,12 @@ namespace UE::ConcertSyncCore::PropertyChain
 {
 	namespace Private
 	{
-		EBreakBehavior VisitPropertyRecursive(FArchiveSerializedPropertyChain& Chain, FProperty& Property, TFunctionRef<EBreakBehavior(const FArchiveSerializedPropertyChain& Chain, const FProperty&LeafProperty)> ProcessProperty);
+		EBreakBehavior VisitPropertyRecursive(FArchiveSerializedPropertyChain& Chain, FProperty& Property, TFunctionRef<EBreakBehavior(const FArchiveSerializedPropertyChain&, FProperty&)> ProcessProperty);
 		
 		EBreakBehavior VisitStructPropertyRecursive(
 			FArchiveSerializedPropertyChain& Chain,
 			FStructProperty& StructProperty,
-			TFunctionRef<EBreakBehavior(const FArchiveSerializedPropertyChain& Chain, const FProperty&LeafProperty)> ProcessProperty
+			TFunctionRef<EBreakBehavior(const FArchiveSerializedPropertyChain&, FProperty&)> ProcessProperty
 			)
 		{
 			// If the struct is in a container, then the container should be pushed, not the inner property.
@@ -39,7 +40,7 @@ namespace UE::ConcertSyncCore::PropertyChain
 		EBreakBehavior VisitPropertyRecursive(
 			FArchiveSerializedPropertyChain& Chain,
 			FProperty& Property,
-			TFunctionRef<EBreakBehavior(const FArchiveSerializedPropertyChain& Chain, const FProperty&LeafProperty)> ProcessProperty
+			TFunctionRef<EBreakBehavior(const FArchiveSerializedPropertyChain&, FProperty&)> ProcessProperty
 			)
 		{
 			if (!IsReplicatableProperty(Property))
@@ -102,9 +103,27 @@ namespace UE::ConcertSyncCore::PropertyChain
 		}
 	}
 	
+	FProperty* ResolveProperty(const UStruct& Class, const FConcertPropertyChain& ChainToResolve, bool bLogOnFail)
+	{
+		FProperty* Result = nullptr;
+		ForEachReplicatableProperty(
+			Class,
+			[&ChainToResolve, &Result](const FArchiveSerializedPropertyChain& Chain, FProperty& LeafProperty) mutable
+			{
+				if (ChainToResolve.MatchesExactly(&Chain, LeafProperty))
+				{
+					Result = &LeafProperty;
+					return EBreakBehavior::Break;
+				}
+				return EBreakBehavior::Continue;
+			});
+		UE_CLOG(bLogOnFail && Result == nullptr, LogConcert, Warning, TEXT("Property chain \"%s\" resolved to no property!"), *ChainToResolve.ToString());
+		return Result;
+	}
+	
 	void ForEachReplicatableProperty(
-		UStruct& Class,
-		TFunctionRef<EBreakBehavior(const FArchiveSerializedPropertyChain& Chain, const FProperty& LeafProperty)> ProcessProperty)
+		const UStruct& Class,
+		TFunctionRef<EBreakBehavior(const FArchiveSerializedPropertyChain& Chain, FProperty& LeafProperty)> ProcessProperty)
 	{
 		FArchiveSerializedPropertyChain Chain;
 		for (TFieldIterator<FProperty> FieldIt(&Class); FieldIt; ++FieldIt)
@@ -117,23 +136,23 @@ namespace UE::ConcertSyncCore::PropertyChain
 	}
 
 	void ForEachReplicatableConcertProperty(
-		UStruct& Class,
+		const UStruct& Class,
 		TFunctionRef<EBreakBehavior(FConcertPropertyChain&& PropertyChain)> ProcessProperty)
 	{
-		ForEachReplicatableProperty(Class, [&ProcessProperty](const FArchiveSerializedPropertyChain& Chain, const FProperty& LeafProperty)
+		ForEachReplicatableProperty(Class, [&ProcessProperty](const FArchiveSerializedPropertyChain& Chain, FProperty& LeafProperty)
 		{
 			return ProcessProperty(FConcertPropertyChain(&Chain, LeafProperty));
 		});
 	}
 
-	void BulkConstructConcertChainsFromPaths(UStruct& Class, uint32 NumPaths, TFunctionRef<bool(const FArchiveSerializedPropertyChain& Chain, const FProperty& LeafProperty)> MatchesPath)
+	void BulkConstructConcertChainsFromPaths(const UStruct& Class, uint32 NumPaths, TFunctionRef<bool(const FArchiveSerializedPropertyChain& Chain, FProperty& LeafProperty)> MatchesPath)
 	{
 		if (NumPaths == 0)
 		{
 			return;
 		}
 		
-		ForEachReplicatableProperty(Class, [&NumPaths, &MatchesPath](const FArchiveSerializedPropertyChain& Chain, const FProperty& LeafProperty) mutable
+		ForEachReplicatableProperty(Class, [&NumPaths, &MatchesPath](const FArchiveSerializedPropertyChain& Chain, FProperty& LeafProperty) mutable
 		{
 			if (MatchesPath(Chain, LeafProperty))
 			{
@@ -166,12 +185,16 @@ namespace UE::ConcertSyncCore::PropertyChain
 	
 	bool IsReplicatableProperty(const FProperty& LeafProperty)
 	{
-		return !LeafProperty.HasAnyPropertyFlags(
+		const bool bHasValidFlags = !LeafProperty.HasAnyPropertyFlags(
 			// Replicating delegates makes no sense
 			CPF_BlueprintAssignable
 			// It does not make sense to serialize the reference of an instanced subobject. Instead the subobject should be added to the list of replicated properties.
 			| CPF_InstancedReference /* for object ptrs */ | CPF_ContainsInstancedReference /* for containers of object ptrs */
 			);
+		const bool bIsAllowedObjectProperty = CastField<FObjectPropertyBase>(&LeafProperty) == nullptr
+			// Soft object properties are allowed because they are serialized as a string.
+			|| CastField<FSoftObjectProperty>(&LeafProperty) != nullptr;
+		return bHasValidFlags && bIsAllowedObjectProperty;
 	}
 
 	bool IsInnerContainerProperty(const FProperty& Property)
diff --git a/Engine/Plugins/Developer/Concert/ConcertSync/ConcertSyncCore/Source/ConcertSyncCore/Public/Replication/Data/ConcertPropertySelection.h b/Engine/Plugins/Developer/Concert/ConcertSync/ConcertSyncCore/Source/ConcertSyncCore/Public/Replication/Data/ConcertPropertySelection.h
--- a/Engine/Plugins/Developer/Concert/ConcertSync/ConcertSyncCore/Source/ConcertSyncCore/Public/Replication/Data/ConcertPropertySelection.h
+++ b/Engine/Plugins/Developer/Concert/ConcertSync/ConcertSyncCore/Source/ConcertSyncCore/Public/Replication/Data/ConcertPropertySelection.h
@@ -56,6 +56,9 @@ struct CONCERTSYNCCORE_API FConcertPropertyChain
 	/** @return Whether OptionalChain and LeafProperty correspond to this path. */
 	bool MatchesExactly(const FArchiveSerializedPropertyChain* OptionalChain, const FProperty& LeafProperty) const;
 
+	/** @return Attempts to resolve this property given the class */
+	FProperty* ResolveProperty(UStruct& Class, bool bLogOnFail = true);
+
 	const TArray<FName>& GetPathToProperty() const { return PathToProperty; }
 
 	enum class EToStringMethod
diff --git a/Engine/Plugins/Developer/Concert/ConcertSync/ConcertSyncCore/Source/ConcertSyncCore/Public/Replication/PropertyChainUtils.h b/Engine/Plugins/Developer/Concert/ConcertSync/ConcertSyncCore/Source/ConcertSyncCore/Public/Replication/PropertyChainUtils.h
--- a/Engine/Plugins/Developer/Concert/ConcertSync/ConcertSyncCore/Source/ConcertSyncCore/Public/Replication/PropertyChainUtils.h
+++ b/Engine/Plugins/Developer/Concert/ConcertSync/ConcertSyncCore/Source/ConcertSyncCore/Public/Replication/PropertyChainUtils.h
@@ -12,6 +12,9 @@ struct FArchiveSerializedPropertyChain;
 
 namespace UE::ConcertSyncCore::PropertyChain
 {
+	/** Tries to find the property referenced by the Chain. */
+	CONCERTSYNCCORE_API FProperty* ResolveProperty(const UStruct& Class, const FConcertPropertyChain& ChainToResolve, bool bLogOnFail = true);
+
 	/**
 	 * Iterates all properties that are supported for replication.
 	 * 
@@ -23,8 +26,8 @@ namespace UE::ConcertSyncCore::PropertyChain
 	 * @param OnlyChildrenOf Instead of iterating all properties, you can restrict to only receive the child properties of OnlyChildrenOf.
 	 */
 	CONCERTSYNCCORE_API void ForEachReplicatableProperty(
-		UStruct& Class,
-		TFunctionRef<EBreakBehavior(const FArchiveSerializedPropertyChain& Chain, const FProperty& LeafProperty)> ProcessProperty
+		const UStruct& Class,
+		TFunctionRef<EBreakBehavior(const FArchiveSerializedPropertyChain& Chain, FProperty& LeafProperty)> ProcessProperty
 		);
 
 	/**
@@ -38,7 +41,7 @@ namespace UE::ConcertSyncCore::PropertyChain
 	 * @param OnlyChildrenOf Instead of iterating all properties, you can restrict to only receive the child properties of OnlyChildrenOf.
 	 */
 	CONCERTSYNCCORE_API void ForEachReplicatableConcertProperty(
-		UStruct& Class,
+		const UStruct& Class,
 		TFunctionRef<EBreakBehavior(FConcertPropertyChain&& PropertyChain)> ProcessProperty
 		);
 
@@ -51,9 +54,9 @@ namespace UE::ConcertSyncCore::PropertyChain
 	 * @param MatchesPath Determines whether you are interested in this path. If so, do something with it and return true.
 	 */
 	CONCERTSYNCCORE_API void BulkConstructConcertChainsFromPaths(
-		UStruct& Class,
+		const UStruct& Class,
 		uint32 NumPaths,
-		TFunctionRef<bool(const FArchiveSerializedPropertyChain& Chain, const FProperty& LeafProperty)> MatchesPath
+		TFunctionRef<bool(const FArchiveSerializedPropertyChain& Chain, FProperty& LeafProperty)> MatchesPath
 		);
 
 	/** Util that can be used with BulkConstructConcertChainsFromPaths. Checks whether Path coincides with Chain & LeafProperty.*/
