Integrate MU replication system into the editor.
- Add replication tab to Concert window (in active session view)

diff --git a/Engine/Plugins/Developer/Concert/ConcertApp/MultiUserClient/Source/MultiUserClient/MultiUserClient.Build.cs b/Engine/Plugins/Developer/Concert/ConcertApp/MultiUserClient/Source/MultiUserClient/MultiUserClient.Build.cs
--- a/Engine/Plugins/Developer/Concert/ConcertApp/MultiUserClient/Source/MultiUserClient/MultiUserClient.Build.cs
+++ b/Engine/Plugins/Developer/Concert/ConcertApp/MultiUserClient/Source/MultiUserClient/MultiUserClient.Build.cs
@@ -29,8 +29,9 @@ namespace UnrealBuildTool.Rules
 					"DesktopPlatform",
 					"EditorStyle",
 					"InputCore",
-					"Projects",
 					"MessageLog",
+					"MultiUserReplication",
+					"Projects",
 					"Slate",
 					"SlateCore",
 					"SourceControl",
@@ -41,13 +42,16 @@ namespace UnrealBuildTool.Rules
 
 			if (Target.bBuildEditor)
 			{
-				PrivateDependencyModuleNames.AddRange(
+				PrivateDependencyModuleNames.AddRange( 
 					new string[]
-					{
+					{ 
 						"EditorFramework",
-						"UnrealEd",
+						"MultiUserReplicationEditor",
+						"PropertyEditor",
+						"Settings",
 						"Sequencer",
 						"ToolMenus",
+						"UnrealEd",
 						"WorkspaceMenuStructure",
 					}
 				);
diff --git a/Engine/Plugins/Developer/Concert/ConcertApp/MultiUserClient/Source/MultiUserClient/Private/Customization/ShowJoinSettingsCustomization.cpp b/Engine/Plugins/Developer/Concert/ConcertApp/MultiUserClient/Source/MultiUserClient/Private/Customization/ShowJoinSettingsCustomization.cpp
new file mode 100644
--- /dev/null
+++ b/Engine/Plugins/Developer/Concert/ConcertApp/MultiUserClient/Source/MultiUserClient/Private/Customization/ShowJoinSettingsCustomization.cpp
@@ -0,0 +1,33 @@
+// Copyright Epic Games, Inc. All Rights Reserved.
+
+#include "ShowJoinSettingsCustomization.h"
+
+#include "DetailCategoryBuilder.h"
+#include "DetailLayoutBuilder.h"
+#include "Settings/MultiUserReplicationSettings.h"
+
+namespace UE::MultiUserClient
+{
+	TSharedRef<IDetailCustomization> FShowJoinSettingsCustomization::Make()
+	{
+		return MakeShared<FShowJoinSettingsCustomization>();
+	}
+
+	void FShowJoinSettingsCustomization::CustomizeDetails(IDetailLayoutBuilder& DetailBuilder)
+	{
+		TSharedRef<IPropertyHandle> DefaultPropertySelection = DetailBuilder.GetProperty(GET_MEMBER_NAME_CHECKED(UMultiUserReplicationSettings, DefaultPropertySelection));
+		TSharedRef<IPropertyHandle> DefaultSessionSettings = DetailBuilder.GetProperty(GET_MEMBER_NAME_CHECKED(UMultiUserReplicationSettings, DefaultSessionSettings));
+		
+		DetailBuilder.HideCategory("Editor");
+		DetailBuilder.HideProperty(DefaultSessionSettings);
+		DetailBuilder.HideProperty(DefaultPropertySelection);
+		
+		TSharedPtr<IPropertyHandle> AutoJoin = DefaultSessionSettings->GetChildHandle(GET_MEMBER_NAME_CHECKED(FDefaultReplicationSessionSettings, bAutoJoin));
+		TSharedPtr<IPropertyHandle> DefaultProfile = DefaultSessionSettings->GetChildHandle(GET_MEMBER_NAME_CHECKED(FDefaultReplicationSessionSettings, DefaultProfile));
+		check(AutoJoin && DefaultProfile);
+		
+		IDetailCategoryBuilder& CategoryBuilder = DetailBuilder.EditCategory("Settings");
+		CategoryBuilder.AddProperty(AutoJoin);
+		CategoryBuilder.AddProperty(DefaultProfile);
+	}
+}
diff --git a/Engine/Plugins/Developer/Concert/ConcertApp/MultiUserClient/Source/MultiUserClient/Private/Customization/ShowJoinSettingsCustomization.h b/Engine/Plugins/Developer/Concert/ConcertApp/MultiUserClient/Source/MultiUserClient/Private/Customization/ShowJoinSettingsCustomization.h
new file mode 100644
--- /dev/null
+++ b/Engine/Plugins/Developer/Concert/ConcertApp/MultiUserClient/Source/MultiUserClient/Private/Customization/ShowJoinSettingsCustomization.h
@@ -0,0 +1,21 @@
+// Copyright Epic Games, Inc. All Rights Reserved.
+
+#pragma once
+
+#include "IDetailCustomization.h"
+#include "Templates/SharedPointer.h"
+
+namespace UE::MultiUserClient
+{
+	/** Customizes UMultiUserReplicationSettings to show only the settings relevant for joining */
+	class FShowJoinSettingsCustomization : public IDetailCustomization
+	{
+	public:
+
+		static TSharedRef<IDetailCustomization> Make();
+		
+		//~ Begin IDetailCustomization Interface
+		virtual void CustomizeDetails(IDetailLayoutBuilder& DetailBuilder) override;
+		//~ End IDetailCustomization Interface
+	};
+}
diff --git a/Engine/Plugins/Developer/Concert/ConcertApp/MultiUserClient/Source/MultiUserClient/Private/MultiUserClientModule.cpp b/Engine/Plugins/Developer/Concert/ConcertApp/MultiUserClient/Source/MultiUserClient/Private/MultiUserClientModule.cpp
--- a/Engine/Plugins/Developer/Concert/ConcertApp/MultiUserClient/Source/MultiUserClient/Private/MultiUserClientModule.cpp
+++ b/Engine/Plugins/Developer/Concert/ConcertApp/MultiUserClient/Source/MultiUserClient/Private/MultiUserClientModule.cpp
@@ -1041,8 +1041,9 @@ private:
 	 */
 	TSharedRef<SDockTab> SpawnConcertBrowserTab(const FSpawnTabArgs& SpawnTabArgs)
 	{
-		const TSharedRef<SDockTab> DockTab = SNew(SDockTab).TabRole(ETabRole::NomadTab);
-		DockTab->SetContent(SNew(SConcertBrowser, DockTab, SpawnTabArgs.GetOwnerWindow(), MultiUserClient));
+		const TSharedRef<SDockTab> DockTab = SNew(SDockTab)
+			.TabRole(NomadTab);
+		DockTab->SetContent(SNew(SConcertBrowser, DockTab, MultiUserClient));
 		return DockTab;
 	}
 
diff --git a/Engine/Plugins/Developer/Concert/ConcertApp/MultiUserClient/Source/MultiUserClient/Private/Util/ReplicationUtils.cpp b/Engine/Plugins/Developer/Concert/ConcertApp/MultiUserClient/Source/MultiUserClient/Private/Util/ReplicationUtils.cpp
new file mode 100644
--- /dev/null
+++ b/Engine/Plugins/Developer/Concert/ConcertApp/MultiUserClient/Source/MultiUserClient/Private/Util/ReplicationUtils.cpp
@@ -0,0 +1,69 @@
+// Copyright Epic Games, Inc. All Rights Reserved.
+
+#include "ReplicationUtils.h"
+
+#include "Replication/Messages/ConcertReplicationHandshakeMessages.h"
+#include "Replication/IConcertClientReplicationManager.h"
+
+#include "Async/Async.h"
+#include "Framework/Notifications/NotificationManager.h"
+#include "Widgets/Notifications/SNotificationList.h"
+
+#define LOCTEXT_NAMESPACE "ReplicationUtils"
+
+namespace UE::MultiUserClient::ReplicationUtils
+{
+	TFuture<ConcertSyncClient::Replication::FJoinReplicatedSessionResult> JoinSessionWithEditorNotifications(
+		IConcertClientReplicationManager& Manager,
+		ConcertSyncClient::Replication::FJoinReplicatedSessionArgs Args)
+	{
+		using namespace ConcertSyncClient::Replication;
+		FNotificationInfo Notification(LOCTEXT("Joining", "Joining Replication Session"));
+		Notification.bFireAndForget = false;
+		Notification.ExpireDuration = 6.f;
+		Notification.bUseThrobber = true;
+		TSharedPtr<SNotificationItem> NotificationItem = FSlateNotificationManager::Get().AddNotification(Notification);
+
+		TPromise<FJoinReplicatedSessionResult> Promise;
+		TFuture<FJoinReplicatedSessionResult> Future = Promise.GetFuture();
+		Manager.JoinReplicationSession(MoveTemp(Args))
+			.Next([Promise = MoveTemp(Promise), NotificationItem](FJoinReplicatedSessionResult ErrorResult) mutable
+			{
+				// We may already be on the game thread. In some cases, like a timeout, we'll be executed on the UDP's thread.
+				// NotificationItem may invalidate itself, which accesses the FSlateApplication singleton - which must be done on the game thread.
+				AsyncTask(ENamedThreads::GameThread, [Promise = MoveTemp(Promise), NotificationItem, ErrorResult]() mutable
+				{
+					const bool bSuccess = ErrorResult.ErrorCode == EJoinReplicationErrorCode::Success;
+					if (bSuccess)
+					{
+						NotificationItem->SetText(LOCTEXT("Joined", "Joined Replication Session"));
+						NotificationItem->SetCompletionState(SNotificationItem::CS_Success);
+					}
+					else
+					{
+						const FString ErrorCodeText = ConcertSyncCore::Replication::LexJoinErrorCode(ErrorResult.ErrorCode);
+						const FText Message = ErrorResult.DetailedErrorMessage.IsEmpty()
+							? FText::Format(LOCTEXT("Error.MissingInfoFmt", "Server returned error code {0}"), FText::FromString(ErrorCodeText))
+							: FText::Format(LOCTEXT("Error.DetailedInfoFmt", "{0} ({1})"), FText::FromString(ErrorResult.DetailedErrorMessage), FText::FromString(ErrorCodeText));
+						NotificationItem->SetSubText(Message);
+						NotificationItem->SetCompletionState(SNotificationItem::CS_Fail);
+					}
+					NotificationItem->ExpireAndFadeout();
+					Promise.EmplaceValue(ErrorResult);
+				});
+			});
+		return Future;
+	}
+	
+	void LeaveSessionWithEditorNotifications(IConcertClientReplicationManager& Manager)
+	{
+		Manager.LeaveReplicationSession();
+		
+		FNotificationInfo Notification(LOCTEXT("Left", "Left Replication Session"));
+		Notification.bFireAndForget = true;
+		Notification.ExpireDuration = 4.f;
+		FSlateNotificationManager::Get().AddNotification(Notification);
+	}
+}
+
+#undef LOCTEXT_NAMESPACE
\ No newline at end of file
diff --git a/Engine/Plugins/Developer/Concert/ConcertApp/MultiUserClient/Source/MultiUserClient/Private/Util/ReplicationUtils.h b/Engine/Plugins/Developer/Concert/ConcertApp/MultiUserClient/Source/MultiUserClient/Private/Util/ReplicationUtils.h
new file mode 100644
--- /dev/null
+++ b/Engine/Plugins/Developer/Concert/ConcertApp/MultiUserClient/Source/MultiUserClient/Private/Util/ReplicationUtils.h
@@ -0,0 +1,30 @@
+// Copyright Epic Games, Inc. All Rights Reserved.
+
+#pragma once
+
+#include "Async/Future.h"
+
+namespace UE::ConcertSyncClient::Replication
+{
+	struct FJoinReplicatedSessionResult;
+	struct FJoinReplicatedSessionArgs;
+}
+
+class IConcertClientReplicationManager;
+
+namespace UE::MultiUserClient::ReplicationUtils
+{
+	/**
+	 * Requests to join the replication session in the currently.
+	 * 
+	 * This decorates the request with fancy editor UI notifications on the bottom-right: One for the duration of the process,
+	 * and then another for either success (leaves after a few seconds) or failure (must be dismissed).
+	 */
+	TFuture<ConcertSyncClient::Replication::FJoinReplicatedSessionResult> JoinSessionWithEditorNotifications(
+		IConcertClientReplicationManager& Manager,
+		ConcertSyncClient::Replication::FJoinReplicatedSessionArgs Args
+		);
+
+	/** Leaves the replications session and shows an editor UI notification on the bottom-right. */
+	void LeaveSessionWithEditorNotifications(IConcertClientReplicationManager& Manager);
+}
diff --git a/Engine/Plugins/Developer/Concert/ConcertApp/MultiUserClient/Source/MultiUserClient/Private/Widgets/ActiveSession/Overview/SActiveSessionOverviewTab.cpp b/Engine/Plugins/Developer/Concert/ConcertApp/MultiUserClient/Source/MultiUserClient/Private/Widgets/ActiveSession/Overview/SActiveSessionOverviewTab.cpp
new file mode 100644
--- /dev/null
+++ b/Engine/Plugins/Developer/Concert/ConcertApp/MultiUserClient/Source/MultiUserClient/Private/Widgets/ActiveSession/Overview/SActiveSessionOverviewTab.cpp
@@ -0,0 +1,761 @@
+// Copyright Epic Games, Inc. All Rights Reserved.
+
+#include "Widgets/ActiveSession/Overview/SActiveSessionOverviewTab.h"
+
+#include "ConcertActionDefinition.h"
+#include "ConcertClientFrontendUtils.h"
+#include "ConcertClientSettings.h"
+#include "ConcertFrontendUtils.h"
+#include "ConcertMessageData.h"
+#include "ConcertMessages.h"
+#include "Dialog/SMessageDialog.h"
+#include "IConcertClientPresenceManager.h"
+#include "IConcertClient.h"
+#include "IConcertClientWorkspace.h"
+#include "IConcertSyncClient.h"
+#include "Widgets/ClientSessionHistoryController.h"
+
+#include "Algo/Transform.h"
+#include "EditorFontGlyphs.h"
+#include "FileHelpers.h"
+#include "Misc/AsyncTaskNotification.h"
+#include "Misc/PackageName.h"
+#include "Dialog/SCustomDialog.h"
+#include "Styling/AppStyle.h"
+#include "Widgets/Text/STextBlock.h"
+#include "Widgets/Images/SImage.h"
+#include "Widgets/Input/SHyperlink.h"
+#include "Widgets/Layout/SExpandableArea.h"
+#include "Widgets/Views/STableViewBase.h"
+#include "Widgets/Views/STableRow.h"
+#include "HAL/PlatformFileManager.h"
+
+#define LOCTEXT_NAMESPACE "SActiveSessionOverviewTab"
+
+namespace ActiveSessionDetailsUI
+{
+	static const FName DisplayNameColumnName(TEXT("DisplayName"));
+	static const FName PresenceColumnName(TEXT("Presence"));
+	static const FName LevelColumnName(TEXT("Level"));
+}
+
+
+class SActiveSessionDetailsRow : public SMultiColumnTableRow<TSharedPtr<FConcertSessionClientInfo>>
+{
+	SLATE_BEGIN_ARGS(SActiveSessionDetailsRow) {}
+	SLATE_END_ARGS()
+
+public:
+	/**
+	 * Constructs the widget.
+	 *
+	 * @param InArgs The construction arguments.
+	 * @param InClientInfo The client displayed by this row.
+	 * @param InClientSession The session in which the client is, used to determine if the client is the local one, so that we can suffix it with a "you".
+	 * @param InOwnerTableView The table to which the row must be added.
+	 */
+	void Construct(const FArguments& InArgs, TWeakPtr<IConcertSyncClient> InSyncClient, TWeakPtr<IConcertClientSession> InClientSession, TSharedPtr<FConcertSessionClientInfo> InClientInfo, const TSharedRef<STableViewBase>& InOwnerTableView)
+	{
+		SyncClient = MoveTemp(InSyncClient);
+		ClientSession = MoveTemp(InClientSession);
+		SessionClientInfo = MoveTemp(InClientInfo);
+		SMultiColumnTableRow<TSharedPtr<FConcertSessionClientInfo>>::Construct(FSuperRowType::FArguments(), InOwnerTableView);
+
+		// Set the tooltip for the entire row. Will show up unless there is another item with a tooltip hovered in the row, such as the "presence" icons.
+		SetToolTipText(MakeAttributeSP(this, &SActiveSessionDetailsRow::GetRowToolTip));
+	}
+
+public:
+	virtual TSharedRef<SWidget> GenerateWidgetForColumn(const FName& ColumnName) override
+	{
+		if (ColumnName == ActiveSessionDetailsUI::DisplayNameColumnName)
+		{
+			// Displays a colored square from a special font (using avatar color) followed by the the display name -> [x] John Smith
+			return SNew(SHorizontalBox)
+				// The 'square' glyph in front of the client name, rendered using special font glyph, in the client avatar color.
+				+SHorizontalBox::Slot()
+				.AutoWidth()
+				.VAlign(VAlign_Center)
+				.Padding(FMargin(4.0f, 0.0f, 0.0f, 0.0f))
+				[
+					SNew(STextBlock)
+					.Font(this, &SActiveSessionDetailsRow::GetAvatarFont)
+					.ColorAndOpacity(this, &SActiveSessionDetailsRow::GetAvatarColor)
+					.Text(FEditorFontGlyphs::Square)
+				]
+
+				// The client display name.
+				+SHorizontalBox::Slot()
+				.VAlign(VAlign_Center)
+				[
+					ConcertFrontendUtils::CreateDisplayName(MakeAttributeSP(this, &SActiveSessionDetailsRow::GetDisplayName))
+				];
+		}
+		else if (ColumnName == ActiveSessionDetailsUI::PresenceColumnName)
+		{
+			// Displays a set of icons corresponding to the client presence. The set may be extended later to include other functionalities.
+			TArray<FConcertActionDefinition> ActionDefs;
+			TSharedRef<SHorizontalBox> PresenceCell = SNew(SHorizontalBox);
+
+			TSharedPtr<FConcertSessionClientInfo> ClientInfoPin = SessionClientInfo.Pin();
+			TSharedPtr<IConcertSyncClient> SyncClientPin = SyncClient.Pin();
+			if (ClientInfoPin.IsValid() && SyncClientPin.IsValid())
+			{
+				SyncClientPin->GetSessionClientActions(*ClientInfoPin, ActionDefs);
+				ConcertClientFrontendUtils::AppendButtons(PresenceCell, ActionDefs);
+			}
+			return PresenceCell;
+		}
+		else // LevelColumnName
+		{
+			check(ColumnName == ActiveSessionDetailsUI::LevelColumnName); // If this fail, was a column added/removed/renamed ?
+
+			TSharedPtr<FConcertSessionClientInfo> ClientInfoPin = SessionClientInfo.Pin();
+			TSharedPtr<IConcertClientSession> ClientSessionPin = ClientSession.Pin();
+			TSharedPtr<SWidget> LevelWidget;
+
+			// If the row is not for "this client", then add a hyper link to open the other client level.
+			if (ClientInfoPin.IsValid() && ClientSessionPin.IsValid() && ClientInfoPin->ClientEndpointId != ClientSessionPin->GetSessionClientEndpointId())
+			{
+				LevelWidget = SNew(SOverlay)
+					+SOverlay::Slot()
+					[
+						SNew(SHyperlink)
+						.Visibility_Lambda([this]() { return GetOtherLevelVisibility(); })
+						.Text(this, &SActiveSessionDetailsRow::GetLevel)
+						.ToolTipText(this, &SActiveSessionDetailsRow::GetOtherLevelTooltip)
+						.OnNavigate(this, &SActiveSessionDetailsRow::OnOtherLevelClicked)
+					]
+					+SOverlay::Slot()
+					[
+						SNew(STextBlock)
+						.Visibility_Lambda([this]() { return GetOtherLevelVisibility() == EVisibility::Visible ? EVisibility::Collapsed : EVisibility::Visible; })
+						.Text(LOCTEXT("LevelNotAvailableInGame", "N/A (-game)")) // The only known reason for the level to be unknown and invisible is b/c the user is in -game.
+						.ToolTipText(LOCTEXT("LevelNotAvailableInGame_Tooltip", "The presence/level information is not emitted by clients running Editor in -game mode."))
+					];
+			}
+			else
+			{
+				LevelWidget = SNew(STextBlock)
+					.Text(this, &SActiveSessionDetailsRow::GetLevel)
+					.ToolTipText(FText(LOCTEXT("YourLevel", "The level opened in your editor.")));
+			}
+
+			// Displays which "level" the client is editing, playing (PIE) or simulating (SIE).
+			return SNew(SHorizontalBox)
+				+SHorizontalBox::Slot()
+				.Padding(0.0f, 4.0f, 2.0f, 6.0f)
+				.AutoWidth()
+				[
+					SNew(SImage)
+					.Image(FAppStyle::Get().GetBrush("LevelEditor.Tabs.Levels"))
+				]
+				+SHorizontalBox::Slot()
+				.AutoWidth()
+				.Padding(FMargin(2.0f, 4.0f))
+				[
+					LevelWidget.ToSharedRef()
+				];
+		}
+	}
+
+	FText GetRowToolTip() const
+	{
+		// This is a tooltip for the entire row. Like display name, the tooltip will not update in real time if the user change its
+		// settings. See GetDisplayName() for more info.
+		TSharedPtr<FConcertSessionClientInfo> ClientInfoPin = SessionClientInfo.Pin();
+		return ClientInfoPin.IsValid() ? ClientInfoPin->ToDisplayString() : FText();
+	}
+
+	FText GetDisplayName() const
+	{
+		TSharedPtr<FConcertSessionClientInfo> ClientInfoPin = SessionClientInfo.Pin();
+		if (ClientInfoPin.IsValid())
+		{
+			// NOTE: The display name doesn't update in real time at the moment because the concert setting are not propagated
+			//       until the client disconnect/reconnect. Since those settings should not change often, this should not
+			//       be a major deal breaker for the users.
+			TSharedPtr<IConcertClientSession> ClientSessionPin = ClientSession.Pin();
+			if (ClientSessionPin.IsValid() && ClientInfoPin->ClientEndpointId == ClientSessionPin->GetSessionClientEndpointId())
+			{
+				return FText::Format(LOCTEXT("ClientDisplayNameIsYouFmt", "{0} (You)"), FText::FromString(ClientSessionPin->GetLocalClientInfo().DisplayName));
+			}
+
+			// Return the ClientInfo cached.
+			return FText::FromString(ClientInfoPin->ClientInfo.DisplayName);
+		}
+
+		return FText();
+	}
+
+	EVisibility GetOtherLevelVisibility() const
+	{
+		TSharedPtr<FConcertSessionClientInfo> ClientInfoPin = SessionClientInfo.Pin();
+		TSharedPtr<IConcertSyncClient> SyncClientPin = SyncClient.Pin();
+
+		if (ClientInfoPin && SyncClientPin)
+		{
+			EEditorPlayMode EditorPlayMode = EEditorPlayMode::None;
+			if (IConcertClientPresenceManager* PresenceManager = SyncClientPin->GetPresenceManager())
+			{
+				if (!PresenceManager->GetPresenceWorldPath(ClientInfoPin->ClientEndpointId, EditorPlayMode).IsEmpty())
+				{
+					return EVisibility::Visible;
+				}
+			}
+		}
+
+		return EVisibility::Hidden;
+	}
+
+	FText GetOtherLevelTooltip() const
+	{
+		if (TSharedPtr<FConcertSessionClientInfo> ClientInfoPin = SessionClientInfo.Pin())
+		{
+			return FText::Format(LOCTEXT("OtherUserLevel_Tooltip", "The level opened in {0} editor. Click to open this level."), FText::FromString(ClientInfoPin->ClientInfo.DisplayName));
+		}
+
+		return FText::GetEmpty();
+	}
+
+	FText GetLevel() const
+	{
+		FString WorldPath;
+		
+		TSharedPtr<FConcertSessionClientInfo> ClientInfoPin = SessionClientInfo.Pin();
+		TSharedPtr<IConcertSyncClient> SyncClientPin = SyncClient.Pin();
+
+		if (ClientInfoPin && SyncClientPin)
+		{
+			// The world path returned by the function below is always the one from the editor context, never the one from the PIE/SIE context, so
+			// it will not contain the prefix added to the world when in PIE/SIE.
+			EEditorPlayMode EditorPlayMode = EEditorPlayMode::None;
+			if (IConcertClientPresenceManager* PresenceManager = SyncClientPin->GetPresenceManager())
+			{
+				WorldPath = PresenceManager->GetPresenceWorldPath(ClientInfoPin->ClientEndpointId, EditorPlayMode);
+			}
+			
+			// The world path is returned as something like /Game/MyMap.MyMap, but we are only interested to keep the
+			// string left to the '.' to display "/Game/MyMap"
+			WorldPath = FPackageName::ObjectPathToPackageName(WorldPath);
+
+			if (EditorPlayMode == EEditorPlayMode::PIE)
+			{
+				WorldPath += TEXT(" (PIE)");
+			}
+			else if (EditorPlayMode == EEditorPlayMode::SIE)
+			{
+				WorldPath += TEXT(" (SIE)");
+			}
+		}
+
+		return FText::FromString(WorldPath);
+	}
+
+	void OnOtherLevelClicked()
+	{
+		TSharedPtr<FConcertSessionClientInfo> ClientInfoPin = SessionClientInfo.Pin();
+		TSharedPtr<IConcertClientSession> ClientSessionPin = ClientSession.Pin();
+		TSharedPtr<IConcertSyncClient> ConcertSyncClientPin = SyncClient.Pin();
+
+		// If the user clicks on a client row other than the one that represents them.
+		if (ClientInfoPin && ClientSessionPin && ConcertSyncClientPin && ConcertSyncClientPin->GetPresenceManager() &&
+			ClientInfoPin->ClientEndpointId != ClientSessionPin->GetSessionClientEndpointId())
+		{
+			IConcertClientPresenceManager* PresenceManager = ConcertSyncClientPin->GetPresenceManager();
+			EEditorPlayMode EditorPlayMode; // The result stored in this variable is not used.
+			FString RowWorldPath = PresenceManager->GetPresenceWorldPath(ClientInfoPin->ClientEndpointId, EditorPlayMode);
+			FString ThisClientSessionWorldPath = PresenceManager->GetPresenceWorldPath(ClientSessionPin->GetSessionClientEndpointId(), EditorPlayMode);
+
+			// If this client is not editing the same world as the client represented by this row.
+			if (ThisClientSessionWorldPath != RowWorldPath)
+			{
+				// If there are any unsaved changes to the current level, see if the user wants to save those first.
+				bool bPromptUserToSave = true;
+				bool bSaveMapPackages = true;
+				bool bSaveContentPackages = false;
+				if (FEditorFileUtils::SaveDirtyPackages(bPromptUserToSave, bSaveMapPackages, bSaveContentPackages) == false)
+				{
+					return;
+				}
+				
+				// Open the level.
+				FString WorlPathname;
+				const bool bLoadAsTemplate = false;
+				const bool bShowProgress = true;
+				FPackageName::TryConvertLongPackageNameToFilename(RowWorldPath, WorlPathname);
+				FEditorFileUtils::LoadMap(WorlPathname, bLoadAsTemplate, bShowProgress);
+
+				// Jump to the client position/orientation to view the same thing as client corresponding to the clicked row currently look at.
+				PresenceManager->InitiateJumpToPresence(ClientInfoPin->ClientEndpointId);
+			}
+		}
+	}
+
+	FSlateFontInfo GetAvatarFont() const
+	{
+		// This font is used to render a small square box filled with the avatar color.
+		FSlateFontInfo ClientIconFontInfo = FAppStyle::Get().GetFontStyle(ConcertClientFrontendUtils::ButtonIconSyle);
+		ClientIconFontInfo.Size = 8;
+		ClientIconFontInfo.OutlineSettings.OutlineSize = 1;
+
+		TSharedPtr<FConcertSessionClientInfo> ClientInfoPin = SessionClientInfo.Pin();
+		if (ClientInfoPin.IsValid())
+		{
+			FLinearColor ClientOutlineColor = ClientInfoPin->ClientInfo.AvatarColor * 0.6f; // Make the font outline darker.
+			ClientOutlineColor.A = ClientInfoPin->ClientInfo.AvatarColor.A; // Put back the original alpha.
+			ClientIconFontInfo.OutlineSettings.OutlineColor = ClientOutlineColor;
+		}
+		else
+		{
+			ClientIconFontInfo.OutlineSettings.OutlineColor = FLinearColor(0.75, 0.75, 0.75); // This is an arbitrary color.
+		}
+
+		return ClientIconFontInfo;
+	}
+
+	FSlateColor GetAvatarColor() const
+	{
+		TSharedPtr<FConcertSessionClientInfo> ClientInfoPin = SessionClientInfo.Pin();
+		if (ClientInfoPin.IsValid())
+		{
+			return ClientInfoPin->ClientInfo.AvatarColor;
+		}
+
+		return FSlateColor(FLinearColor(0.75, 0.75, 0.75)); // This is an arbitrary color.
+	}
+
+private:
+	TWeakPtr<IConcertSyncClient> SyncClient;
+	TWeakPtr<IConcertClientSession> ClientSession;
+	TWeakPtr<FConcertSessionClientInfo> SessionClientInfo;
+};
+
+void SActiveSessionOverviewTab::Construct(const FArguments& InArgs, TSharedPtr<IConcertSyncClient> InConcertSyncClient)
+{
+	WeakConcertSyncClient = InConcertSyncClient;
+	SessionHistoryController = MakeShared<FClientSessionHistoryController>(InConcertSyncClient.ToSharedRef());
+
+	if (InConcertSyncClient.IsValid())
+	{
+		IConcertClientRef ConcertClient = InConcertSyncClient->GetConcertClient();
+		ConcertClient->OnSessionStartup().AddSP(this, &SActiveSessionOverviewTab::HandleSessionStartup);
+		ConcertClient->OnSessionShutdown().AddSP(this, &SActiveSessionOverviewTab::HandleSessionShutdown);
+
+		if (TSharedPtr<IConcertClientSession> ClientSession = ConcertClient->GetCurrentSession())
+		{
+			WeakSessionPtr = ClientSession;
+			ClientInfo = MakeShared<FConcertSessionClientInfo>(FConcertSessionClientInfo{ClientSession->GetSessionClientEndpointId(), ClientSession->GetLocalClientInfo()});
+			ClientSession->OnSessionClientChanged().AddSP(this, &SActiveSessionOverviewTab::HandleSessionClientChanged);
+			InConcertSyncClient->GetTransactionBridge()->OnConflictResolutionForPendingSend().AddSP(this, &SActiveSessionOverviewTab::OnSendConflict);
+		}
+
+		TSharedPtr<IConcertClientWorkspace> Workspace = InConcertSyncClient->GetWorkspace();
+		Workspace->OnActivityAddedOrUpdated().AddSP(this, &SActiveSessionOverviewTab::ActivityUpdated);
+		Workspace->AddWorkspaceCanProcessPackagesDelegate(
+			TEXT("SActiveSession"), FCanProcessPendingPackages::CreateSP(this, &SActiveSessionOverviewTab::CanProcessPendingPackages));
+	}
+
+	ChildSlot
+	[
+		SNew(SVerticalBox)
+
+		// Clients + History
+		+SVerticalBox::Slot()
+		.FillHeight(1.0)
+		[
+			SNew(SSplitter)
+			.Orientation(Orient_Vertical)
+
+			// Clients.
+			+SSplitter::Slot()
+			.SizeRule(TAttribute<SSplitter::ESizeRule>(this, &SActiveSessionOverviewTab::GetClientAreaSizeRule))
+			.Value(0.3)
+			[
+				SNew(SBorder)
+				.BorderImage(FAppStyle::Get().GetBrush("ToolPanel.GroupBorder"))
+				.Padding(0.0f)
+				[
+					SAssignNew(ClientArea, SExpandableArea)
+					.BorderBackgroundColor(FLinearColor(0.6f, 0.6f, 0.6f, 1.0f))
+					.BorderImage_Lambda([this]() { return ConcertFrontendUtils::GetExpandableAreaBorderImage(*ClientArea); })
+					.BodyBorderImage(FAppStyle::Get().GetBrush("ToolPanel.GroupBorder"))
+					.BodyBorderBackgroundColor(FLinearColor::White)
+					.OnAreaExpansionChanged(this, &SActiveSessionOverviewTab::OnClientAreaExpansionChanged)
+					.Padding(0.0f)
+					.HeaderContent()
+					[
+						SNew(STextBlock)
+						.Text(LOCTEXT("SessionConnectedClients", "Clients"))
+						.Font(FAppStyle::Get().GetFontStyle("DetailsView.CategoryFontStyle"))
+						.ShadowOffset(FVector2D(1.0f, 1.0f))
+					]
+					.BodyContent()
+					[
+						SAssignNew(ClientsListView, SListView<TSharedPtr<FConcertSessionClientInfo>>)
+						.ItemHeight(20.0f)
+						.SelectionMode(ESelectionMode::Single)
+						.ListItemsSource(&Clients)
+						.OnGenerateRow(this, &SActiveSessionOverviewTab::HandleGenerateRow)
+						.HeaderRow
+						(
+							SNew(SHeaderRow)
+							+SHeaderRow::Column(ActiveSessionDetailsUI::DisplayNameColumnName)
+							.DefaultLabel(LOCTEXT("UserDisplayName", "Display Name"))
+							+SHeaderRow::Column(ActiveSessionDetailsUI::PresenceColumnName)
+							.DefaultLabel(LOCTEXT("UserPresence", "Presence"))
+							+SHeaderRow::Column(ActiveSessionDetailsUI::LevelColumnName)
+							.DefaultLabel(LOCTEXT("UserLevel", "Level"))
+						)
+					]
+				]
+			]
+
+			// History (activity feed).
+			+SSplitter::Slot()
+			.SizeRule(TAttribute<SSplitter::ESizeRule>(this, &SActiveSessionOverviewTab::GetHistoryAreaSizeRule))
+			.Value(0.7)
+			[
+				SNew(SBorder)
+				.BorderImage(FAppStyle::Get().GetBrush("ToolPanel.GroupBorder"))
+				.Padding(0.0f)
+				[
+					SAssignNew(HistoryArea, SExpandableArea)
+					.BorderBackgroundColor(FLinearColor(0.6f, 0.6f, 0.6f, 1.0f))
+					.BorderImage_Lambda([this]() { return ConcertFrontendUtils::GetExpandableAreaBorderImage(*HistoryArea); })
+					.BodyBorderImage(FAppStyle::Get().GetBrush("ToolPanel.GroupBorder"))
+					.BodyBorderBackgroundColor(FLinearColor::White)
+					.OnAreaExpansionChanged(this, &SActiveSessionOverviewTab::OnHistoryAreaExpansionChanged)
+					.Padding(0.0f)
+					.HeaderContent()
+					[
+						SNew(STextBlock)
+						.Text(LOCTEXT("SessionHistory", "History"))
+						.Font(FAppStyle::Get().GetFontStyle("DetailsView.CategoryFontStyle"))
+						.ShadowOffset(FVector2D(1.0f, 1.0f))
+					]
+					.BodyContent()
+					[
+						SessionHistoryController->GetSessionHistory()
+					]
+				]
+			]
+		]
+	];
+
+	// Create a timer to periodically poll the this client(s) info to detect if it changed its display name or avatar color because
+	// IConcertClientsession::OnSessionClientChanged() doesn't trigger when the 'local' client changes. This needs to be polled.
+	RegisterActiveTimer(1.0f, FWidgetActiveTimerDelegate::CreateSP(this, &SActiveSessionOverviewTab::HandleLocalClientInfoChangePollingTimer));
+
+	UpdateSessionClientListView();
+}
+
+TSharedRef<ITableRow> SActiveSessionOverviewTab::HandleGenerateRow(TSharedPtr<FConcertSessionClientInfo> InClientInfo, const TSharedRef<STableViewBase>& OwnerTable) const
+{
+	// Generate a row for the client corresponding to InClientInfo.
+	return SNew(SActiveSessionDetailsRow, WeakConcertSyncClient, WeakSessionPtr, InClientInfo, OwnerTable);
+}
+
+void SActiveSessionOverviewTab::HandleSessionStartup(TSharedRef<IConcertClientSession> InClientSession)
+{
+	WeakSessionPtr = InClientSession;
+	InClientSession->OnSessionClientChanged().AddSP(this, &SActiveSessionOverviewTab::HandleSessionClientChanged);
+	if (TSharedPtr<IConcertSyncClient> SyncClient = WeakConcertSyncClient.Pin())
+	{
+		SyncClient->GetTransactionBridge()->OnConflictResolutionForPendingSend().AddSP(this, &SActiveSessionOverviewTab::OnSendConflict);
+	}
+
+	ClientInfo = MakeShared<FConcertSessionClientInfo>(FConcertSessionClientInfo{InClientSession->GetSessionClientEndpointId(), InClientSession->GetLocalClientInfo()});
+
+	UpdateSessionClientListView();
+
+	if (SessionHistoryController.IsValid())
+	{
+		SessionHistoryController->ReloadActivities();
+	}
+}
+
+void SActiveSessionOverviewTab::HandleSessionShutdown(TSharedRef<IConcertClientSession> InClientSession)
+{
+	if (InClientSession == WeakSessionPtr)
+	{
+		WeakSessionPtr.Reset();
+		InClientSession->OnSessionClientChanged().RemoveAll(this);
+		if (TSharedPtr<IConcertSyncClient> SyncClient = WeakConcertSyncClient.Pin())
+		{
+			SyncClient->GetTransactionBridge()->OnConflictResolutionForPendingSend().RemoveAll(this);
+		}
+		Clients.Reset();
+
+		if (ClientsListView.IsValid())
+		{
+			ClientsListView->RequestListRefresh();
+		}
+
+		if (SessionHistoryController.IsValid())
+		{
+			SessionHistoryController->ReloadActivities();
+		}
+	}
+}
+
+void SActiveSessionOverviewTab::HandleSessionClientChanged(IConcertClientSession&, EConcertClientStatus ClientStatus, const FConcertSessionClientInfo& InClientInfo)
+{
+	// Update the view for a specific client.
+	UpdateSessionClientListView(&InClientInfo, ClientStatus);
+}
+
+EActiveTimerReturnType SActiveSessionOverviewTab::HandleLocalClientInfoChangePollingTimer(double InCurrentTime, float InDeltaTime)
+{
+// This switch is used to test if updating the session FConcertClientInfo dynamically correctly applies across the board. Don't
+// turn it on by default, it would overwrite runtime changes made to the info from some state/logic with the one from the settings.
+// For testing: Keep the MU browser tab open, join a session, open the MU settings and change one of the updatable field tracked below.
+#define MU_TEST_UPDATE_CLIENT_INFO_FROM_SETTINGS 0
+
+	TSharedPtr<IConcertClientSession> Session = WeakSessionPtr.Pin();
+	if (Session.IsValid() && ClientInfo.IsValid())
+	{
+#if MU_TEST_UPDATE_CLIENT_INFO_FROM_SETTINGS
+		if (TSharedPtr<IConcertSyncClient> SyncClientPin = WeakConcertSyncClient.Pin())
+		{
+			// User changed its settings.
+			const FConcertClientInfo& LatestClientInfo = Session->GetLocalClientInfo();
+			const UConcertClientConfig* ConcertClientConfig = SyncClientPin->GetConcertClient()->GetConfiguration();
+			if (ConcertClientConfig->ClientSettings.AvatarColor != LatestClientInfo.AvatarColor ||
+				ConcertClientConfig->ClientSettings.DisplayName != LatestClientInfo.DisplayName ||
+				ConcertClientConfig->ClientSettings.DesktopAvatarActorClass != LatestClientInfo.DesktopAvatarActorClass ||
+				ConcertClientConfig->ClientSettings.VRAvatarActorClass != LatestClientInfo.VRAvatarActorClass)
+			{
+				FConcertClientInfoUpdate ClientInfoUpdate;
+				ClientInfoUpdate.DisplayName = ConcertClientConfig->ClientSettings.DisplayName;
+				ClientInfoUpdate.AvatarColor = ConcertClientConfig->ClientSettings.AvatarColor;
+				ClientInfoUpdate.DesktopAvatarActorClass = ConcertClientConfig->ClientSettings.DesktopAvatarActorClass.ToString();
+				ClientInfoUpdate.VRAvatarActorClass = ConcertClientConfig->ClientSettings.VRAvatarActorClass.ToString();
+				Session->UpdateLocalClientInfo(ClientInfoUpdate);
+			}
+		}
+#endif
+
+		// Check if the local client info cached as class member is out dated with respect to the one held by the session changes. Just check the info displayed by this panel.
+		const FConcertClientInfo& LatestClientInfo = Session->GetLocalClientInfo();
+		if (LatestClientInfo.DisplayName != ClientInfo->ClientInfo.DisplayName || LatestClientInfo.AvatarColor != ClientInfo->ClientInfo.AvatarColor)
+		{
+			// Update the view for this client info.
+			ClientInfo->ClientInfo = LatestClientInfo;
+			UpdateSessionClientListView(ClientInfo.Get(), EConcertClientStatus::Updated);
+		}
+	}
+
+	return EActiveTimerReturnType::Continue;
+}
+
+void SActiveSessionOverviewTab::UpdateSessionClientListView(const FConcertSessionClientInfo* InClientInfo, EConcertClientStatus ClientStatus)
+{
+	// We expect the UI to be constructed in Construct() function, prior this function gets called.
+	check(ClientsListView.IsValid());
+
+	// Try pinning the session.
+	TSharedPtr<IConcertClientSession> Session = WeakSessionPtr.Pin();
+	if (Session.IsValid())
+	{
+		// Remember the element selected in the list view to reselect it later.
+		TArray<TSharedPtr<FConcertSessionClientInfo>> SelectedItems = ClientsListView->GetSelectedItems();
+		checkf(SelectedItems.Num() <= 1, TEXT("ActiveSession's client list view should not support multiple selection."));
+
+		// If this is about a specific client update. (i.e. responding to a IConcertClientSession::OnSessionClientChanged() event)
+		if (InClientInfo)
+		{
+			if (ClientStatus == EConcertClientStatus::Connected)
+			{
+				Clients.Emplace(MakeShared<FConcertSessionClientInfo>(*InClientInfo));
+			}
+			else
+			{
+				int32 Index = Clients.IndexOfByPredicate([InClientInfo](const TSharedPtr<FConcertSessionClientInfo>& Visited) { return InClientInfo->ClientEndpointId == Visited->ClientEndpointId; });
+				if (Index != INDEX_NONE)
+				{
+					if (ClientStatus == EConcertClientStatus::Disconnected)
+					{
+						Clients.RemoveAt(Index); // We want to preserve items relative order.
+					}
+					else if (ensure(ClientStatus == EConcertClientStatus::Updated)) // Ensure we handled all status in EConcertClientStatus.
+					{
+						*Clients[Index] = *InClientInfo; // Update the client info.
+					}
+				}
+			}
+		}
+		else // Sync everything.
+		{
+			// Convert the list of clients to a list of shared pointer to client (the list view model asks for shared pointers).
+			TArray<FConcertSessionClientInfo> OtherConnectedClients = Session->GetSessionClients();
+			TArray<TSharedPtr<FConcertSessionClientInfo>> UpdatedClientList;
+			UpdatedClientList.Reserve(OtherConnectedClients.Num() + 1); // +1 is to include this local client (see below).
+			Algo::Transform(OtherConnectedClients, UpdatedClientList, [](FConcertSessionClientInfo InClient)
+			{
+				return MakeShared<FConcertSessionClientInfo>(MoveTemp(InClient));
+			});
+
+			// Add this local client as it is not part of the list returned by GetSessionClients(). (The client connected from this process).
+			if (ClientInfo.IsValid())
+			{
+				UpdatedClientList.Emplace(ClientInfo);
+			}
+
+			// Merge the list used by the list view (the model) with the updated list, removing clients who left and adding the one who joined.
+			ConcertFrontendUtils::SyncArraysByPredicate(Clients, MoveTemp(UpdatedClientList), [](const TSharedPtr<FConcertSessionClientInfo>& ClientToFind)
+			{
+				return [ClientToFind](const TSharedPtr<FConcertSessionClientInfo>& PotentialClient)
+				{
+					return PotentialClient->ClientEndpointId == ClientToFind->ClientEndpointId;
+				};
+			});
+		}
+
+		// Sort the list by display name alphabetically.
+		Clients.StableSort([](const TSharedPtr<FConcertSessionClientInfo>& ClientOne, const TSharedPtr<FConcertSessionClientInfo>& ClientTwo)
+		{
+			return ClientOne->ClientInfo.DisplayName < ClientTwo->ClientInfo.DisplayName;
+		});
+
+		// If a client row was selected, select it back (if still available).
+		if (SelectedItems.Num() > 0)
+		{
+			ClientsListView->SetSelection(SelectedItems[0]);
+		}
+	}
+	else // The session appears to be invalid.
+	{
+		// Clear the list of clients.
+		Clients.Reset();
+	}
+
+	ClientsListView->RequestListRefresh();
+}
+
+void SActiveSessionOverviewTab::OnSendConflict(const FConcertConflictDescriptionBase& ConflictDescription)
+{
+	FAsyncTaskNotificationConfig Config;
+	Config.bKeepOpenOnFailure = true;
+	Config.TitleText = ConflictDescription.GetConflictTitle();
+	Config.ExpireDuration = 5.0f;
+
+	FAsyncTaskNotification Notification(Config);
+	Notification.SetHyperlink(FSimpleDelegate(), LOCTEXT("TransactionConflictDetails", "See Details..."));
+	Notification.SetComplete(Config.TitleText, ConflictDescription.GetConflictDetails(), false);
+}
+
+void SActiveSessionOverviewTab::SetSelectedClient(const FGuid& InClientEndpointId, ESelectInfo::Type SelectInfo)
+{
+	if (ClientsListView.IsValid())
+	{
+		TSharedPtr<FConcertSessionClientInfo> NewSelectedClient = FindAvailableClient(InClientEndpointId);
+
+		if (NewSelectedClient.IsValid())
+		{
+			ClientsListView->SetSelection(NewSelectedClient, SelectInfo);
+		}
+		else
+		{
+			ClientsListView->ClearSelection();
+		}
+	}
+}
+
+TSharedPtr<FConcertSessionClientInfo> SActiveSessionOverviewTab::FindAvailableClient(const FGuid& InClientEndpointId) const
+{
+	const TSharedPtr<FConcertSessionClientInfo>* FoundClientPtr = Clients.FindByPredicate([&InClientEndpointId](const TSharedPtr<FConcertSessionClientInfo>& PotentialClient)
+	{
+		return PotentialClient->ClientEndpointId == InClientEndpointId;
+	});
+
+	return FoundClientPtr ? *FoundClientPtr : nullptr;
+}
+
+bool SActiveSessionOverviewTab::CanProcessPendingPackages() const
+{
+	return bCanHotReload;
+}
+
+void SActiveSessionOverviewTab::Tick(const FGeometry& AllottedGeometry, const double InCurrentTime, const float InDeltaTime)
+{
+	if (bCanHotReload && CanReloadDialog)
+	{
+		CanReloadDialog = {};
+	}
+}
+
+namespace UE::MultiUserClient::Private
+{
+static bool IsMapAsset(FName PackageName)
+{
+	FString FullPath;
+	if (FPackageName::TryConvertLongPackageNameToFilename(PackageName.ToString(), FullPath, FPackageName::GetMapPackageExtension()))
+	{
+		if (FPlatformFileManager::Get().GetPlatformFile().FileExists(*FullPath))
+		{
+			return true;
+		}
+	}
+	return false;
+}
+
+static TSharedPtr<SMessageDialog> MakeReloadDialog(FName PackageName)
+{
+	static const FText DialogTitle = LOCTEXT("ActiveSessionPendingHotReloadTitle", "A level hot reload is pending.");
+
+	static const FText ReloadButton = LOCTEXT("ActiveSessionPendingHotReloadButton", "Reload");
+
+	const FText ReloadMessage =
+		FText::Format(
+			LOCTEXT("ActiveSessionPendingHotReload","A hot reload on {0} is pending. Please dismiss this notification to continue with reload. All transactions will be queued until this occurs."),
+			 FText::FromName(PackageName)
+			);
+
+	return SNew(SMessageDialog)
+		.Title(DialogTitle)
+		.Message(ReloadMessage)
+		.Buttons( { SCustomDialog::FButton(ReloadButton) } );
+}
+
+}
+void SActiveSessionOverviewTab::ActivityUpdated(const FConcertClientInfo& InClientInfo, const FConcertSyncActivity& InActivity, const FStructOnScope& ActivitySummary)
+{
+	if (InActivity.EventType == EConcertSyncActivityEventType::Package)
+	{
+		TSharedPtr<IConcertSyncClient> SyncClientPin = WeakConcertSyncClient.Pin();
+		if (!WeakSessionPtr.IsValid() || !SyncClientPin || CanReloadDialog)
+		{
+			return;
+		}
+
+		if (InActivity.EndpointId == WeakSessionPtr.Pin()->GetSessionClientEndpointId())
+		{
+			return;
+		}
+
+		check(ActivitySummary.GetStruct()->IsChildOf(FConcertSyncPackageActivitySummary::StaticStruct()));
+		const FConcertSyncPackageActivitySummary* ActivitySummaryPtr =
+			(FConcertSyncPackageActivitySummary*)ActivitySummary.GetStructMemory();
+
+		// User changed its settings.
+		const UConcertClientConfig* ConcertClientConfig = SyncClientPin->GetConcertClient()->GetConfiguration();
+		if (ConcertClientConfig->bShouldPromptForHotReloadOnLevel)
+		{
+			const bool bIsAMap = UE::MultiUserClient::Private::IsMapAsset(ActivitySummaryPtr->PackageName);
+			if (bIsAMap)
+			{
+				bCanHotReload = false;
+				CanReloadDialog = UE::MultiUserClient::Private::MakeReloadDialog(ActivitySummaryPtr->PackageName);
+				CanReloadDialog->GetOnWindowClosedEvent().AddLambda([this](const TSharedRef<SWindow>& ClosedWindow)
+					{
+						bCanHotReload = true;
+					});
+
+				CanReloadDialog->Show();
+			}
+		}
+	}
+}
+
+#undef LOCTEXT_NAMESPACE /* SActiveSession */
diff --git a/Engine/Plugins/Developer/Concert/ConcertApp/MultiUserClient/Source/MultiUserClient/Private/Widgets/ActiveSession/Overview/SActiveSessionOverviewTab.h b/Engine/Plugins/Developer/Concert/ConcertApp/MultiUserClient/Source/MultiUserClient/Private/Widgets/ActiveSession/Overview/SActiveSessionOverviewTab.h
new file mode 100644
--- /dev/null
+++ b/Engine/Plugins/Developer/Concert/ConcertApp/MultiUserClient/Source/MultiUserClient/Private/Widgets/ActiveSession/Overview/SActiveSessionOverviewTab.h
@@ -0,0 +1,124 @@
+// Copyright Epic Games, Inc. All Rights Reserved.
+
+#pragma once
+
+#include "ConcertMessages.h"
+#include "Widgets/ClientSessionHistoryController.h"
+#include "Widgets/DeclarativeSyntaxSupport.h"
+#include "Widgets/SCompoundWidget.h"
+#include "Widgets/Views/SListView.h"
+#include "Widgets/Input/SComboBox.h"
+
+class IConcertClientSession;
+class IConcertSyncClient;
+
+class FAsyncTaskNotification;
+struct FConcertConflictDescriptionBase;
+struct FConcertSessionClientInfo;
+
+class SCustomDialog;
+class SDockTab;
+class SExpandableArea;
+
+
+/**
+ * Displays the multi-users active session's clients and activities.
+ */
+class SActiveSessionOverviewTab : public SCompoundWidget
+{
+public:
+	
+	SLATE_BEGIN_ARGS(SActiveSessionOverviewTab) { }
+	SLATE_END_ARGS();
+
+	/**
+	* Constructs the active session tab.
+	*
+	* @param InArgs The Slate argument list.
+	*/
+	void Construct(const FArguments& InArgs, TSharedPtr<IConcertSyncClient> InConcertSyncClient);
+
+	virtual void Tick(const FGeometry& AllottedGeometry, const double InCurrentTime, const float InDeltaTime) override;
+
+private:
+	
+	/** Pointer on the client sync. */
+	TWeakPtr<IConcertSyncClient> WeakConcertSyncClient;
+
+	/** Holds a concert client session. */
+	TWeakPtr<IConcertClientSession> WeakSessionPtr;
+
+	/** List view for AvailableClients. */
+	TSharedPtr<SListView<TSharedPtr<FConcertSessionClientInfo>>> ClientsListView;
+
+	/** List of clients for the current session. */
+	TArray<TSharedPtr<FConcertSessionClientInfo>> Clients;
+
+	/** Information about the machine's client. */
+	TSharedPtr<FConcertSessionClientInfo> ClientInfo;
+
+	/** Holds a concert activity log. */
+	TSharedPtr<FClientSessionHistoryController> SessionHistoryController;
+
+	/** The 'Clients' expandable area. */
+	TSharedPtr<SExpandableArea> ClientArea;
+
+	/** The 'History' expandable area. */
+	TSharedPtr<SExpandableArea> HistoryArea;
+
+	/** Notification handler for hot reload. */
+	TSharedPtr<SCustomDialog> CanReloadDialog;
+
+	/** Flag to indicate if it is OK to hotreload the packages. */
+	bool bCanHotReload = true;
+
+	/** Keeps the status of 'Clients' area expansion. */
+	bool bClientAreaExpanded = true;
+
+	/** Keep the status of 'History' area expansion. */
+	bool bHistoryAreaExpanded = true;
+
+	/** Generate a new client row */
+	TSharedRef<ITableRow> HandleGenerateRow(TSharedPtr<FConcertSessionClientInfo> InClientInfo, const TSharedRef<STableViewBase>& OwnerTable) const;
+
+	/** Handle a session startup */
+	void HandleSessionStartup(TSharedRef<IConcertClientSession> InClientSession);
+
+	/** Handle a session shutdown */
+	void HandleSessionShutdown(TSharedRef<IConcertClientSession> InClientSession);
+
+	/** Handle a session client change */
+	void HandleSessionClientChanged(IConcertClientSession&, EConcertClientStatus ClientStatus, const FConcertSessionClientInfo& ClientInfo);
+
+	/** Polls the local client info and detect if it changed in order to update its representation in real time.*/
+	EActiveTimerReturnType HandleLocalClientInfoChangePollingTimer(double InCurrentTime, float InDeltaTime);
+
+	/** Update the list of clients while keeping the alphabetical sorting */
+	void UpdateSessionClientListView(const FConcertSessionClientInfo* Client = nullptr, EConcertClientStatus Status = EConcertClientStatus::Updated);
+
+	/** Set the selected client in the clients list view*/
+	void SetSelectedClient(const FGuid& InClientEndpointId, ESelectInfo::Type SelectInfo = ESelectInfo::Direct);
+
+	/** Find a client with its endpoint id */
+	TSharedPtr<FConcertSessionClientInfo> FindAvailableClient(const FGuid& InClientEndpointId) const;
+
+	/** Handles how much space the 'Clients' area uses with respect to its expansion state. */
+	SSplitter::ESizeRule GetClientAreaSizeRule() const { return bClientAreaExpanded ? SSplitter::ESizeRule::FractionOfParent : SSplitter::ESizeRule::SizeToContent; }
+	void OnClientAreaExpansionChanged(bool bExpanded) { bClientAreaExpanded = bExpanded; }
+
+	/** Handles how much space the 'History' area uses with respect to its expansion state. */
+	SSplitter::ESizeRule GetHistoryAreaSizeRule() const { return bHistoryAreaExpanded ? SSplitter::ESizeRule::FractionOfParent : SSplitter::ESizeRule::SizeToContent; }
+	void OnHistoryAreaExpansionChanged(bool bExpanded) { bHistoryAreaExpanded = bExpanded; }
+
+	/** Delegate handler when a conflict between inbound transaction and those stored pending outbound transactions. */
+	void OnSendConflict(const FConcertConflictDescriptionBase& ConflictMsg);
+
+	/** Delegate handler on inbound transactions that focus specifically on packages. */
+	void OnPackageChangeActivity();
+
+	/** Delegate handler to indicate if we can hot reload an inbound package. */
+	bool CanProcessPendingPackages() const;
+
+	/** Delegate handler invoked when an activity has been received by multi-user. */
+	void ActivityUpdated(const FConcertClientInfo& InClientInfo, const FConcertSyncActivity& InActivity, const FStructOnScope& /*unused*/);
+};
diff --git a/Engine/Plugins/Developer/Concert/ConcertApp/MultiUserClient/Source/MultiUserClient/Private/Widgets/ActiveSession/Replication/SReplicationConnectionControls.cpp b/Engine/Plugins/Developer/Concert/ConcertApp/MultiUserClient/Source/MultiUserClient/Private/Widgets/ActiveSession/Replication/SReplicationConnectionControls.cpp
new file mode 100644
--- /dev/null
+++ b/Engine/Plugins/Developer/Concert/ConcertApp/MultiUserClient/Source/MultiUserClient/Private/Widgets/ActiveSession/Replication/SReplicationConnectionControls.cpp
@@ -0,0 +1,188 @@
+// Copyright Epic Games, Inc. All Rights Reserved.
+
+#include "SReplicationConnectionControls.h"
+
+#include "IConcertSyncClient.h"
+#include "IConcertSyncClientModule.h"
+#include "Customization/ShowJoinSettingsCustomization.h"
+#include "Replication/IConcertClientReplicationManager.h"
+#include "Settings/MultiUserReplicationSettings.h"
+#include "Util/ReplicationUtils.h"
+
+#include "IDetailsView.h"
+#include "Modules/ModuleManager.h"
+#include "PropertyEditorModule.h"
+#include "SNegativeActionButton.h"
+#include "SPositiveActionButton.h"
+#include "Styling/AppStyle.h"
+#include "Widgets/Images/SImage.h"
+#include "Widgets/Input/SButton.h"
+#include "Widgets/Layout/SBox.h"
+#include "Widgets/SBoxPanel.h"
+
+#if WITH_EDITOR
+#include "IMultiUserReplicationEditorModule.h"
+#include "ISettingsModule.h"
+#endif
+
+#define LOCTEXT_NAMESPACE "SReplicationConnectionControls"
+
+namespace UE::MultiUserClient
+{
+	void SReplicationConnectionControls::Construct(const FArguments& InArgs)
+	{
+		TSharedPtr<SVerticalBox> Content;
+		ChildSlot
+		[
+			SAssignNew(Content, SVerticalBox)
+
+			// Buttons
+			+SVerticalBox::Slot()
+			.AutoHeight()
+			.Padding(4.f)
+			[
+				SNew(SHorizontalBox)
+
+				// Join
+				+SHorizontalBox::Slot()
+				.AutoWidth()
+				.Padding(2.f)
+				[
+					SNew(SPositiveActionButton)
+					.Text(LOCTEXT("Join.Label", "Join"))
+					.ToolTipText(this, &SReplicationConnectionControls::GetJoinButtonTooltip)
+					.OnClicked(this, &SReplicationConnectionControls::OnJoinButtonClicked)
+					.IsEnabled(this, &SReplicationConnectionControls::IsJoinButtonEnabled)
+					.Visibility(this, &SReplicationConnectionControls::GetJoinButtonVisibility)
+				]
+
+				// Leave
+				+SHorizontalBox::Slot()
+				.AutoWidth()
+				.Padding(2.f)
+				[
+					SNew(SNegativeActionButton)
+					.ActionButtonStyle(EActionButtonStyle::Warning)
+					.Text(LOCTEXT("Leave.Label", "Leave"))
+					.ToolTipText(LOCTEXT("Leave.Description", "Leaves the replication session. Will stop sending and receiving replication data. You will stay in the Concert session, i.e. continue to receive transaction events, etc."))
+					.OnClicked(this, &SReplicationConnectionControls::OnLeaveButtonClicked)
+					.Visibility(this, &SReplicationConnectionControls::GetLeaveButtonVisibility)
+				]
+
+				// The "Settings" icons.
+				+SHorizontalBox::Slot()
+				.AutoWidth()
+				.Padding(2.f)
+				.VAlign(VAlign_Fill)
+				[
+					SNew(SButton)
+					.ButtonStyle(FAppStyle::Get(), TEXT("SimpleButton"))
+					.Visibility(this, &SReplicationConnectionControls::GetCogWheelSettingsVisibility)
+					.OnClicked_Lambda([]()
+					{
+						using namespace UE::MultiUserReplicationEditor;
+						IMultiUserReplicationEditorModule& ReplicationEditorModule = IMultiUserReplicationEditorModule::Get();
+						const IMultiUserReplicationEditorModule::FSettingPath ReplicationSettings = ReplicationEditorModule.GetReplicationSettingsInfo();
+						FModuleManager::GetModulePtr<ISettingsModule>("Settings")->ShowViewer(ReplicationSettings.ContainerName, ReplicationSettings.CategoryName, ReplicationSettings.SectionName);
+						return FReply::Handled();
+					})
+					.VAlign(VAlign_Center)
+					.HAlign(HAlign_Fill)
+					[
+						SNew(SImage)
+						.Image(FAppStyle::GetBrush("Icons.Settings"))
+					]
+				]
+			]
+		];
+
+#if WITH_EDITOR
+		FPropertyEditorModule& PropertyEditorModule = FModuleManager::GetModuleChecked<FPropertyEditorModule>("PropertyEditor");
+		FDetailsViewArgs DetailsViewArgs;
+		DetailsViewArgs.bAllowSearch = false;
+		DetailsViewArgs.NameAreaSettings = FDetailsViewArgs::HideNameArea;
+		DetailsViewArgs.ColumnWidth = 0.5f;
+		
+		TSharedRef<IDetailsView> DetailsView = PropertyEditorModule.CreateDetailView(DetailsViewArgs);
+		DetailsView->RegisterInstancedCustomPropertyLayout(UMultiUserReplicationSettings::StaticClass(), FOnGetDetailCustomizationInstance::CreateStatic(&FShowJoinSettingsCustomization::Make));
+		DetailsView->SetObject(GetMutableDefault<UMultiUserReplicationSettings>());
+		DetailsView->SetVisibility(TAttribute<EVisibility>::CreateSP(this, &SReplicationConnectionControls::GetOverrideSettingsVisibility));
+		
+		Content->AddSlot()
+		.AutoHeight()
+		[
+			DetailsView
+		];
+#endif
+	}
+
+	UMultiUserReplicationClientProfileAsset* SReplicationConnectionControls::GetClientProfile() const
+	{
+		return GetMutableDefault<UMultiUserReplicationSettings>()->DefaultSessionSettings.DefaultProfile.LoadSynchronous();
+	}
+
+	FReply SReplicationConnectionControls::OnJoinButtonClicked()
+	{
+		TSharedPtr<IConcertSyncClient> Client = IConcertSyncClientModule::Get().GetClient(TEXT("MultiUser"));
+		IConcertClientReplicationManager* ReplicationManager = Client ? Client->GetReplicationManager() : nullptr;
+		if (!ensure(ReplicationManager))
+		{
+			return FReply::Handled();
+		}
+
+		UMultiUserReplicationClientProfileAsset* ClientProfile = GetClientProfile();
+		checkf(ClientProfile, TEXT("Validate IsJoinButtonEnabled implementation!"));
+		// All the UI callbacks for visibility, etc. will update automatically since they do polling.
+		ReplicationUtils::JoinSessionWithEditorNotifications(*ReplicationManager, ClientProfile->ToJoinArgs());
+		return FReply::Handled();
+	}
+
+	FText SReplicationConnectionControls::GetJoinButtonTooltip() const
+	{
+		FText ErrorReason = LOCTEXT("Join.Tooltip", "Attempts to join a replication session with the server to start sending and receive replication data.");
+		IsJoinButtonEnabled(&ErrorReason);
+		return ErrorReason;
+	}
+
+	bool SReplicationConnectionControls::IsJoinButtonEnabled(FText* ErrorReason) const
+	{
+		const bool bHasValidAsset = GetClientProfile() != nullptr;
+		if (!bHasValidAsset && ErrorReason)
+		{
+			*ErrorReason = LOCTEXT("IsJoinButtonEnabled.Error.InvalidProfile", "No client profile selected");
+			return false;
+		}
+		
+		TSharedPtr<IConcertSyncClient> Client = IConcertSyncClientModule::Get().GetClient(TEXT("MultiUser"));
+		IConcertClientReplicationManager* ReplicationManager = Client ? Client->GetReplicationManager() : nullptr;
+		const bool bCanMakeJoinRequest = ReplicationManager && ReplicationManager->CanJoin();
+		if (!bCanMakeJoinRequest && ErrorReason)
+		{
+			*ErrorReason = LOCTEXT("IsJoinButtonEnabled.Error.JoinRequestInProgress", "A join request is in progress.");
+			return false;
+		}
+		
+		return bHasValidAsset && bCanMakeJoinRequest;
+	}
+
+	EVisibility SReplicationConnectionControls::GetJoinButtonVisibility() const
+	{
+		TSharedPtr<IConcertSyncClient> Client = IConcertSyncClientModule::Get().GetClient(TEXT("MultiUser"));
+		IConcertClientReplicationManager* ReplicationManager = Client ? Client->GetReplicationManager() : nullptr;
+		const bool bShouldBeVisible = !ReplicationManager || !ReplicationManager->IsConnectedToReplicationSession();
+		return bShouldBeVisible ? EVisibility::Visible : EVisibility::Collapsed;
+	}
+
+	FReply SReplicationConnectionControls::OnLeaveButtonClicked()
+	{
+		TSharedPtr<IConcertSyncClient> Client = IConcertSyncClientModule::Get().GetClient(TEXT("MultiUser"));
+		IConcertClientReplicationManager* ReplicationManager = Client ? Client->GetReplicationManager() : nullptr;
+		if (ensure(ReplicationManager))
+		{
+			ReplicationUtils::LeaveSessionWithEditorNotifications(*ReplicationManager);
+		}
+		return FReply::Handled();
+	}
+}
+
+#undef LOCTEXT_NAMESPACE
diff --git a/Engine/Plugins/Developer/Concert/ConcertApp/MultiUserClient/Source/MultiUserClient/Private/Widgets/ActiveSession/Replication/SReplicationConnectionControls.h b/Engine/Plugins/Developer/Concert/ConcertApp/MultiUserClient/Source/MultiUserClient/Private/Widgets/ActiveSession/Replication/SReplicationConnectionControls.h
new file mode 100644
--- /dev/null
+++ b/Engine/Plugins/Developer/Concert/ConcertApp/MultiUserClient/Source/MultiUserClient/Private/Widgets/ActiveSession/Replication/SReplicationConnectionControls.h
@@ -0,0 +1,45 @@
+// Copyright Epic Games, Inc. All Rights Reserved.
+
+#pragma once
+
+#include "Widgets/DeclarativeSyntaxSupport.h"
+#include "Widgets/SCompoundWidget.h"
+
+class UMultiUserReplicationClientProfileAsset;
+enum class EJoinReplicationErrorCode : uint8;
+class UJoinReplicationSessionArgs;
+
+namespace UE::MultiUserClient
+{
+	class SReplicationConnectionControls : public SCompoundWidget
+	{
+	public:
+
+		SLATE_BEGIN_ARGS(SReplicationConnectionControls)
+		{}
+		SLATE_END_ARGS()
+
+		void Construct(const FArguments& InArgs);
+
+	private:
+
+		UMultiUserReplicationClientProfileAsset* GetClientProfile() const;
+
+		// Join Button
+		FReply OnJoinButtonClicked();
+		bool IsJoinButtonEnabled() const { return IsJoinButtonEnabled(nullptr); }
+		FText GetJoinButtonTooltip() const;
+		bool IsJoinButtonEnabled(FText* ErrorReason) const;
+		EVisibility GetJoinButtonVisibility() const;
+
+		// Leave button
+		FReply OnLeaveButtonClicked();
+		EVisibility GetLeaveButtonVisibility() const { return GetJoinButtonVisibility() == EVisibility::Visible ? EVisibility::Collapsed : EVisibility::Visible; }
+
+		// Settings
+		EVisibility GetCogWheelSettingsVisibility() const { return GetJoinButtonVisibility(); }
+		EVisibility GetOverrideSettingsVisibility() const { return GetJoinButtonVisibility(); }
+	};
+}
+
+
diff --git a/Engine/Plugins/Developer/Concert/ConcertApp/MultiUserClient/Source/MultiUserClient/Private/Widgets/ActiveSession/Replication/SReplicationControlsTab.cpp b/Engine/Plugins/Developer/Concert/ConcertApp/MultiUserClient/Source/MultiUserClient/Private/Widgets/ActiveSession/Replication/SReplicationControlsTab.cpp
new file mode 100644
--- /dev/null
+++ b/Engine/Plugins/Developer/Concert/ConcertApp/MultiUserClient/Source/MultiUserClient/Private/Widgets/ActiveSession/Replication/SReplicationControlsTab.cpp
@@ -0,0 +1,118 @@
+// Copyright Epic Games, Inc. All Rights Reserved.
+
+#include "SReplicationControlsTab.h"
+
+#include "ConcertFrontendUtils.h"
+#include "SReplicationConnectionControls.h"
+
+#include "Styling/AppStyle.h"
+#include "Widgets/Layout/SExpandableArea.h"
+#include "Widgets/Layout/SBorder.h"
+#include "Widgets/Layout/SSplitter.h"
+#include "Widgets/Text/STextBlock.h"
+
+#define LOCTEXT_NAMESPACE "SReplicationControlsTab"
+
+namespace UE::MultiUserClient
+{
+	namespace Private
+	{
+		static TSharedRef<SWidget> MakeExpandableArea(
+			TSharedPtr<SExpandableArea>& OutArea,
+			FOnBooleanValueChanged OnAreaExpansionChanged,
+			FText Label,
+			FText Tooltip,
+			TSharedRef<SWidget> BodyContent,
+			bool bInitiallyCollapsed
+			)
+		{
+			return SNew(SBorder)
+				.BorderImage(FAppStyle::Get().GetBrush("ToolPanel.GroupBorder"))
+				.Padding(0.0f)
+				[
+					SAssignNew(OutArea, SExpandableArea)
+					.BorderBackgroundColor(FLinearColor(0.6f, 0.6f, 0.6f, 1.0f))
+					.BorderImage_Lambda([&OutArea]() { return ConcertFrontendUtils::GetExpandableAreaBorderImage(*OutArea); })
+					.BodyBorderImage(FAppStyle::Get().GetBrush("ToolPanel.GroupBorder"))
+					.BodyBorderBackgroundColor(FLinearColor::White)
+					.OnAreaExpansionChanged(MoveTemp(OnAreaExpansionChanged))
+					.Padding(0.0f)
+					.InitiallyCollapsed(bInitiallyCollapsed)
+					.HeaderContent()
+					[
+						SNew(STextBlock)
+						.Text(MoveTemp(Label))
+						.ToolTipText(MoveTemp(Tooltip))
+						.Font(FAppStyle::Get().GetFontStyle("DetailsView.CategoryFontStyle"))
+						.ShadowOffset(FVector2D(1.0f, 1.0f))
+					]
+					.BodyContent()
+					[
+						MoveTemp(BodyContent)
+					]
+				];
+		}
+	}
+	
+	void SReplicationControlsTab::Construct(const FArguments& InArgs)
+	{
+		constexpr bool bCollpaseInitially_ConnectionArea	= false;
+		constexpr bool bCollpaseInitially_AttributesArea	= true;
+		constexpr bool bCollpaseInitially_StreamsArea		= true;
+		
+		ChildSlot
+		[
+			SNew(SSplitter)
+			.Orientation(Orient_Vertical)
+
+			// Connection area
+			+SSplitter::Slot()
+			.SizeRule(TAttribute<SSplitter::ESizeRule>(this, &SReplicationControlsTab::GetConnectionAreaSizeRule))
+			.Value(0.3)
+			[
+				Private::MakeExpandableArea(
+					ConnectionArea,
+					FOnBooleanValueChanged::CreateSP(this, &SReplicationControlsTab::OnConnectionAreaExpansionChanged),
+					LOCTEXT("ConnectionArea.Label", "Connection"),
+					LOCTEXT("ConnectionArea.Tooltip", "Control the replication connection, e.g. join and leave"),
+					SNew(SReplicationConnectionControls),
+					bCollpaseInitially_ConnectionArea
+					)
+			]
+			
+			// Attributes area
+			+SSplitter::Slot()
+			.SizeRule(TAttribute<SSplitter::ESizeRule>(this, &SReplicationControlsTab::GetClientAttributesAreaSizeRule))
+			.Value(0.3)
+			[
+				Private::MakeExpandableArea(
+					ClientAttributesArea,
+					FOnBooleanValueChanged::CreateSP(this, &SReplicationControlsTab::OnAClientsttributesAreaExpansionChanged),
+					LOCTEXT("AttributesArea.Label", "Client Receive Attributes"),
+					LOCTEXT("AttributesArea.Tooltip", "View and edit your client attributes, which determine the data received."),
+					SNew(STextBlock)
+						.Text(LOCTEXT("ToDo", "ToDo")),
+					bCollpaseInitially_AttributesArea
+					)
+			]
+			
+			// Streams area
+			+SSplitter::Slot()
+			.SizeRule(TAttribute<SSplitter::ESizeRule>(this, &SReplicationControlsTab::GetStreamsAreaSizeRule))
+			.Value(0.3)
+			[
+				Private::MakeExpandableArea(
+					StreamsArea,
+					FOnBooleanValueChanged::CreateSP(this, &SReplicationControlsTab::OnStreamsAreaExpansionChanged),
+					LOCTEXT("StreamsArea.Label", "Streams"),
+					LOCTEXT("StreamsArea.Tooltip", "View your replication streams as well as toggle sent objects and properties"),
+					SNew(STextBlock)
+						.Text(LOCTEXT("ToDo", "ToDo")),
+					bCollpaseInitially_StreamsArea
+					)
+			]
+		];
+	}
+}
+
+#undef LOCTEXT_NAMESPACE
\ No newline at end of file
diff --git a/Engine/Plugins/Developer/Concert/ConcertApp/MultiUserClient/Source/MultiUserClient/Private/Widgets/ActiveSession/Replication/SReplicationControlsTab.h b/Engine/Plugins/Developer/Concert/ConcertApp/MultiUserClient/Source/MultiUserClient/Private/Widgets/ActiveSession/Replication/SReplicationControlsTab.h
new file mode 100644
--- /dev/null
+++ b/Engine/Plugins/Developer/Concert/ConcertApp/MultiUserClient/Source/MultiUserClient/Private/Widgets/ActiveSession/Replication/SReplicationControlsTab.h
@@ -0,0 +1,49 @@
+// Copyright Epic Games, Inc. All Rights Reserved.
+
+#pragma once
+
+#include "Widgets/DeclarativeSyntaxSupport.h"
+#include "Widgets/SCompoundWidget.h"
+#include "Widgets/Layout/SSplitter.h"
+
+class SExpandableArea;
+
+namespace UE::MultiUserClient
+{
+	class SReplicationControlsTab : public SCompoundWidget
+	{
+	public:
+
+		SLATE_BEGIN_ARGS(SReplicationControlsTab)
+		{}
+		SLATE_END_ARGS()
+
+		void Construct(const FArguments& InArgs);
+
+	private:
+
+		/** In this area the user can control the replication connection, e.g. join and leave. */
+		TSharedPtr<SExpandableArea> ConnectionArea;
+		/** In this area the user can view and edit stream attributes, which will be replicated back to the server. */
+		TSharedPtr<SExpandableArea> ClientAttributesArea;
+		/** In this area the user can view the replication streams as well as toggle sent objects and its properties. */
+		TSharedPtr<SExpandableArea> StreamsArea;
+
+		bool bConnectionAreaExpanded = false;
+		bool bAttributesAreaExpanded = false;
+		bool bStreamsAreaExpanded = false;
+
+		/** Handles how much space the 'ConnectionArea' area uses with respect to its expansion state. */
+		SSplitter::ESizeRule GetConnectionAreaSizeRule() const { return SSplitter::ESizeRule::SizeToContent; }
+		void OnConnectionAreaExpansionChanged(bool bExpanded) { bConnectionAreaExpanded = bExpanded; }
+
+		/** Handles how much space the 'ClientAttributesArea' area uses with respect to its expansion state. */
+		SSplitter::ESizeRule GetClientAttributesAreaSizeRule() const { return bAttributesAreaExpanded ? SSplitter::ESizeRule::FractionOfParent : SSplitter::ESizeRule::SizeToContent; }
+		void OnAClientsttributesAreaExpansionChanged(bool bExpanded) { bAttributesAreaExpanded = bExpanded; }
+		
+		/** Handles how much space the 'StreamsArea' area uses with respect to its expansion state. */
+		SSplitter::ESizeRule GetStreamsAreaSizeRule() const { return bStreamsAreaExpanded ? SSplitter::ESizeRule::FractionOfParent : SSplitter::ESizeRule::SizeToContent; }
+		void OnStreamsAreaExpansionChanged(bool bExpanded) { bStreamsAreaExpanded = bExpanded; }
+	};
+}
+
diff --git a/Engine/Plugins/Developer/Concert/ConcertApp/MultiUserClient/Source/MultiUserClient/Private/Widgets/ActiveSession/Replication/SReplicationTabWithWarningOverlay.cpp b/Engine/Plugins/Developer/Concert/ConcertApp/MultiUserClient/Source/MultiUserClient/Private/Widgets/ActiveSession/Replication/SReplicationTabWithWarningOverlay.cpp
new file mode 100644
--- /dev/null
+++ b/Engine/Plugins/Developer/Concert/ConcertApp/MultiUserClient/Source/MultiUserClient/Private/Widgets/ActiveSession/Replication/SReplicationTabWithWarningOverlay.cpp
@@ -0,0 +1,78 @@
+// Copyright Epic Games, Inc. All Rights Reserved.
+
+#include "SReplicationTabWithWarningOverlay.h"
+
+#include "SReplicationControlsTab.h"
+
+#include "Widgets/Input/SButton.h"
+#include "Widgets/Layout/SBox.h"
+#include "Widgets/SOverlay.h"
+#include "Widgets/SBoxPanel.h"
+#include "Widgets/Text/STextBlock.h"
+
+#define LOCTEXT_NAMESPACE "SReplicationIsExperimentalWarning"
+
+namespace UE::MultiUserClient
+{
+	void SReplicationTabWithWarningOverlay::Construct(const FArguments& InArgs)
+	{
+		ChildSlot
+		[
+			SAssignNew(Overlay, SOverlay)
+
+			// The real content
+			+SOverlay::Slot()
+			[
+				SNew(SReplicationControlsTab)
+			]
+
+			//  
+			+SOverlay::Slot()
+			.ZOrder(1)
+			[
+				SNew(SBorder)
+				.BorderImage(FAppStyle::Get().GetBrush("ToolPanel.GroupBorder"))
+				.BorderBackgroundColor(FSlateColor(FLinearColor(0.6, 0.6, 0.6, 0.8f)))
+				.Padding(0)
+				[
+					SNew(SBox)
+					.Padding(0.f, 0.f, 0.f, 10.f)
+					.HAlign(HAlign_Center)
+					.VAlign(VAlign_Center)
+					.Clipping(EWidgetClipping::ClipToBounds)
+					[
+						SNew(SVerticalBox)
+
+						+SVerticalBox::Slot()
+						.AutoHeight()
+						.HAlign(HAlign_Center)
+						[
+							SNew(STextBlock)
+							.Text(LOCTEXT("EnableLoggingVisibility", "Warning: Replication is a work in progress feature."))
+							.Justification(ETextJustify::Center)
+						]
+
+						+SVerticalBox::Slot()
+						.Padding(0, 4, 0, 0)
+						.AutoHeight()
+						.HAlign(HAlign_Center)
+						[
+							SNew(SButton)
+							.OnClicked_Lambda([this]()
+							{
+								Overlay->RemoveSlot(1);
+								return FReply::Handled();
+							})
+							[
+								SNew(STextBlock)
+								.Text(LOCTEXT("EnableLoggingVisibility.Button.Text", "Ignore"))
+							]
+						]
+					]
+				]
+			]
+		];
+	}
+}
+
+#undef LOCTEXT_NAMESPACE
\ No newline at end of file
diff --git a/Engine/Plugins/Developer/Concert/ConcertApp/MultiUserClient/Source/MultiUserClient/Private/Widgets/ActiveSession/Replication/SReplicationTabWithWarningOverlay.h b/Engine/Plugins/Developer/Concert/ConcertApp/MultiUserClient/Source/MultiUserClient/Private/Widgets/ActiveSession/Replication/SReplicationTabWithWarningOverlay.h
new file mode 100644
--- /dev/null
+++ b/Engine/Plugins/Developer/Concert/ConcertApp/MultiUserClient/Source/MultiUserClient/Private/Widgets/ActiveSession/Replication/SReplicationTabWithWarningOverlay.h
@@ -0,0 +1,28 @@
+// Copyright Epic Games, Inc. All Rights Reserved.
+
+#pragma once
+
+#include "Widgets/DeclarativeSyntaxSupport.h"
+#include "Widgets/SCompoundWidget.h"
+
+class SOverlay;
+
+namespace UE::MultiUserClient
+{
+	/** Wraps SReplicationControlsTab with an overlay warning the user that this feature is still experimental */
+	class SReplicationTabWithWarningOverlay : public SCompoundWidget
+	{
+	public:
+
+		SLATE_BEGIN_ARGS(SReplicationTabWithWarningOverlay)
+		{}
+		SLATE_END_ARGS()
+
+		void Construct(const FArguments& InArgs);
+		
+	private:
+
+		TSharedPtr<SOverlay> Overlay;
+	};
+}
+
diff --git a/Engine/Plugins/Developer/Concert/ConcertApp/MultiUserClient/Source/MultiUserClient/Private/Widgets/ActiveSession/SActiveSessionRoot.cpp b/Engine/Plugins/Developer/Concert/ConcertApp/MultiUserClient/Source/MultiUserClient/Private/Widgets/ActiveSession/SActiveSessionRoot.cpp
new file mode 100644
--- /dev/null
+++ b/Engine/Plugins/Developer/Concert/ConcertApp/MultiUserClient/Source/MultiUserClient/Private/Widgets/ActiveSession/SActiveSessionRoot.cpp
@@ -0,0 +1,114 @@
+// Copyright Epic Games, Inc. All Rights Reserved.
+
+#include "SActiveSessionRoot.h"
+
+#include "Overview/SActiveSessionOverviewTab.h"
+#include "Replication/SReplicationTabWithWarningOverlay.h"
+#include "SActiveSessionToolbar.h"
+
+#include "Framework/Docking/TabManager.h"
+#include "Styling/AppStyle.h"
+#include "Styling/SlateTypes.h"
+#include "Widgets/Docking/SDockTab.h"
+#include "Widgets/SBoxPanel.h"
+#include "Widgets/SNullWidget.h"
+
+#define LOCTEXT_NAMESPACE "SActiveSessionRoot"
+
+namespace UE::MultiUserClient
+{
+	const FName SActiveSessionRoot::OverviewTabId(TEXT("OverviewTabId"));
+	const FName SActiveSessionRoot::ReplicationTabId(TEXT("ReplicationTabId"));
+	
+	void SActiveSessionRoot::Construct(
+		const FArguments& InArgs,
+		const TSharedRef<SDockTab>& ConstructUnderMajorTab,
+		TSharedPtr<IConcertSyncClient> InConcertSyncClient
+		)
+	{
+		ConcertSyncClient = InConcertSyncClient;
+		
+		TabManager = FGlobalTabmanager::Get()->NewTabManager(ConstructUnderMajorTab);
+		TSharedRef<FWorkspaceItem> AppMenuGroup = TabManager->AddLocalWorkspaceMenuCategory(LOCTEXT("ConcertActiveSession", "Active Session"));
+		TabManager->SetAllowWindowMenuBar(true);
+		RegisterTabSpawners(TabManager.ToSharedRef(), AppMenuGroup);
+
+		// Create our content
+		const TSharedRef<FTabManager::FLayout> Layout =
+			FTabManager::NewLayout("ConcertActiveSession_Layout_v3")
+			->AddArea
+			(
+				FTabManager::NewPrimaryArea()
+				->Split
+				(
+					FTabManager::NewStack()
+					->AddTab(OverviewTabId, ETabState::OpenedTab)
+					->AddTab(ReplicationTabId, ETabState::OpenedTab)
+					->SetForegroundTab(OverviewTabId)
+					->SetHideTabWell(false)
+					)
+			);
+		
+		ChildSlot
+		[
+			SNew(SVerticalBox)
+
+			// Toolbar
+			+SVerticalBox::Slot()
+			.AutoHeight()
+			[
+				SNew(SBorder)
+				.BorderImage(FAppStyle::Get().GetBrush("ToolPanel.GroupBorder"))
+				.Padding(0.0f)
+				[
+					SNew(SActiveSessionToolbar, InConcertSyncClient)
+				]
+			]
+
+			// Tabs
+			+SVerticalBox::Slot()
+			.FillHeight(1.f)
+			.HAlign(HAlign_Fill)
+			.VAlign(VAlign_Fill)
+			[
+				TabManager->RestoreFrom(Layout, nullptr).ToSharedRef()
+			]
+		];
+	}
+
+	void SActiveSessionRoot::RegisterTabSpawners(const TSharedRef<FTabManager>& InTabManager, const TSharedRef<FWorkspaceItem>& AppMenuGroup)
+	{
+		InTabManager->RegisterTabSpawner(OverviewTabId, FOnSpawnTab::CreateSP(this, &SActiveSessionRoot::SpawnTab_Overview))
+			.SetDisplayName(LOCTEXT("OverviewTab.DisplayName", "Overview"))
+			.SetGroup(AppMenuGroup);
+	
+		InTabManager->RegisterTabSpawner(ReplicationTabId, FOnSpawnTab::CreateSP(this, &SActiveSessionRoot::SpawnTab_ReplicationControls))
+			.SetDisplayName(LOCTEXT("ReplicationTab.DisplayName", "Replication"))
+			.SetGroup(AppMenuGroup);
+	}
+
+	TSharedRef<SDockTab> SActiveSessionRoot::SpawnTab_Overview(const FSpawnTabArgs& Args)
+	{
+		TSharedPtr<IConcertSyncClient> ConcertSyncClientPin = ConcertSyncClient.Pin();
+		return SNew(SDockTab)
+			.Label(LOCTEXT("OverviewTab.Label", "Overview"))
+			.ToolTipText(LOCTEXT("OverviewTab.Tooltip", "Displays active session clients and activity."))
+			[
+				ConcertSyncClientPin
+					? SNew(SActiveSessionOverviewTab, MoveTemp(ConcertSyncClientPin))
+					: SNullWidget::NullWidget
+			];
+	}
+
+	TSharedRef<SDockTab> SActiveSessionRoot::SpawnTab_ReplicationControls(const FSpawnTabArgs& Args)
+	{
+		return SNew(SDockTab)
+			.Label(LOCTEXT("ReplicationTab.Label", "Replication"))
+			.ToolTipText(LOCTEXT("ReplicationTab.Tooltip", "Manage real-time object replication"))
+			[
+				SNew(SReplicationTabWithWarningOverlay)
+			];
+	}
+}
+
+#undef LOCTEXT_NAMESPACE
\ No newline at end of file
diff --git a/Engine/Plugins/Developer/Concert/ConcertApp/MultiUserClient/Source/MultiUserClient/Private/Widgets/ActiveSession/SActiveSessionRoot.h b/Engine/Plugins/Developer/Concert/ConcertApp/MultiUserClient/Source/MultiUserClient/Private/Widgets/ActiveSession/SActiveSessionRoot.h
new file mode 100644
--- /dev/null
+++ b/Engine/Plugins/Developer/Concert/ConcertApp/MultiUserClient/Source/MultiUserClient/Private/Widgets/ActiveSession/SActiveSessionRoot.h
@@ -0,0 +1,46 @@
+// Copyright Epic Games, Inc. All Rights Reserved.
+
+#pragma once
+
+#include "Widgets/DeclarativeSyntaxSupport.h"
+#include "Widgets/SCompoundWidget.h"
+
+class FSpawnTabArgs;
+class FTabManager;
+class FWorkspaceItem;
+class IConcertSyncClient;
+class SDockTab;
+class SWindow;
+
+namespace UE::MultiUserClient
+{
+	/**
+	 * Displayed when the client is connected to an active session.
+	 * Manages the child content in tabs.
+	 */
+	class SActiveSessionRoot : public SCompoundWidget
+	{
+	public:
+		
+		static const FName OverviewTabId;
+		static const FName ReplicationTabId;
+
+		SLATE_BEGIN_ARGS(SActiveSessionRoot)
+		{}
+		SLATE_END_ARGS()
+
+		void Construct(const FArguments& InArgs, const TSharedRef<SDockTab>& ConstructUnderMajorTab, TSharedPtr<IConcertSyncClient> InConcertSyncClient);
+
+	private:
+
+		TWeakPtr<IConcertSyncClient> ConcertSyncClient;
+		/** Holds the child content: SActiveSessionOverviewTab and SReplicationControlsTab. */
+		TSharedPtr<FTabManager> TabManager;
+		
+		void RegisterTabSpawners(const TSharedRef<FTabManager>& InTabManager, const TSharedRef<FWorkspaceItem>& AppMenuGroup);
+		TSharedRef<SDockTab> SpawnTab_Overview(const FSpawnTabArgs& Args);
+		TSharedRef<SDockTab> SpawnTab_ReplicationControls(const FSpawnTabArgs& Args);
+	};
+}
+
+
diff --git a/Engine/Plugins/Developer/Concert/ConcertApp/MultiUserClient/Source/MultiUserClient/Private/Widgets/ActiveSession/SActiveSessionToolbar.cpp b/Engine/Plugins/Developer/Concert/ConcertApp/MultiUserClient/Source/MultiUserClient/Private/Widgets/ActiveSession/SActiveSessionToolbar.cpp
new file mode 100644
--- /dev/null
+++ b/Engine/Plugins/Developer/Concert/ConcertApp/MultiUserClient/Source/MultiUserClient/Private/Widgets/ActiveSession/SActiveSessionToolbar.cpp
@@ -0,0 +1,261 @@
+// Copyright Epic Games, Inc. All Rights Reserved.
+
+#include "SActiveSessionToolbar.h"
+
+#include "ConcertClientFrontendUtils.h"
+#include "IConcertSyncClient.h"
+#include "IMultiUserClientModule.h"
+
+#include "EditorFontGlyphs.h"
+#include "Styling/AppStyle.h"
+#include "Widgets/Images/SImage.h"
+#include "Widgets/Input/SButton.h"
+#include "Widgets/Input/SComboButton.h"
+#include "Widgets/Layout/SBorder.h"
+#include "Widgets/SBoxPanel.h"
+#include "Widgets/Text/STextBlock.h"
+
+#if WITH_EDITOR
+#include "ISettingsModule.h"
+#endif
+
+#define LOCTEXT_NAMESPACE "SActiveSessionToolbar"
+
+namespace UE::MultiUserClient
+{
+	namespace Private
+	{
+		static TArray< TSharedPtr<SActiveSessionToolbar::FSendReceiveComboItem>> ConstructSendReceiveComboList()
+		{
+			TArray<TSharedPtr<SActiveSessionToolbar::FSendReceiveComboItem>> ComboList;
+			ComboList.Add(MakeShared<SActiveSessionToolbar::FSendReceiveComboItem>(
+					LOCTEXT("DefaultSendReceiveState", "Default"),
+					LOCTEXT("DefaultSendReceiveStateTooltip", "Full send/receive mode for multi-user events."),
+					EConcertSendReceiveState::Default)
+				);
+			ComboList.Add(MakeShared<SActiveSessionToolbar::FSendReceiveComboItem>(
+					LOCTEXT("SendState", "Send only"),
+					LOCTEXT("SendStateTooltip", "Transactions received from clients will be suspended; however local transactions will be sent."),
+					EConcertSendReceiveState::SendOnly)
+				);
+			ComboList.Add(MakeShared<SActiveSessionToolbar::FSendReceiveComboItem>(
+					LOCTEXT("ReceiveState", "Receive only"),
+					LOCTEXT("ReceiveStateTooltip", "Local changes will be queued for transmission but not sent. Updates from clients will be received."),
+					EConcertSendReceiveState::ReceiveOnly)
+				);
+			return ComboList;
+		}
+	}
+	
+	void SActiveSessionToolbar::Construct(const FArguments& InArgs, TSharedPtr<IConcertSyncClient> InConcertSyncClient)
+	{
+		WeakSessionPtr = InConcertSyncClient->GetConcertClient()->GetCurrentSession();
+		SendReceiveComboList = Private::ConstructSendReceiveComboList();
+
+		TSharedPtr<SHorizontalBox> Toolbar;
+		ChildSlot
+		[
+			SAssignNew(Toolbar, SHorizontalBox)
+
+			// Status Icon
+			+SHorizontalBox::Slot()
+			.AutoWidth()
+			.VAlign(VAlign_Center)
+			.Padding(FMargin(2.0f, 1.0f, 0.0f, 1.0f))
+			[
+				SNew(STextBlock)
+				.Font(this, &SActiveSessionToolbar::GetConnectionIconFontInfo)
+				.ColorAndOpacity(this, &SActiveSessionToolbar::GetConnectionIconColor)
+				.Text(FEditorFontGlyphs::Circle)
+			]
+
+			// Status Message
+			+SHorizontalBox::Slot()
+			.VAlign(VAlign_Center)
+			.Padding(FMargin(4.0f, 1.0f))
+			[
+				SNew(SBorder)
+				.BorderImage(FAppStyle::Get().GetBrush("NoBorder"))
+				.ColorAndOpacity(FLinearColor(0.75f, 0.75f, 0.75f))
+				.Padding(FMargin(0.0f, 4.0f, 6.0f, 4.0f))
+				[
+					SNew(STextBlock)
+					.Font(FAppStyle::Get().GetFontStyle("BoldFont"))
+					.Text(this, &SActiveSessionToolbar::GetConnectionStatusText)
+				]
+			]
+			+SHorizontalBox::Slot()
+			.AutoWidth()
+			.VAlign(VAlign_Center)
+			.Padding(FMargin(1.0f))
+			[
+				SNew(SBox)
+				.MinDesiredWidth(130.0f)
+				[
+					SAssignNew(SendReceiveComboBox, SComboBox<TSharedPtr<FSendReceiveComboItem> >)
+					.OptionsSource(&SendReceiveComboList)
+					.InitiallySelectedItem(SendReceiveComboList[GetInitialSendReceiveComboIndex()])
+					.ContentPadding(FMargin(4.0f,1.0f))
+					.OnGenerateWidget(this, &SActiveSessionToolbar::GenerateSendReceiveComboItem)
+					.OnSelectionChanged(this, &SActiveSessionToolbar::HandleSendReceiveChanged)
+					[
+						SNew(STextBlock)
+						.Text(this, &SActiveSessionToolbar::GetRequestedSendReceiveComboText)
+					]
+				]
+			]
+			// The "Settings" icons.
+			+SHorizontalBox::Slot()
+			.VAlign(VAlign_Fill)
+			.AutoWidth()
+			[
+				SNew(SButton)
+				.ButtonStyle(FAppStyle::Get(), TEXT("SimpleButton"))
+				.OnClicked_Lambda([](){ FModuleManager::GetModulePtr<ISettingsModule>("Settings")->ShowViewer("Project", "Plugins", "Concert"); return FReply::Handled(); })
+				.VAlign(VAlign_Center)
+				.HAlign(HAlign_Fill)
+				[
+					SNew(SImage)
+					.Image(FAppStyle::GetBrush("Icons.Settings"))
+				]
+			]
+		];
+		
+		// Append the buttons to the status bar
+		{
+			TArray<FConcertActionDefinition> ButtonDefs;
+
+			// Leave Session
+			FConcertActionDefinition& LeaveSessionDef = ButtonDefs.AddDefaulted_GetRef();
+			LeaveSessionDef.Type = EConcertActionType::Normal;
+			LeaveSessionDef.IsVisible = MakeAttributeSP(this, &SActiveSessionToolbar::IsStatusBarLeaveSessionVisible);
+			LeaveSessionDef.Text = FEditorFontGlyphs::Sign_Out;
+			LeaveSessionDef.ToolTipText = LOCTEXT("LeaveCurrentSessionToolTip", "Leave the current session");
+			LeaveSessionDef.OnExecute.BindLambda([this]() { OnClickLeaveSession(); });
+			LeaveSessionDef.IconStyle = TEXT("Concert.LeaveSession");
+
+			ConcertClientFrontendUtils::AppendButtons(Toolbar.ToSharedRef(), ButtonDefs);
+		}
+	}
+
+	FSlateFontInfo SActiveSessionToolbar::GetConnectionIconFontInfo() const
+	{
+		FSlateFontInfo ConnectionIconFontInfo = FAppStyle::Get().GetFontStyle(ConcertClientFrontendUtils::ButtonIconSyle);
+		ConnectionIconFontInfo.OutlineSettings.OutlineSize = 1;
+		ConnectionIconFontInfo.OutlineSettings.OutlineColor = GetConnectionIconStyle().Pressed.TintColor.GetSpecifiedColor();
+
+		return ConnectionIconFontInfo;
+	}
+
+	FSlateColor SActiveSessionToolbar::GetConnectionIconColor() const
+	{
+		return GetConnectionIconStyle().Normal.TintColor;
+	}
+
+	const FButtonStyle& SActiveSessionToolbar::GetConnectionIconStyle() const
+	{
+		EConcertActionType ButtonStyle = EConcertActionType::Danger;
+
+		TSharedPtr<IConcertClientSession> ClientSession = WeakSessionPtr.Pin();
+		if (ClientSession.IsValid())
+		{
+			if (ClientSession->GetConnectionStatus() == EConcertConnectionStatus::Connected)
+			{
+				const bool bIsDefault = ClientSession->GetSendReceiveState() == EConcertSendReceiveState::Default;
+				if (bIsDefault)
+				{
+					ButtonStyle = EConcertActionType::Success;
+				}
+				else
+				{
+					ButtonStyle = EConcertActionType::Warning;
+				}
+			}
+		}
+
+		return FAppStyle::Get().GetWidgetStyle<FButtonStyle>(ConcertClientFrontendUtils::ButtonStyleNames[static_cast<int32>(ButtonStyle)]);
+	}
+
+	FText SActiveSessionToolbar::GetConnectionStatusText() const
+	{
+		FText StatusText = LOCTEXT("StatusDisconnected", "Disconnected");
+		TSharedPtr<IConcertClientSession> ClientSessionPtr = WeakSessionPtr.Pin();
+		if (ClientSessionPtr.IsValid() && ClientSessionPtr->GetConnectionStatus() == EConcertConnectionStatus::Connected)
+		{
+			const FText SessionDisplayName = FText::FromString(ClientSessionPtr->GetSessionInfo().SessionName);
+
+			switch (ClientSessionPtr->GetSendReceiveState())
+			{
+			case EConcertSendReceiveState::SendOnly:
+				StatusText = FText::Format(LOCTEXT("StatusSendOnlyFmt", "Send Only: {0}"), SessionDisplayName);
+				break;
+			case EConcertSendReceiveState::ReceiveOnly:
+				StatusText = FText::Format(LOCTEXT("StatusReceiveOnlyFmt", "Receive Only: {0}"), SessionDisplayName);
+				break;
+			case EConcertSendReceiveState::Default:
+				StatusText = FText::Format(LOCTEXT("StatusConnectedFmt", "Connected: {0}"), SessionDisplayName);
+				break;
+			};
+		}
+
+		return StatusText;
+	}
+
+	int32 SActiveSessionToolbar::GetInitialSendReceiveComboIndex() const
+	{
+		if (!WeakSessionPtr.IsValid())
+		{
+			return 0;
+		}
+
+		TSharedPtr<IConcertClientSession> Session = WeakSessionPtr.Pin();
+		switch(Session->GetSendReceiveState())
+		{
+		case EConcertSendReceiveState::Default:
+			return 0;
+		case EConcertSendReceiveState::SendOnly:
+			return 1;
+		case EConcertSendReceiveState::ReceiveOnly:
+			return 2;
+		}
+
+		return 0;
+	}
+
+	TSharedRef<SWidget> SActiveSessionToolbar::GenerateSendReceiveComboItem(TSharedPtr<FSendReceiveComboItem> InItem) const
+	{
+		return SNew(STextBlock)
+			.Text(InItem->Name)
+			.ToolTipText(InItem->ToolTip);
+	}
+
+	void SActiveSessionToolbar::HandleSendReceiveChanged(TSharedPtr<FSendReceiveComboItem> Item, ESelectInfo::Type SelectInfo) const
+	{
+		check(Item.IsValid());
+		if (TSharedPtr<IConcertClientSession> Session = WeakSessionPtr.Pin())
+		{
+			Session->SetSendReceiveState(Item->State);
+		}
+	}
+
+	FText SActiveSessionToolbar::GetRequestedSendReceiveComboText() const
+	{
+		return SendReceiveComboBox->GetSelectedItem().IsValid()
+			? SendReceiveComboBox->GetSelectedItem()->Name
+			: LOCTEXT("ActiveSessionDefaultSendReceive", "Default");
+	}
+
+	bool SActiveSessionToolbar::IsStatusBarLeaveSessionVisible() const
+	{
+		const TSharedPtr<IConcertClientSession> ClientSession = WeakSessionPtr.Pin();
+		return ClientSession && ClientSession->GetConnectionStatus() == EConcertConnectionStatus::Connected;
+	}
+
+	FReply SActiveSessionToolbar::OnClickLeaveSession() const
+	{
+		IMultiUserClientModule::Get().DisconnectSession();
+		return FReply::Handled();
+	}
+}
+
+#undef LOCTEXT_NAMESPACE
\ No newline at end of file
diff --git a/Engine/Plugins/Developer/Concert/ConcertApp/MultiUserClient/Source/MultiUserClient/Private/Widgets/ActiveSession/SActiveSessionToolbar.h b/Engine/Plugins/Developer/Concert/ConcertApp/MultiUserClient/Source/MultiUserClient/Private/Widgets/ActiveSession/SActiveSessionToolbar.h
new file mode 100644
--- /dev/null
+++ b/Engine/Plugins/Developer/Concert/ConcertApp/MultiUserClient/Source/MultiUserClient/Private/Widgets/ActiveSession/SActiveSessionToolbar.h
@@ -0,0 +1,72 @@
+// Copyright Epic Games, Inc. All Rights Reserved.
+
+#pragma once
+
+#include "ConcertMessages.h"
+#include "Widgets/DeclarativeSyntaxSupport.h"
+#include "Widgets/Input/SComboBox.h"
+#include "Widgets/SCompoundWidget.h"
+
+class IConcertClientSession;
+class IConcertSyncClient;
+struct FButtonStyle;
+
+namespace UE::MultiUserClient
+{
+	/**
+	 * Displayed on the top of the window when in a session.
+	 *
+	 * Enables the client to leave the session, open the same level as another client or teleport to another
+	 * client presence.
+	 */
+	class SActiveSessionToolbar : public SCompoundWidget
+	{
+	public:
+		
+		/** Struct to store the current send / receive state. */
+		struct FSendReceiveComboItem
+		{
+			FSendReceiveComboItem(FText InName, FText InToolTip, EConcertSendReceiveState InState) :
+				Name(MoveTemp(InName)), ToolTip(MoveTemp(InToolTip)), State(InState) {};
+
+			FText Name;
+			FText ToolTip;
+			EConcertSendReceiveState State;
+		};
+
+		SLATE_BEGIN_ARGS(SActiveSessionToolbar)
+		{}
+		SLATE_END_ARGS()
+
+		void Construct(const FArguments& InArgs, TSharedPtr<IConcertSyncClient> InConcertSyncClient);
+
+	private:
+		
+		/** Holds a concert client session. */
+		TWeakPtr<IConcertClientSession> WeakSessionPtr;
+		
+		/** Sets whether this client sends or receives. */
+		TSharedPtr<SComboBox<TSharedPtr<FSendReceiveComboItem>>> SendReceiveComboBox;
+		
+		/** Available states for the SendReceiveComboBox */
+		TArray<TSharedPtr<FSendReceiveComboItem>> SendReceiveComboList;
+
+		// Status icon and text
+		FSlateFontInfo GetConnectionIconFontInfo() const;
+		FSlateColor GetConnectionIconColor() const;
+		const FButtonStyle& GetConnectionIconStyle() const;
+		FText GetConnectionStatusText() const;
+
+		// Default, Send Only & Receive Only combo box
+		int32 GetInitialSendReceiveComboIndex() const;
+		TSharedRef<SWidget> GenerateSendReceiveComboItem(TSharedPtr<FSendReceiveComboItem> InItem) const;
+		void HandleSendReceiveChanged(TSharedPtr<FSendReceiveComboItem> Item, ESelectInfo::Type SelectInfo) const;
+		FText GetRequestedSendReceiveComboText() const;
+		
+		// Handling for leave session button 
+		bool IsStatusBarLeaveSessionVisible() const;
+		FReply OnClickLeaveSession() const;
+
+	};
+}
+
diff --git a/Engine/Plugins/Developer/Concert/ConcertApp/MultiUserClient/Source/MultiUserClient/Private/Widgets/SActiveSession.cpp b/Engine/Plugins/Developer/Concert/ConcertApp/MultiUserClient/Source/MultiUserClient/Private/Widgets/SActiveSession.cpp
deleted file mode 100644
diff --git a/Engine/Plugins/Developer/Concert/ConcertApp/MultiUserClient/Source/MultiUserClient/Private/Widgets/SActiveSession.h b/Engine/Plugins/Developer/Concert/ConcertApp/MultiUserClient/Source/MultiUserClient/Private/Widgets/SActiveSession.h
deleted file mode 100644
diff --git a/Engine/Plugins/Developer/Concert/ConcertApp/MultiUserClient/Source/MultiUserClient/Private/Widgets/SConcertBrowser.cpp b/Engine/Plugins/Developer/Concert/ConcertApp/MultiUserClient/Source/MultiUserClient/Private/Widgets/SConcertBrowser.cpp
--- a/Engine/Plugins/Developer/Concert/ConcertApp/MultiUserClient/Source/MultiUserClient/Private/Widgets/SConcertBrowser.cpp
+++ b/Engine/Plugins/Developer/Concert/ConcertApp/MultiUserClient/Source/MultiUserClient/Private/Widgets/SConcertBrowser.cpp
@@ -5,7 +5,8 @@
 #include "IConcertClient.h"
 #include "IConcertSyncClient.h"
 #include "MultiUserClientUtils.h"
-#include "SActiveSession.h"
+#include "ActiveSession/SActiveSessionRoot.h"
+#include "Widgets/ActiveSession/Overview/SActiveSessionOverviewTab.h"
 #include "Widgets/Disconnected/SConcertClientSessionBrowser.h"
 #include "Widgets/Disconnected/SConcertNoAvailability.h"
 
@@ -17,7 +18,7 @@
 
 #define LOCTEXT_NAMESPACE "SConcertBrowser"
 
-void SConcertBrowser::Construct(const FArguments& InArgs, const TSharedRef<SDockTab>& ConstructUnderMajorTab, const TSharedPtr<SWindow>& ConstructUnderWindow, TWeakPtr<IConcertSyncClient> InSyncClient)
+void SConcertBrowser::Construct(const FArguments& InArgs, TSharedRef<SDockTab> InConstructUnderMajorTab, TWeakPtr<IConcertSyncClient> InSyncClient)
 {
 	if (!MultiUserClientUtils::HasServerCompatibleCommunicationPluginEnabled())
 	{
@@ -32,6 +33,7 @@ void SConcertBrowser::Construct(const FArguments& InArgs, const TSharedRef<SDock
 	}
 
 	WeakConcertSyncClient = InSyncClient;
+	ConstructedUnderMajorTab = MoveTemp(InConstructUnderMajorTab);
 	if (TSharedPtr<IConcertSyncClient> ConcertSyncClient = WeakConcertSyncClient.Pin())
 	{
 		SearchedText = MakeShared<FText>(); // Will keep in memory the session browser search text between join/leave UI transitions.
@@ -57,7 +59,7 @@ void SConcertBrowser::AttachChildWidget(EConcertConnectionStatus ConnectionStatu
 	{
 		if (ConnectionStatus == EConcertConnectionStatus::Connected)
 		{
-			ChildSlot.AttachWidget(SNew(SActiveSession, ConcertSyncClient));
+			ChildSlot.AttachWidget(SNew(UE::MultiUserClient::SActiveSessionRoot, ConstructedUnderMajorTab.ToSharedRef(), ConcertSyncClient));
 		}
 		else if (ConnectionStatus == EConcertConnectionStatus::Disconnected)
 		{
diff --git a/Engine/Plugins/Developer/Concert/ConcertApp/MultiUserClient/Source/MultiUserClient/Private/Widgets/SConcertBrowser.h b/Engine/Plugins/Developer/Concert/ConcertApp/MultiUserClient/Source/MultiUserClient/Private/Widgets/SConcertBrowser.h
--- a/Engine/Plugins/Developer/Concert/ConcertApp/MultiUserClient/Source/MultiUserClient/Private/Widgets/SConcertBrowser.h
+++ b/Engine/Plugins/Developer/Concert/ConcertApp/MultiUserClient/Source/MultiUserClient/Private/Widgets/SConcertBrowser.h
@@ -28,22 +28,26 @@ public:
 	* Constructs the Browser.
 	*
 	* @param InArgs The Slate argument list.
-	* @param ConstructUnderMajorTab The major tab which will contain the session front-end.
-	* @param ConstructUnderWindow The window in which this widget is being constructed.
+	* @param InConstructUnderMajorTab The major tab which will contain the session front-end.
+	* @param InConstructUnderWindow The window in which this widget is being constructed.
 	* @param InSyncClient The sync client.
 	*/
-	void Construct(const FArguments& InArgs, const TSharedRef<SDockTab>& ConstructUnderMajorTab, const TSharedPtr<SWindow>& ConstructUnderWindow, TWeakPtr<IConcertSyncClient> InSyncClient);
+	void Construct(const FArguments& InArgs, TSharedRef<SDockTab> InConstructUnderMajorTab, TWeakPtr<IConcertSyncClient> InSyncClient);
 
 private:
-	/** Invoked when the session connection state is changed. */
-	void HandleSessionConnectionChanged(IConcertClientSession& InSession, EConcertConnectionStatus ConnectionStatus);
-
-	/** Attaches the child widgets according to the connection status. */
-	void AttachChildWidget(EConcertConnectionStatus ConnectionStatus);
 
 	/** Keeps the sync client interface. */
 	TWeakPtr<IConcertSyncClient> WeakConcertSyncClient;
 
+	/** Kept so it can be passed on to SActiveSessionRoot. */
+	TSharedPtr<SDockTab> ConstructedUnderMajorTab;
+
 	/** Keeps the session browser searched text in memory to reapply it when a user leaves a session and goes back to the session browser. */
 	TSharedPtr<FText> SearchedText;
+	
+	/** Invoked when the session connection state is changed. */
+	void HandleSessionConnectionChanged(IConcertClientSession& InSession, EConcertConnectionStatus ConnectionStatus);
+
+	/** Attaches the child widgets according to the connection status. */
+	void AttachChildWidget(EConcertConnectionStatus ConnectionStatus);
 };
diff --git a/Engine/Plugins/Developer/Concert/ConcertApp/MultiUserClient/Source/MultiUserReplication/MultiUserReplication.Build.cs b/Engine/Plugins/Developer/Concert/ConcertApp/MultiUserClient/Source/MultiUserReplication/MultiUserReplication.Build.cs
--- a/Engine/Plugins/Developer/Concert/ConcertApp/MultiUserClient/Source/MultiUserReplication/MultiUserReplication.Build.cs
+++ b/Engine/Plugins/Developer/Concert/ConcertApp/MultiUserClient/Source/MultiUserReplication/MultiUserReplication.Build.cs
@@ -20,7 +20,8 @@ namespace UnrealBuildTool.Rules
 				new string[]
 				{
 					"Concert",
-					"ConcertClient"
+					"ConcertClient",
+					"ConcertSyncClient"
 				}
 			);
 		}
diff --git a/Engine/Plugins/Developer/Concert/ConcertApp/MultiUserClient/Source/MultiUserReplication/Private/MultiUserReplicationClientProfileAsset.cpp b/Engine/Plugins/Developer/Concert/ConcertApp/MultiUserClient/Source/MultiUserReplication/Private/MultiUserReplicationClientProfileAsset.cpp
new file mode 100644
--- /dev/null
+++ b/Engine/Plugins/Developer/Concert/ConcertApp/MultiUserClient/Source/MultiUserReplication/Private/MultiUserReplicationClientProfileAsset.cpp
@@ -0,0 +1,23 @@
+// Copyright Epic Games, Inc. All Rights Reserved.
+
+#include "MultiUserReplicationClientProfileAsset.h"
+
+#include "Replication/IConcertClientReplicationManager.h"
+
+#include "Algo/Transform.h"
+
+UE::ConcertSyncClient::Replication::FJoinReplicatedSessionArgs UMultiUserReplicationClientProfileAsset::ToJoinArgs() const
+{
+	UE::ConcertSyncClient::Replication::FJoinReplicatedSessionArgs Result;
+	Result.ClientInfo = ClientDescription;
+	Algo::TransformIf(Streams, Result.Streams,
+		[](const UMultiUserReplicationStreamAsset* StreamAsset)
+		{
+			return StreamAsset != nullptr;
+		},
+		[](const UMultiUserReplicationStreamAsset* StreamAsset)
+		{
+			return StreamAsset->GenerateDescription();
+		});
+	return Result;
+}
diff --git a/Engine/Plugins/Developer/Concert/ConcertApp/MultiUserClient/Source/MultiUserReplication/Private/MultiUserReplicationStreamAsset.cpp b/Engine/Plugins/Developer/Concert/ConcertApp/MultiUserClient/Source/MultiUserReplication/Private/MultiUserReplicationStreamAsset.cpp
--- a/Engine/Plugins/Developer/Concert/ConcertApp/MultiUserClient/Source/MultiUserReplication/Private/MultiUserReplicationStreamAsset.cpp
+++ b/Engine/Plugins/Developer/Concert/ConcertApp/MultiUserClient/Source/MultiUserReplication/Private/MultiUserReplicationStreamAsset.cpp
@@ -6,10 +6,14 @@ UMultiUserReplicationStreamAsset::UMultiUserReplicationStreamAsset()
 {
 	ReplicationList = CreateDefaultSubobject<UMultiUserPropertyReplicationSelection>(TEXT("ReplicationList"));
 	ReplicationList->SetFlags(RF_Transactional);
+
+	// To avoid issues with delta serialization the CDO is always 0
+	StreamId = HasAnyFlags(RF_ClassDefaultObject)
+		? FGuid{}
+		: FGuid::NewGuid();
 }
 
 FReplicationStreamDescription UMultiUserReplicationStreamAsset::GenerateDescription() const
 {
-	// TODO:
-	return {};
+	return { StreamId, ReplicationList->ReplicationMap };
 }
diff --git a/Engine/Plugins/Developer/Concert/ConcertApp/MultiUserClient/Source/MultiUserReplication/Public/MultiUserReplicationClientProfileAsset.h b/Engine/Plugins/Developer/Concert/ConcertApp/MultiUserClient/Source/MultiUserReplication/Public/MultiUserReplicationClientProfileAsset.h
--- a/Engine/Plugins/Developer/Concert/ConcertApp/MultiUserClient/Source/MultiUserReplication/Public/MultiUserReplicationClientProfileAsset.h
+++ b/Engine/Plugins/Developer/Concert/ConcertApp/MultiUserClient/Source/MultiUserReplication/Public/MultiUserReplicationClientProfileAsset.h
@@ -7,6 +7,13 @@
 #include "MultiUserReplicationStreamAsset.h"
 #include "MultiUserReplicationClientProfileAsset.generated.h"
 
+namespace UE::ConcertSyncClient::Replication
+{
+	struct FJoinReplicatedSessionArgs;
+}
+
+enum class EJoinReplicationErrorCode : uint8;
+
 /** Asset for users to describe a client to the server. */
 UCLASS()
 class MULTIUSERREPLICATION_API UMultiUserReplicationClientProfileAsset : public UObject
@@ -21,4 +28,7 @@ public:
 	/** The streams this client will send to the server. It takes authority over the listed properties. */
 	UPROPERTY(EditAnywhere, Category = "Replication")
 	TArray<TObjectPtr<UMultiUserReplicationStreamAsset>> Streams;
-};
\ No newline at end of file
+
+	/** Converts this asset's data to the args expected by IConcertClientReplicationManager for joining. */
+	UE::ConcertSyncClient::Replication::FJoinReplicatedSessionArgs ToJoinArgs() const;
+};
diff --git a/Engine/Plugins/Developer/Concert/ConcertApp/MultiUserClient/Source/MultiUserReplication/Public/MultiUserReplicationStreamAsset.h b/Engine/Plugins/Developer/Concert/ConcertApp/MultiUserClient/Source/MultiUserReplication/Public/MultiUserReplicationStreamAsset.h
--- a/Engine/Plugins/Developer/Concert/ConcertApp/MultiUserClient/Source/MultiUserReplication/Public/MultiUserReplicationStreamAsset.h
+++ b/Engine/Plugins/Developer/Concert/ConcertApp/MultiUserClient/Source/MultiUserReplication/Public/MultiUserReplicationStreamAsset.h
@@ -15,6 +15,10 @@ class MULTIUSERREPLICATION_API UMultiUserReplicationStreamAsset : public UObject
 	GENERATED_BODY()
 public:
 
+	/** Unique ID for this stream. CDO has value 0 to not cause issues with delta serialization. */
+	UPROPERTY()
+	FGuid StreamId;
+	
 	UPROPERTY(Instanced)
 	TObjectPtr<UMultiUserPropertyReplicationSelection> ReplicationList;
 
diff --git a/Engine/Plugins/Developer/Concert/ConcertApp/MultiUserClient/Source/MultiUserReplicationEditor/MultiUserReplicationEditor.Build.cs b/Engine/Plugins/Developer/Concert/ConcertApp/MultiUserClient/Source/MultiUserReplicationEditor/MultiUserReplicationEditor.Build.cs
--- a/Engine/Plugins/Developer/Concert/ConcertApp/MultiUserClient/Source/MultiUserReplicationEditor/MultiUserReplicationEditor.Build.cs
+++ b/Engine/Plugins/Developer/Concert/ConcertApp/MultiUserClient/Source/MultiUserReplicationEditor/MultiUserReplicationEditor.Build.cs
@@ -37,8 +37,7 @@ namespace UnrealBuildTool.Rules
 					// Concert
 					"Concert",
 					"ConcertClient",
-					"ConcertSharedSlate",
-					"MultiUserClient"
+					"ConcertSharedSlate"
 				}
 			);
 		}
diff --git a/Engine/Plugins/Developer/Concert/ConcertApp/MultiUserClient/Source/MultiUserReplicationEditor/Private/MultiUserReplicationEditorModule.cpp b/Engine/Plugins/Developer/Concert/ConcertApp/MultiUserClient/Source/MultiUserReplicationEditor/Private/MultiUserReplicationEditorModule.cpp
--- a/Engine/Plugins/Developer/Concert/ConcertApp/MultiUserClient/Source/MultiUserReplicationEditor/Private/MultiUserReplicationEditorModule.cpp
+++ b/Engine/Plugins/Developer/Concert/ConcertApp/MultiUserClient/Source/MultiUserReplicationEditor/Private/MultiUserReplicationEditorModule.cpp
@@ -3,7 +3,6 @@
 #include "MultiUserReplicationEditorModule.h"
 
 #include "ConcertSyncSessionFlags.h"
-#include "IMultiUserClientModule.h"
 
 #include "ISettingsModule.h"
 #include "MultiUserReplicationEditorStyle.h"
@@ -30,12 +29,18 @@ namespace UE::MultiUserReplicationEditor
 		UnregisterSettings();
 	}
 
+	IMultiUserReplicationEditorModule::FSettingPath FMultiUserReplicationEditorModule::GetReplicationSettingsInfo() const
+	{
+		return { "Project", "Plugins", "Multi-User Replication" };
+	}
+
 	void FMultiUserReplicationEditorModule::RegisterSettings()
 	{
 		ISettingsModule* SettingsModule = FModuleManager::GetModulePtr<ISettingsModule>("Settings");
 		if (ensure(SettingsModule))
 		{
-			SettingsModule->RegisterSettings("Project", "Plugins", "Multi-User Replication",
+			const FSettingPath Path = GetReplicationSettingsInfo();
+			SettingsSection = SettingsModule->RegisterSettings(Path.ContainerName, Path.CategoryName, Path.SectionName,
 				LOCTEXT("MultiUserReplicationSettingsName", "Multi-User Replication"),
 				LOCTEXT("MultiUserReplicationSettingsDescription", "Configure the Multi-User Replication settings."),
 				UMultiUserReplicationSettings::Get()
@@ -48,8 +53,11 @@ namespace UE::MultiUserReplicationEditor
 		ISettingsModule* SettingsModule = FModuleManager::GetModulePtr<ISettingsModule>("Settings");
 		if (SettingsModule)
 		{
-			SettingsModule->UnregisterSettings("Project", "Plugins", "Multi-User Replication");
+			const FSettingPath Path = GetReplicationSettingsInfo();
+			SettingsModule->UnregisterSettings(Path.ContainerName, Path.CategoryName, Path.SectionName);
 		}
+
+		SettingsSection = nullptr;
 	}
 };
 
diff --git a/Engine/Plugins/Developer/Concert/ConcertApp/MultiUserClient/Source/MultiUserReplicationEditor/Private/MultiUserReplicationEditorModule.h b/Engine/Plugins/Developer/Concert/ConcertApp/MultiUserClient/Source/MultiUserReplicationEditor/Private/MultiUserReplicationEditorModule.h
--- a/Engine/Plugins/Developer/Concert/ConcertApp/MultiUserClient/Source/MultiUserReplicationEditor/Private/MultiUserReplicationEditorModule.h
+++ b/Engine/Plugins/Developer/Concert/ConcertApp/MultiUserClient/Source/MultiUserReplicationEditor/Private/MultiUserReplicationEditorModule.h
@@ -20,12 +20,15 @@ namespace UE::MultiUserReplicationEditor
 
 		//~ Begin IMultiUserReplicationEditor Interface
 		virtual FAssetCategoryPath GetMultiUserReplicationCategory() const override { return MultiUserReplicationCategory; }
+		virtual FSettingPath GetReplicationSettingsInfo() const override;
 		//~ End IMultiUserReplicationEditor Interface
 
 	private:
 
 		FAssetCategoryPath MultiUserReplicationCategory;
 
+		TSharedPtr<ISettingsSection> SettingsSection;
+
 		void RegisterSettings();
 		void UnregisterSettings();
 	};
diff --git a/Engine/Plugins/Developer/Concert/ConcertApp/MultiUserClient/Source/MultiUserReplicationEditor/Public/IMultiUserReplicationEditorModule.h b/Engine/Plugins/Developer/Concert/ConcertApp/MultiUserClient/Source/MultiUserReplicationEditor/Public/IMultiUserReplicationEditorModule.h
--- a/Engine/Plugins/Developer/Concert/ConcertApp/MultiUserClient/Source/MultiUserReplicationEditor/Public/IMultiUserReplicationEditorModule.h
+++ b/Engine/Plugins/Developer/Concert/ConcertApp/MultiUserClient/Source/MultiUserReplicationEditor/Public/IMultiUserReplicationEditorModule.h
@@ -6,6 +6,7 @@
 #include "Modules/ModuleManager.h"
 #include "UObject/NameTypes.h"
 
+class ISettingsSection;
 struct FAssetCategoryPath;
 
 namespace UE::MultiUserReplicationEditor
@@ -39,5 +40,14 @@ namespace UE::MultiUserReplicationEditor
 
 		/** Gets the category under which assets for Multi-User Replication are nested. */
 		virtual FAssetCategoryPath GetMultiUserReplicationCategory() const = 0;
+
+		struct FSettingPath
+		{
+			const FName ContainerName;
+			const FName CategoryName;
+			const FName SectionName;
+		};
+		/** Gets the sections the replication settings are stored under. */
+		virtual FSettingPath GetReplicationSettingsInfo() const = 0;
 	};
 }
\ No newline at end of file
diff --git a/Engine/Plugins/Developer/Concert/ConcertApp/MultiUserClient/Source/MultiUserReplicationEditor/Public/Settings/MultiUserReplicationSettings.h b/Engine/Plugins/Developer/Concert/ConcertApp/MultiUserClient/Source/MultiUserReplicationEditor/Public/Settings/MultiUserReplicationSettings.h
--- a/Engine/Plugins/Developer/Concert/ConcertApp/MultiUserClient/Source/MultiUserReplicationEditor/Public/Settings/MultiUserReplicationSettings.h
+++ b/Engine/Plugins/Developer/Concert/ConcertApp/MultiUserClient/Source/MultiUserReplicationEditor/Public/Settings/MultiUserReplicationSettings.h
@@ -4,11 +4,9 @@
 
 #include "UObject/Object.h"
 #include "MultiUserReplicationClientProfileAsset.h"
-#if WITH_EDITORONLY_DATA
 #include "Settings/MultiUserDefaultPropertySelection.h"
 #include "Settings/MultiUserDefaultSubobjectSelection.h"
 #include "Templates/Function.h"
-#endif
 #include "MultiUserReplicationSettings.generated.h"
 
 // Architecturally speaking FDefaultReplicationSessionSettings should be in the MultiUserClient module.
@@ -39,7 +37,6 @@ public:
 	UPROPERTY(EditAnywhere, Config, Category = "Replication")
 	FDefaultReplicationSessionSettings DefaultSessionSettings;
 
-#if WITH_EDITORONLY_DATA
 	/** Properties you want selected by default when you add a new replicated object in the editor */
 	UPROPERTY(EditAnywhere, Config, Category = "Replication|Editor")
 	TMap<FSoftClassPath, FMultiUserDefaultPropertySelection> DefaultPropertySelection;
@@ -52,15 +49,12 @@ public:
 	 */
 	UPROPERTY(EditAnywhere, Config, Category = "Replication|Editor")
 	TMap<FSoftClassPath, FMultiUserDefaultSubobjectSelection> DefaultSubobjectSelection;
-#endif
 
 	static UMultiUserReplicationSettings* Get() { return GetMutableDefault<UMultiUserReplicationSettings>(); }
 
-#if WITH_EDITORONLY_DATA
 	/** Reads DefaultPropertySelection and applies any default property selections to Info based on the Class just added. */
 	void AddDefaultPropertiesFromSettings(FReplicatedObjectInfo& Info, UClass& Class);
 
 	/** Reads DefaultComponentSelection and calls FurtherObjectsCallback on any further objects that should also be added. */
 	void AddAdditionalObjectsFromSettings(UObject& AddedObject, TFunctionRef<void(UObject&)> FurtherObjectsCallback);
-#endif
 };
\ No newline at end of file
