[Backout] - CL26955790
Original CL Desc
-----------------------------------------------------------------
Integrate MU replication system into the editor.
- Add replication tab to Concert window (in active session view)

diff --git a/Engine/Plugins/Developer/Concert/ConcertApp/MultiUserClient/Source/MultiUserClient/MultiUserClient.Build.cs b/Engine/Plugins/Developer/Concert/ConcertApp/MultiUserClient/Source/MultiUserClient/MultiUserClient.Build.cs
--- a/Engine/Plugins/Developer/Concert/ConcertApp/MultiUserClient/Source/MultiUserClient/MultiUserClient.Build.cs
+++ b/Engine/Plugins/Developer/Concert/ConcertApp/MultiUserClient/Source/MultiUserClient/MultiUserClient.Build.cs
@@ -29,9 +29,8 @@ namespace UnrealBuildTool.Rules
 					"DesktopPlatform",
 					"EditorStyle",
 					"InputCore",
-					"MessageLog",
-					"MultiUserReplication",
 					"Projects",
+					"MessageLog",
 					"Slate",
 					"SlateCore",
 					"SourceControl",
@@ -42,16 +41,13 @@ namespace UnrealBuildTool.Rules
 
 			if (Target.bBuildEditor)
 			{
-				PrivateDependencyModuleNames.AddRange( 
+				PrivateDependencyModuleNames.AddRange(
 					new string[]
-					{ 
+					{
 						"EditorFramework",
-						"MultiUserReplicationEditor",
-						"PropertyEditor",
-						"Settings",
+						"UnrealEd",
 						"Sequencer",
 						"ToolMenus",
-						"UnrealEd",
 						"WorkspaceMenuStructure",
 					}
 				);
diff --git a/Engine/Plugins/Developer/Concert/ConcertApp/MultiUserClient/Source/MultiUserClient/Private/Customization/ShowJoinSettingsCustomization.cpp b/Engine/Plugins/Developer/Concert/ConcertApp/MultiUserClient/Source/MultiUserClient/Private/Customization/ShowJoinSettingsCustomization.cpp
deleted file mode 100644
diff --git a/Engine/Plugins/Developer/Concert/ConcertApp/MultiUserClient/Source/MultiUserClient/Private/Customization/ShowJoinSettingsCustomization.h b/Engine/Plugins/Developer/Concert/ConcertApp/MultiUserClient/Source/MultiUserClient/Private/Customization/ShowJoinSettingsCustomization.h
deleted file mode 100644
diff --git a/Engine/Plugins/Developer/Concert/ConcertApp/MultiUserClient/Source/MultiUserClient/Private/MultiUserClientModule.cpp b/Engine/Plugins/Developer/Concert/ConcertApp/MultiUserClient/Source/MultiUserClient/Private/MultiUserClientModule.cpp
--- a/Engine/Plugins/Developer/Concert/ConcertApp/MultiUserClient/Source/MultiUserClient/Private/MultiUserClientModule.cpp
+++ b/Engine/Plugins/Developer/Concert/ConcertApp/MultiUserClient/Source/MultiUserClient/Private/MultiUserClientModule.cpp
@@ -1041,9 +1041,8 @@ private:
 	 */
 	TSharedRef<SDockTab> SpawnConcertBrowserTab(const FSpawnTabArgs& SpawnTabArgs)
 	{
-		const TSharedRef<SDockTab> DockTab = SNew(SDockTab)
-			.TabRole(NomadTab);
-		DockTab->SetContent(SNew(SConcertBrowser, DockTab, MultiUserClient));
+		const TSharedRef<SDockTab> DockTab = SNew(SDockTab).TabRole(ETabRole::NomadTab);
+		DockTab->SetContent(SNew(SConcertBrowser, DockTab, SpawnTabArgs.GetOwnerWindow(), MultiUserClient));
 		return DockTab;
 	}
 
diff --git a/Engine/Plugins/Developer/Concert/ConcertApp/MultiUserClient/Source/MultiUserClient/Private/Util/ReplicationUtils.cpp b/Engine/Plugins/Developer/Concert/ConcertApp/MultiUserClient/Source/MultiUserClient/Private/Util/ReplicationUtils.cpp
deleted file mode 100644
diff --git a/Engine/Plugins/Developer/Concert/ConcertApp/MultiUserClient/Source/MultiUserClient/Private/Util/ReplicationUtils.h b/Engine/Plugins/Developer/Concert/ConcertApp/MultiUserClient/Source/MultiUserClient/Private/Util/ReplicationUtils.h
deleted file mode 100644
diff --git a/Engine/Plugins/Developer/Concert/ConcertApp/MultiUserClient/Source/MultiUserClient/Private/Widgets/ActiveSession/Overview/SActiveSessionOverviewTab.cpp b/Engine/Plugins/Developer/Concert/ConcertApp/MultiUserClient/Source/MultiUserClient/Private/Widgets/ActiveSession/Overview/SActiveSessionOverviewTab.cpp
deleted file mode 100644
diff --git a/Engine/Plugins/Developer/Concert/ConcertApp/MultiUserClient/Source/MultiUserClient/Private/Widgets/ActiveSession/Overview/SActiveSessionOverviewTab.h b/Engine/Plugins/Developer/Concert/ConcertApp/MultiUserClient/Source/MultiUserClient/Private/Widgets/ActiveSession/Overview/SActiveSessionOverviewTab.h
deleted file mode 100644
diff --git a/Engine/Plugins/Developer/Concert/ConcertApp/MultiUserClient/Source/MultiUserClient/Private/Widgets/ActiveSession/Replication/SReplicationConnectionControls.cpp b/Engine/Plugins/Developer/Concert/ConcertApp/MultiUserClient/Source/MultiUserClient/Private/Widgets/ActiveSession/Replication/SReplicationConnectionControls.cpp
deleted file mode 100644
diff --git a/Engine/Plugins/Developer/Concert/ConcertApp/MultiUserClient/Source/MultiUserClient/Private/Widgets/ActiveSession/Replication/SReplicationConnectionControls.h b/Engine/Plugins/Developer/Concert/ConcertApp/MultiUserClient/Source/MultiUserClient/Private/Widgets/ActiveSession/Replication/SReplicationConnectionControls.h
deleted file mode 100644
diff --git a/Engine/Plugins/Developer/Concert/ConcertApp/MultiUserClient/Source/MultiUserClient/Private/Widgets/ActiveSession/Replication/SReplicationControlsTab.cpp b/Engine/Plugins/Developer/Concert/ConcertApp/MultiUserClient/Source/MultiUserClient/Private/Widgets/ActiveSession/Replication/SReplicationControlsTab.cpp
deleted file mode 100644
diff --git a/Engine/Plugins/Developer/Concert/ConcertApp/MultiUserClient/Source/MultiUserClient/Private/Widgets/ActiveSession/Replication/SReplicationControlsTab.h b/Engine/Plugins/Developer/Concert/ConcertApp/MultiUserClient/Source/MultiUserClient/Private/Widgets/ActiveSession/Replication/SReplicationControlsTab.h
deleted file mode 100644
diff --git a/Engine/Plugins/Developer/Concert/ConcertApp/MultiUserClient/Source/MultiUserClient/Private/Widgets/ActiveSession/Replication/SReplicationTabWithWarningOverlay.cpp b/Engine/Plugins/Developer/Concert/ConcertApp/MultiUserClient/Source/MultiUserClient/Private/Widgets/ActiveSession/Replication/SReplicationTabWithWarningOverlay.cpp
deleted file mode 100644
diff --git a/Engine/Plugins/Developer/Concert/ConcertApp/MultiUserClient/Source/MultiUserClient/Private/Widgets/ActiveSession/Replication/SReplicationTabWithWarningOverlay.h b/Engine/Plugins/Developer/Concert/ConcertApp/MultiUserClient/Source/MultiUserClient/Private/Widgets/ActiveSession/Replication/SReplicationTabWithWarningOverlay.h
deleted file mode 100644
diff --git a/Engine/Plugins/Developer/Concert/ConcertApp/MultiUserClient/Source/MultiUserClient/Private/Widgets/ActiveSession/SActiveSessionRoot.cpp b/Engine/Plugins/Developer/Concert/ConcertApp/MultiUserClient/Source/MultiUserClient/Private/Widgets/ActiveSession/SActiveSessionRoot.cpp
deleted file mode 100644
diff --git a/Engine/Plugins/Developer/Concert/ConcertApp/MultiUserClient/Source/MultiUserClient/Private/Widgets/ActiveSession/SActiveSessionRoot.h b/Engine/Plugins/Developer/Concert/ConcertApp/MultiUserClient/Source/MultiUserClient/Private/Widgets/ActiveSession/SActiveSessionRoot.h
deleted file mode 100644
diff --git a/Engine/Plugins/Developer/Concert/ConcertApp/MultiUserClient/Source/MultiUserClient/Private/Widgets/ActiveSession/SActiveSessionToolbar.cpp b/Engine/Plugins/Developer/Concert/ConcertApp/MultiUserClient/Source/MultiUserClient/Private/Widgets/ActiveSession/SActiveSessionToolbar.cpp
deleted file mode 100644
diff --git a/Engine/Plugins/Developer/Concert/ConcertApp/MultiUserClient/Source/MultiUserClient/Private/Widgets/ActiveSession/SActiveSessionToolbar.h b/Engine/Plugins/Developer/Concert/ConcertApp/MultiUserClient/Source/MultiUserClient/Private/Widgets/ActiveSession/SActiveSessionToolbar.h
deleted file mode 100644
diff --git a/Engine/Plugins/Developer/Concert/ConcertApp/MultiUserClient/Source/MultiUserClient/Private/Widgets/SActiveSession.cpp b/Engine/Plugins/Developer/Concert/ConcertApp/MultiUserClient/Source/MultiUserClient/Private/Widgets/SActiveSession.cpp
new file mode 100644
--- /dev/null
+++ b/Engine/Plugins/Developer/Concert/ConcertApp/MultiUserClient/Source/MultiUserClient/Private/Widgets/SActiveSession.cpp
@@ -0,0 +1,1018 @@
+// Copyright Epic Games, Inc. All Rights Reserved.
+
+#include "SActiveSession.h"
+
+#include "ConcertActionDefinition.h"
+#include "ConcertClientFrontendUtils.h"
+#include "ConcertFrontendUtils.h"
+#include "ConcertMessageData.h"
+#include "ClientSessionHistoryController.h"
+#include "ConcertMessages.h"
+#include "Dialog/SMessageDialog.h"
+#include "IConcertClientPackageBridge.h"
+#include "IConcertClientPresenceManager.h"
+#include "IConcertClient.h"
+#include "IConcertClientWorkspace.h"
+#include "IConcertSyncClient.h"
+#include "ConcertClientSettings.h"
+#include "IMultiUserClientModule.h"
+
+#include "Algo/Transform.h"
+#include "EditorFontGlyphs.h"
+#include "FileHelpers.h"
+#include "ISettingsModule.h"
+#include "Misc/AsyncTaskNotification.h"
+#include "Misc/PackageName.h"
+#include "Styling/AppStyle.h"
+#include "Dialog/SCustomDialog.h"
+#include "Session/History/SSessionHistory.h"
+#include "Session/History/SSessionHistoryWrapper.h"
+#include "Widgets/Notifications/SNotificationList.h"
+#include "Widgets/Text/STextBlock.h"
+#include "Widgets/Images/SImage.h"
+#include "Widgets/Input/SHyperlink.h"
+#include "Widgets/Layout/SExpandableArea.h"
+#include "Widgets/Views/STableViewBase.h"
+#include "Widgets/Views/STableRow.h"
+#include "Framework/Notifications/NotificationManager.h"
+#include "HAL/PlatformFileManager.h"
+
+
+#define LOCTEXT_NAMESPACE "SActiveSession"
+
+namespace ActiveSessionDetailsUI
+{
+	static const FName DisplayNameColumnName(TEXT("DisplayName"));
+	static const FName PresenceColumnName(TEXT("Presence"));
+	static const FName LevelColumnName(TEXT("Level"));
+}
+
+
+class SActiveSessionDetailsRow : public SMultiColumnTableRow<TSharedPtr<FConcertSessionClientInfo>>
+{
+	SLATE_BEGIN_ARGS(SActiveSessionDetailsRow) {}
+	SLATE_END_ARGS()
+
+public:
+	/**
+	 * Constructs the widget.
+	 *
+	 * @param InArgs The construction arguments.
+	 * @param InClientInfo The client displayed by this row.
+	 * @param InClientSession The session in which the client is, used to determine if the client is the local one, so that we can suffix it with a "you".
+	 * @param InOwnerTableView The table to which the row must be added.
+	 */
+	void Construct(const FArguments& InArgs, TWeakPtr<IConcertSyncClient> InSyncClient, TWeakPtr<IConcertClientSession> InClientSession, TSharedPtr<FConcertSessionClientInfo> InClientInfo, const TSharedRef<STableViewBase>& InOwnerTableView)
+	{
+		SyncClient = MoveTemp(InSyncClient);
+		ClientSession = MoveTemp(InClientSession);
+		SessionClientInfo = MoveTemp(InClientInfo);
+		SMultiColumnTableRow<TSharedPtr<FConcertSessionClientInfo>>::Construct(FSuperRowType::FArguments(), InOwnerTableView);
+
+		// Set the tooltip for the entire row. Will show up unless there is another item with a tooltip hovered in the row, such as the "presence" icons.
+		SetToolTipText(MakeAttributeSP(this, &SActiveSessionDetailsRow::GetRowToolTip));
+	}
+
+public:
+	virtual TSharedRef<SWidget> GenerateWidgetForColumn(const FName& ColumnName) override
+	{
+		if (ColumnName == ActiveSessionDetailsUI::DisplayNameColumnName)
+		{
+			// Displays a colored square from a special font (using avatar color) followed by the the display name -> [x] John Smith
+			return SNew(SHorizontalBox)
+				// The 'square' glyph in front of the client name, rendered using special font glyph, in the client avatar color.
+				+SHorizontalBox::Slot()
+				.AutoWidth()
+				.VAlign(VAlign_Center)
+				.Padding(FMargin(4.0f, 0.0f, 0.0f, 0.0f))
+				[
+					SNew(STextBlock)
+					.Font(this, &SActiveSessionDetailsRow::GetAvatarFont)
+					.ColorAndOpacity(this, &SActiveSessionDetailsRow::GetAvatarColor)
+					.Text(FEditorFontGlyphs::Square)
+				]
+
+				// The client display name.
+				+SHorizontalBox::Slot()
+				.VAlign(VAlign_Center)
+				[
+					ConcertFrontendUtils::CreateDisplayName(MakeAttributeSP(this, &SActiveSessionDetailsRow::GetDisplayName))
+				];
+		}
+		else if (ColumnName == ActiveSessionDetailsUI::PresenceColumnName)
+		{
+			// Displays a set of icons corresponding to the client presence. The set may be extended later to include other functionalities.
+			TArray<FConcertActionDefinition> ActionDefs;
+			TSharedRef<SHorizontalBox> PresenceCell = SNew(SHorizontalBox);
+
+			TSharedPtr<FConcertSessionClientInfo> ClientInfoPin = SessionClientInfo.Pin();
+			TSharedPtr<IConcertSyncClient> SyncClientPin = SyncClient.Pin();
+			if (ClientInfoPin.IsValid() && SyncClientPin.IsValid())
+			{
+				SyncClientPin->GetSessionClientActions(*ClientInfoPin, ActionDefs);
+				ConcertClientFrontendUtils::AppendButtons(PresenceCell, ActionDefs);
+			}
+			return PresenceCell;
+		}
+		else // LevelColumnName
+		{
+			check(ColumnName == ActiveSessionDetailsUI::LevelColumnName); // If this fail, was a column added/removed/renamed ?
+
+			TSharedPtr<FConcertSessionClientInfo> ClientInfoPin = SessionClientInfo.Pin();
+			TSharedPtr<IConcertClientSession> ClientSessionPin = ClientSession.Pin();
+			TSharedPtr<SWidget> LevelWidget;
+
+			// If the row is not for "this client", then add a hyper link to open the other client level.
+			if (ClientInfoPin.IsValid() && ClientSessionPin.IsValid() && ClientInfoPin->ClientEndpointId != ClientSessionPin->GetSessionClientEndpointId())
+			{
+				LevelWidget = SNew(SOverlay)
+					+SOverlay::Slot()
+					[
+						SNew(SHyperlink)
+						.Visibility_Lambda([this]() { return GetOtherLevelVisibility(); })
+						.Text(this, &SActiveSessionDetailsRow::GetLevel)
+						.ToolTipText(this, &SActiveSessionDetailsRow::GetOtherLevelTooltip)
+						.OnNavigate(this, &SActiveSessionDetailsRow::OnOtherLevelClicked)
+					]
+					+SOverlay::Slot()
+					[
+						SNew(STextBlock)
+						.Visibility_Lambda([this]() { return GetOtherLevelVisibility() == EVisibility::Visible ? EVisibility::Collapsed : EVisibility::Visible; })
+						.Text(LOCTEXT("LevelNotAvailableInGame", "N/A (-game)")) // The only known reason for the level to be unknown and invisible is b/c the user is in -game.
+						.ToolTipText(LOCTEXT("LevelNotAvailableInGame_Tooltip", "The presence/level information is not emitted by clients running Editor in -game mode."))
+					];
+			}
+			else
+			{
+				LevelWidget = SNew(STextBlock)
+					.Text(this, &SActiveSessionDetailsRow::GetLevel)
+					.ToolTipText(FText(LOCTEXT("YourLevel", "The level opened in your editor.")));
+			}
+
+			// Displays which "level" the client is editing, playing (PIE) or simulating (SIE).
+			return SNew(SHorizontalBox)
+				+SHorizontalBox::Slot()
+				.Padding(0.0f, 4.0f, 2.0f, 6.0f)
+				.AutoWidth()
+				[
+					SNew(SImage)
+					.Image(FAppStyle::Get().GetBrush("LevelEditor.Tabs.Levels"))
+				]
+				+SHorizontalBox::Slot()
+				.AutoWidth()
+				.Padding(FMargin(2.0f, 4.0f))
+				[
+					LevelWidget.ToSharedRef()
+				];
+		}
+	}
+
+	FText GetRowToolTip() const
+	{
+		// This is a tooltip for the entire row. Like display name, the tooltip will not update in real time if the user change its
+		// settings. See GetDisplayName() for more info.
+		TSharedPtr<FConcertSessionClientInfo> ClientInfoPin = SessionClientInfo.Pin();
+		return ClientInfoPin.IsValid() ? ClientInfoPin->ToDisplayString() : FText();
+	}
+
+	FText GetDisplayName() const
+	{
+		TSharedPtr<FConcertSessionClientInfo> ClientInfoPin = SessionClientInfo.Pin();
+		if (ClientInfoPin.IsValid())
+		{
+			// NOTE: The display name doesn't update in real time at the moment because the concert setting are not propagated
+			//       until the client disconnect/reconnect. Since those settings should not change often, this should not
+			//       be a major deal breaker for the users.
+			TSharedPtr<IConcertClientSession> ClientSessionPin = ClientSession.Pin();
+			if (ClientSessionPin.IsValid() && ClientInfoPin->ClientEndpointId == ClientSessionPin->GetSessionClientEndpointId())
+			{
+				return FText::Format(LOCTEXT("ClientDisplayNameIsYouFmt", "{0} (You)"), FText::FromString(ClientSessionPin->GetLocalClientInfo().DisplayName));
+			}
+
+			// Return the ClientInfo cached.
+			return FText::FromString(ClientInfoPin->ClientInfo.DisplayName);
+		}
+
+		return FText();
+	}
+
+	EVisibility GetOtherLevelVisibility() const
+	{
+		TSharedPtr<FConcertSessionClientInfo> ClientInfoPin = SessionClientInfo.Pin();
+		TSharedPtr<IConcertSyncClient> SyncClientPin = SyncClient.Pin();
+
+		if (ClientInfoPin && SyncClientPin)
+		{
+			EEditorPlayMode EditorPlayMode = EEditorPlayMode::None;
+			if (IConcertClientPresenceManager* PresenceManager = SyncClientPin->GetPresenceManager())
+			{
+				if (!PresenceManager->GetPresenceWorldPath(ClientInfoPin->ClientEndpointId, EditorPlayMode).IsEmpty())
+				{
+					return EVisibility::Visible;
+				}
+			}
+		}
+
+		return EVisibility::Hidden;
+	}
+
+	FText GetOtherLevelTooltip() const
+	{
+		if (TSharedPtr<FConcertSessionClientInfo> ClientInfoPin = SessionClientInfo.Pin())
+		{
+			return FText::Format(LOCTEXT("OtherUserLevel_Tooltip", "The level opened in {0} editor. Click to open this level."), FText::FromString(ClientInfoPin->ClientInfo.DisplayName));
+		}
+
+		return FText::GetEmpty();
+	}
+
+	FText GetLevel() const
+	{
+		FString WorldPath;
+		
+		TSharedPtr<FConcertSessionClientInfo> ClientInfoPin = SessionClientInfo.Pin();
+		TSharedPtr<IConcertSyncClient> SyncClientPin = SyncClient.Pin();
+
+		if (ClientInfoPin && SyncClientPin)
+		{
+			// The world path returned by the function below is always the one from the editor context, never the one from the PIE/SIE context, so
+			// it will not contain the prefix added to the world when in PIE/SIE.
+			EEditorPlayMode EditorPlayMode = EEditorPlayMode::None;
+			if (IConcertClientPresenceManager* PresenceManager = SyncClientPin->GetPresenceManager())
+			{
+				WorldPath = PresenceManager->GetPresenceWorldPath(ClientInfoPin->ClientEndpointId, EditorPlayMode);
+			}
+			
+			// The world path is returned as something like /Game/MyMap.MyMap, but we are only interested to keep the
+			// string left to the '.' to display "/Game/MyMap"
+			WorldPath = FPackageName::ObjectPathToPackageName(WorldPath);
+
+			if (EditorPlayMode == EEditorPlayMode::PIE)
+			{
+				WorldPath += TEXT(" (PIE)");
+			}
+			else if (EditorPlayMode == EEditorPlayMode::SIE)
+			{
+				WorldPath += TEXT(" (SIE)");
+			}
+		}
+
+		return FText::FromString(WorldPath);
+	}
+
+	void OnOtherLevelClicked()
+	{
+		TSharedPtr<FConcertSessionClientInfo> ClientInfoPin = SessionClientInfo.Pin();
+		TSharedPtr<IConcertClientSession> ClientSessionPin = ClientSession.Pin();
+		TSharedPtr<IConcertSyncClient> ConcertSyncClientPin = SyncClient.Pin();
+
+		// If the user clicks on a client row other than the one that represents them.
+		if (ClientInfoPin && ClientSessionPin && ConcertSyncClientPin && ConcertSyncClientPin->GetPresenceManager() &&
+			ClientInfoPin->ClientEndpointId != ClientSessionPin->GetSessionClientEndpointId())
+		{
+			IConcertClientPresenceManager* PresenceManager = ConcertSyncClientPin->GetPresenceManager();
+			EEditorPlayMode EditorPlayMode; // The result stored in this variable is not used.
+			FString RowWorldPath = PresenceManager->GetPresenceWorldPath(ClientInfoPin->ClientEndpointId, EditorPlayMode);
+			FString ThisClientSessionWorldPath = PresenceManager->GetPresenceWorldPath(ClientSessionPin->GetSessionClientEndpointId(), EditorPlayMode);
+
+			// If this client is not editing the same world as the client represented by this row.
+			if (ThisClientSessionWorldPath != RowWorldPath)
+			{
+				// If there are any unsaved changes to the current level, see if the user wants to save those first.
+				bool bPromptUserToSave = true;
+				bool bSaveMapPackages = true;
+				bool bSaveContentPackages = false;
+				if (FEditorFileUtils::SaveDirtyPackages(bPromptUserToSave, bSaveMapPackages, bSaveContentPackages) == false)
+				{
+					return;
+				}
+				
+				// Open the level.
+				FString WorlPathname;
+				const bool bLoadAsTemplate = false;
+				const bool bShowProgress = true;
+				FPackageName::TryConvertLongPackageNameToFilename(RowWorldPath, WorlPathname);
+				FEditorFileUtils::LoadMap(WorlPathname, bLoadAsTemplate, bShowProgress);
+
+				// Jump to the client position/orientation to view the same thing as client corresponding to the clicked row currently look at.
+				PresenceManager->InitiateJumpToPresence(ClientInfoPin->ClientEndpointId);
+			}
+		}
+	}
+
+	FSlateFontInfo GetAvatarFont() const
+	{
+		// This font is used to render a small square box filled with the avatar color.
+		FSlateFontInfo ClientIconFontInfo = FAppStyle::Get().GetFontStyle(ConcertClientFrontendUtils::ButtonIconSyle);
+		ClientIconFontInfo.Size = 8;
+		ClientIconFontInfo.OutlineSettings.OutlineSize = 1;
+
+		TSharedPtr<FConcertSessionClientInfo> ClientInfoPin = SessionClientInfo.Pin();
+		if (ClientInfoPin.IsValid())
+		{
+			FLinearColor ClientOutlineColor = ClientInfoPin->ClientInfo.AvatarColor * 0.6f; // Make the font outline darker.
+			ClientOutlineColor.A = ClientInfoPin->ClientInfo.AvatarColor.A; // Put back the original alpha.
+			ClientIconFontInfo.OutlineSettings.OutlineColor = ClientOutlineColor;
+		}
+		else
+		{
+			ClientIconFontInfo.OutlineSettings.OutlineColor = FLinearColor(0.75, 0.75, 0.75); // This is an arbitrary color.
+		}
+
+		return ClientIconFontInfo;
+	}
+
+	FSlateColor GetAvatarColor() const
+	{
+		TSharedPtr<FConcertSessionClientInfo> ClientInfoPin = SessionClientInfo.Pin();
+		if (ClientInfoPin.IsValid())
+		{
+			return ClientInfoPin->ClientInfo.AvatarColor;
+		}
+
+		return FSlateColor(FLinearColor(0.75, 0.75, 0.75)); // This is an arbitrary color.
+	}
+
+private:
+	TWeakPtr<IConcertSyncClient> SyncClient;
+	TWeakPtr<IConcertClientSession> ClientSession;
+	TWeakPtr<FConcertSessionClientInfo> SessionClientInfo;
+};
+
+int32 SActiveSession::GetInitialSendReceiveComboIndex()
+{
+	if (!WeakSessionPtr.IsValid())
+	{
+		return 0;
+	}
+
+	TSharedPtr<IConcertClientSession> Session = WeakSessionPtr.Pin();
+	switch(Session->GetSendReceiveState())
+	{
+	case EConcertSendReceiveState::Default:
+		return 0;
+	case EConcertSendReceiveState::SendOnly:
+		return 1;
+	case EConcertSendReceiveState::ReceiveOnly:
+		return 2;
+	}
+
+	return 0;
+}
+
+TArray< TSharedPtr< SActiveSession::FSendReceiveComboItem > > ConstructSendReceiveComboList()
+{
+	TArray< TSharedPtr< SActiveSession::FSendReceiveComboItem > > ComboList;
+	ComboList.Add( MakeShared<SActiveSession::FSendReceiveComboItem>(
+			LOCTEXT("DefaultSendReceiveState", "Default"),
+			LOCTEXT("DefaultSendReceiveStateTooltip", "Full send/receive mode for multi-user events."),
+			EConcertSendReceiveState::Default)
+		);
+
+	ComboList.Add( MakeShared<SActiveSession::FSendReceiveComboItem>(
+			LOCTEXT("SendState", "Send only"),
+			LOCTEXT("SendStateTooltip", "Transactions received from clients will be suspended; however local transactions will be sent."),
+			EConcertSendReceiveState::SendOnly)
+		);
+
+	ComboList.Add( MakeShared<SActiveSession::FSendReceiveComboItem>(
+			LOCTEXT("ReceiveState", "Receive only"),
+			LOCTEXT("ReceiveStateTooltip", "Local changes will be queued for transmission but not sent. Updates from clients will be received."),
+			EConcertSendReceiveState::ReceiveOnly)
+		);
+	return ComboList;
+}
+
+void SActiveSession::Construct(const FArguments& InArgs, TSharedPtr<IConcertSyncClient> InConcertSyncClient)
+{
+	WeakConcertSyncClient = InConcertSyncClient;
+	SessionHistoryController = MakeShared<FClientSessionHistoryController>(InConcertSyncClient.ToSharedRef());
+
+	if (InConcertSyncClient.IsValid())
+	{
+		IConcertClientRef ConcertClient = InConcertSyncClient->GetConcertClient();
+		ConcertClient->OnSessionStartup().AddSP(this, &SActiveSession::HandleSessionStartup);
+		ConcertClient->OnSessionShutdown().AddSP(this, &SActiveSession::HandleSessionShutdown);
+
+		if (TSharedPtr<IConcertClientSession> ClientSession = ConcertClient->GetCurrentSession())
+		{
+			WeakSessionPtr = ClientSession;
+			ClientInfo = MakeShared<FConcertSessionClientInfo>(FConcertSessionClientInfo{ClientSession->GetSessionClientEndpointId(), ClientSession->GetLocalClientInfo()});
+			ClientSession->OnSessionClientChanged().AddSP(this, &SActiveSession::HandleSessionClientChanged);
+			InConcertSyncClient->GetTransactionBridge()->OnConflictResolutionForPendingSend().AddSP(this, &SActiveSession::OnSendConflict);
+		}
+
+		TSharedPtr<IConcertClientWorkspace> Workspace = InConcertSyncClient->GetWorkspace();
+		Workspace->OnActivityAddedOrUpdated().AddSP(this, &SActiveSession::ActivityUpdated);
+		Workspace->AddWorkspaceCanProcessPackagesDelegate(
+			TEXT("SActiveSession"), FCanProcessPendingPackages::CreateSP(this, &SActiveSession::CanProcessPendingPackages));
+	}
+
+	SendReceiveComboList = ConstructSendReceiveComboList();
+	int32 InitialIndex = GetInitialSendReceiveComboIndex();
+
+	TSharedRef<SHorizontalBox> StatusBar =
+		SNew(SHorizontalBox)
+
+		// Status Icon
+		+SHorizontalBox::Slot()
+		.AutoWidth()
+		.VAlign(VAlign_Center)
+		.Padding(FMargin(2.0f, 1.0f, 0.0f, 1.0f))
+		[
+			SNew(STextBlock)
+			.Font(this, &SActiveSession::GetConnectionIconFontInfo)
+			.ColorAndOpacity(this, &SActiveSession::GetConnectionIconColor)
+			.Text(FEditorFontGlyphs::Circle)
+		]
+
+		// Status Message
+		+SHorizontalBox::Slot()
+		.VAlign(VAlign_Center)
+		.Padding(FMargin(4.0f, 1.0f))
+		[
+			SNew(SBorder)
+			.BorderImage(FAppStyle::Get().GetBrush("NoBorder"))
+			.ColorAndOpacity(FLinearColor(0.75f, 0.75f, 0.75f))
+			.Padding(FMargin(0.0f, 4.0f, 6.0f, 4.0f))
+			[
+				SNew(STextBlock)
+				.Font(FAppStyle::Get().GetFontStyle("BoldFont"))
+				.Text(this, &SActiveSession::GetConnectionStatusText)
+			]
+		]
+		+SHorizontalBox::Slot()
+		.AutoWidth()
+		.VAlign(VAlign_Center)
+		.Padding(FMargin(1.0f))
+		[
+			SNew(SBox)
+			.MinDesiredWidth(130.0f)
+			[
+				SAssignNew(SendReceiveComboBox, SComboBox< TSharedPtr<FSendReceiveComboItem> >)
+				.OptionsSource(&SendReceiveComboList)
+				.InitiallySelectedItem(SendReceiveComboList[InitialIndex])
+				.ContentPadding(FMargin(4.0f,1.0f))
+				.OnGenerateWidget(this, &SActiveSession::GenerateSendReceiveComboItem)
+				.OnSelectionChanged(this, &SActiveSession::HandleSendReceiveChanged)
+				[
+					SNew(STextBlock)
+					.Text(this, &SActiveSession::GetRequestedSendReceiveComboText)
+				]
+			]
+		]
+		// The "Settings" icons.
+		+SHorizontalBox::Slot()
+		.VAlign(VAlign_Fill)
+		.AutoWidth()
+		[
+			SNew(SButton)
+			.ButtonStyle(FAppStyle::Get(), TEXT("SimpleButton"))
+			.OnClicked_Lambda([](){FModuleManager::GetModulePtr<ISettingsModule>("Settings")->ShowViewer("Project", "Plugins", "Concert"); return FReply::Handled(); })
+			.VAlign(VAlign_Center)
+			.HAlign(HAlign_Fill)
+			[
+				SNew(SImage)
+				.Image(FAppStyle::GetBrush("Icons.Settings"))
+			]
+		];
+
+	// Append the buttons to the status bar
+	{
+		TArray<FConcertActionDefinition> ButtonDefs;
+
+		// Leave Session
+		FConcertActionDefinition& LeaveSessionDef = ButtonDefs.AddDefaulted_GetRef();
+		LeaveSessionDef.Type = EConcertActionType::Normal;
+		LeaveSessionDef.IsVisible = MakeAttributeSP(this, &SActiveSession::IsStatusBarLeaveSessionVisible);
+		LeaveSessionDef.Text = FEditorFontGlyphs::Sign_Out;
+		LeaveSessionDef.ToolTipText = LOCTEXT("LeaveCurrentSessionToolTip", "Leave the current session");
+		LeaveSessionDef.OnExecute.BindLambda([this]() { OnClickLeaveSession(); });
+		LeaveSessionDef.IconStyle = TEXT("Concert.LeaveSession");
+
+		ConcertClientFrontendUtils::AppendButtons(StatusBar, ButtonDefs);
+	}
+
+	ChildSlot
+	[
+		SNew(SVerticalBox)
+
+		// Status bar.
+		+SVerticalBox::Slot()
+		.AutoHeight()
+		.Padding(0.0f, 0.0f, 0.0f, 4.0f) // Add space between the status bar and the clients list.
+		[
+			SNew(SBorder)
+			.BorderImage(FAppStyle::Get().GetBrush("ToolPanel.GroupBorder"))
+			.Padding(0.0f)
+			[
+				StatusBar
+			]
+		]
+
+		// Clients + History
+		+SVerticalBox::Slot()
+		.FillHeight(1.0)
+		[
+			SNew(SSplitter)
+			.Orientation(Orient_Vertical)
+
+			// Clients.
+			+SSplitter::Slot()
+			.SizeRule(TAttribute<SSplitter::ESizeRule>(this, &SActiveSession::GetClientAreaSizeRule))
+			.Value(0.3)
+			[
+				SNew(SBorder)
+				.BorderImage(FAppStyle::Get().GetBrush("ToolPanel.GroupBorder"))
+				.Padding(0.0f)
+				[
+					SAssignNew(ClientArea, SExpandableArea)
+					.BorderBackgroundColor(FLinearColor(0.6f, 0.6f, 0.6f, 1.0f))
+					.BorderImage_Lambda([this]() { return ConcertFrontendUtils::GetExpandableAreaBorderImage(*ClientArea); })
+					.BodyBorderImage(FAppStyle::Get().GetBrush("ToolPanel.GroupBorder"))
+					.BodyBorderBackgroundColor(FLinearColor::White)
+					.OnAreaExpansionChanged(this, &SActiveSession::OnClientAreaExpansionChanged)
+					.Padding(0.0f)
+					.HeaderContent()
+					[
+						SNew(STextBlock)
+						.Text(LOCTEXT("SessionConnectedClients", "Clients"))
+						.Font(FAppStyle::Get().GetFontStyle("DetailsView.CategoryFontStyle"))
+						.ShadowOffset(FVector2D(1.0f, 1.0f))
+					]
+					.BodyContent()
+					[
+						SAssignNew(ClientsListView, SListView<TSharedPtr<FConcertSessionClientInfo>>)
+						.ItemHeight(20.0f)
+						.SelectionMode(ESelectionMode::Single)
+						.ListItemsSource(&Clients)
+						.OnGenerateRow(this, &SActiveSession::HandleGenerateRow)
+						.HeaderRow
+						(
+							SNew(SHeaderRow)
+							+SHeaderRow::Column(ActiveSessionDetailsUI::DisplayNameColumnName)
+							.DefaultLabel(LOCTEXT("UserDisplayName", "Display Name"))
+							+SHeaderRow::Column(ActiveSessionDetailsUI::PresenceColumnName)
+							.DefaultLabel(LOCTEXT("UserPresence", "Presence"))
+							+SHeaderRow::Column(ActiveSessionDetailsUI::LevelColumnName)
+							.DefaultLabel(LOCTEXT("UserLevel", "Level"))
+						)
+					]
+				]
+			]
+
+			// History (activity feed).
+			+SSplitter::Slot()
+			.SizeRule(TAttribute<SSplitter::ESizeRule>(this, &SActiveSession::GetHistoryAreaSizeRule))
+			.Value(0.7)
+			[
+				SNew(SBorder)
+				.BorderImage(FAppStyle::Get().GetBrush("ToolPanel.GroupBorder"))
+				.Padding(0.0f)
+				[
+					SAssignNew(HistoryArea, SExpandableArea)
+					.BorderBackgroundColor(FLinearColor(0.6f, 0.6f, 0.6f, 1.0f))
+					.BorderImage_Lambda([this]() { return ConcertFrontendUtils::GetExpandableAreaBorderImage(*HistoryArea); })
+					.BodyBorderImage(FAppStyle::Get().GetBrush("ToolPanel.GroupBorder"))
+					.BodyBorderBackgroundColor(FLinearColor::White)
+					.OnAreaExpansionChanged(this, &SActiveSession::OnHistoryAreaExpansionChanged)
+					.Padding(0.0f)
+					.HeaderContent()
+					[
+						SNew(STextBlock)
+						.Text(LOCTEXT("SessionHistory", "History"))
+						.Font(FAppStyle::Get().GetFontStyle("DetailsView.CategoryFontStyle"))
+						.ShadowOffset(FVector2D(1.0f, 1.0f))
+					]
+					.BodyContent()
+					[
+						SessionHistoryController->GetSessionHistory()
+					]
+				]
+			]
+		]
+	];
+
+	// Create a timer to periodically poll the this client(s) info to detect if it changed its display name or avatar color because
+	// IConcertClientsession::OnSessionClientChanged() doesn't trigger when the 'local' client changes. This needs to be polled.
+	RegisterActiveTimer(1.0f, FWidgetActiveTimerDelegate::CreateSP(this, &SActiveSession::HandleLocalClientInfoChangePollingTimer));
+
+	UpdateSessionClientListView();
+}
+
+FText SActiveSession::GetRequestedSendReceiveComboText() const
+{
+	if (SendReceiveComboBox->GetSelectedItem().IsValid())
+	{
+		return SendReceiveComboBox->GetSelectedItem()->Name;
+	}
+	else
+	{
+		return LOCTEXT("ActiveSessionDefaultSendReceive", "Default");
+	}
+}
+
+TSharedRef<SWidget> SActiveSession::GenerateSendReceiveComboItem(TSharedPtr<SActiveSession::FSendReceiveComboItem> InItem)
+{
+	return SNew(STextBlock)
+		.Text(InItem->Name)
+		.ToolTipText(InItem->ToolTip);
+}
+
+void SActiveSession::HandleSendReceiveChanged(TSharedPtr<SActiveSession::FSendReceiveComboItem> Item, ESelectInfo::Type SelectInfo)
+{
+	check(Item.IsValid());
+	if (TSharedPtr<IConcertClientSession> Session = WeakSessionPtr.Pin())
+	{
+		Session->SetSendReceiveState(Item->State);
+	}
+}
+
+TSharedRef<ITableRow> SActiveSession::HandleGenerateRow(TSharedPtr<FConcertSessionClientInfo> InClientInfo, const TSharedRef<STableViewBase>& OwnerTable) const
+{
+	// Generate a row for the client corresponding to InClientInfo.
+	return SNew(SActiveSessionDetailsRow, WeakConcertSyncClient, WeakSessionPtr, InClientInfo, OwnerTable);
+}
+
+void SActiveSession::HandleSessionStartup(TSharedRef<IConcertClientSession> InClientSession)
+{
+	WeakSessionPtr = InClientSession;
+	InClientSession->OnSessionClientChanged().AddSP(this, &SActiveSession::HandleSessionClientChanged);
+	if (TSharedPtr<IConcertSyncClient> SyncClient = WeakConcertSyncClient.Pin())
+	{
+		SyncClient->GetTransactionBridge()->OnConflictResolutionForPendingSend().AddSP(this, &SActiveSession::OnSendConflict);
+	}
+
+	ClientInfo = MakeShared<FConcertSessionClientInfo>(FConcertSessionClientInfo{InClientSession->GetSessionClientEndpointId(), InClientSession->GetLocalClientInfo()});
+
+	UpdateSessionClientListView();
+
+	if (SessionHistoryController.IsValid())
+	{
+		SessionHistoryController->ReloadActivities();
+	}
+}
+
+void SActiveSession::HandleSessionShutdown(TSharedRef<IConcertClientSession> InClientSession)
+{
+	if (InClientSession == WeakSessionPtr)
+	{
+		WeakSessionPtr.Reset();
+		InClientSession->OnSessionClientChanged().RemoveAll(this);
+		if (TSharedPtr<IConcertSyncClient> SyncClient = WeakConcertSyncClient.Pin())
+		{
+			SyncClient->GetTransactionBridge()->OnConflictResolutionForPendingSend().RemoveAll(this);
+		}
+		Clients.Reset();
+
+		if (ClientsListView.IsValid())
+		{
+			ClientsListView->RequestListRefresh();
+		}
+
+		if (SessionHistoryController.IsValid())
+		{
+			SessionHistoryController->ReloadActivities();
+		}
+	}
+}
+
+void SActiveSession::HandleSessionClientChanged(IConcertClientSession&, EConcertClientStatus ClientStatus, const FConcertSessionClientInfo& InClientInfo)
+{
+	// Update the view for a specific client.
+	UpdateSessionClientListView(&InClientInfo, ClientStatus);
+}
+
+EActiveTimerReturnType SActiveSession::HandleLocalClientInfoChangePollingTimer(double InCurrentTime, float InDeltaTime)
+{
+// This switch is used to test if updating the session FConcertClientInfo dynamically correctly applies across the board. Don't
+// turn it on by default, it would overwrite runtime changes made to the info from some state/logic with the one from the settings.
+// For testing: Keep the MU browser tab open, join a session, open the MU settings and change one of the updatable field tracked below.
+#define MU_TEST_UPDATE_CLIENT_INFO_FROM_SETTINGS 0
+
+	TSharedPtr<IConcertClientSession> Session = WeakSessionPtr.Pin();
+	if (Session.IsValid() && ClientInfo.IsValid())
+	{
+#if MU_TEST_UPDATE_CLIENT_INFO_FROM_SETTINGS
+		if (TSharedPtr<IConcertSyncClient> SyncClientPin = WeakConcertSyncClient.Pin())
+		{
+			// User changed its settings.
+			const FConcertClientInfo& LatestClientInfo = Session->GetLocalClientInfo();
+			const UConcertClientConfig* ConcertClientConfig = SyncClientPin->GetConcertClient()->GetConfiguration();
+			if (ConcertClientConfig->ClientSettings.AvatarColor != LatestClientInfo.AvatarColor ||
+				ConcertClientConfig->ClientSettings.DisplayName != LatestClientInfo.DisplayName ||
+				ConcertClientConfig->ClientSettings.DesktopAvatarActorClass != LatestClientInfo.DesktopAvatarActorClass ||
+				ConcertClientConfig->ClientSettings.VRAvatarActorClass != LatestClientInfo.VRAvatarActorClass)
+			{
+				FConcertClientInfoUpdate ClientInfoUpdate;
+				ClientInfoUpdate.DisplayName = ConcertClientConfig->ClientSettings.DisplayName;
+				ClientInfoUpdate.AvatarColor = ConcertClientConfig->ClientSettings.AvatarColor;
+				ClientInfoUpdate.DesktopAvatarActorClass = ConcertClientConfig->ClientSettings.DesktopAvatarActorClass.ToString();
+				ClientInfoUpdate.VRAvatarActorClass = ConcertClientConfig->ClientSettings.VRAvatarActorClass.ToString();
+				Session->UpdateLocalClientInfo(ClientInfoUpdate);
+			}
+		}
+#endif
+
+		// Check if the local client info cached as class member is out dated with respect to the one held by the session changes. Just check the info displayed by this panel.
+		const FConcertClientInfo& LatestClientInfo = Session->GetLocalClientInfo();
+		if (LatestClientInfo.DisplayName != ClientInfo->ClientInfo.DisplayName || LatestClientInfo.AvatarColor != ClientInfo->ClientInfo.AvatarColor)
+		{
+			// Update the view for this client info.
+			ClientInfo->ClientInfo = LatestClientInfo;
+			UpdateSessionClientListView(ClientInfo.Get(), EConcertClientStatus::Updated);
+		}
+	}
+
+	return EActiveTimerReturnType::Continue;
+}
+
+void SActiveSession::UpdateSessionClientListView(const FConcertSessionClientInfo* InClientInfo, EConcertClientStatus ClientStatus)
+{
+	// We expect the UI to be constructed in Construct() function, prior this function gets called.
+	check(ClientsListView.IsValid());
+
+	// Try pinning the session.
+	TSharedPtr<IConcertClientSession> Session = WeakSessionPtr.Pin();
+	if (Session.IsValid())
+	{
+		// Remember the element selected in the list view to reselect it later.
+		TArray<TSharedPtr<FConcertSessionClientInfo>> SelectedItems = ClientsListView->GetSelectedItems();
+		checkf(SelectedItems.Num() <= 1, TEXT("ActiveSession's client list view should not support multiple selection."));
+
+		// If this is about a specific client update. (i.e. responding to a IConcertClientSession::OnSessionClientChanged() event)
+		if (InClientInfo)
+		{
+			if (ClientStatus == EConcertClientStatus::Connected)
+			{
+				Clients.Emplace(MakeShared<FConcertSessionClientInfo>(*InClientInfo));
+			}
+			else
+			{
+				int32 Index = Clients.IndexOfByPredicate([InClientInfo](const TSharedPtr<FConcertSessionClientInfo>& Visited) { return InClientInfo->ClientEndpointId == Visited->ClientEndpointId; });
+				if (Index != INDEX_NONE)
+				{
+					if (ClientStatus == EConcertClientStatus::Disconnected)
+					{
+						Clients.RemoveAt(Index); // We want to preserve items relative order.
+					}
+					else if (ensure(ClientStatus == EConcertClientStatus::Updated)) // Ensure we handled all status in EConcertClientStatus.
+					{
+						*Clients[Index] = *InClientInfo; // Update the client info.
+					}
+				}
+			}
+		}
+		else // Sync everything.
+		{
+			// Convert the list of clients to a list of shared pointer to client (the list view model asks for shared pointers).
+			TArray<FConcertSessionClientInfo> OtherConnectedClients = Session->GetSessionClients();
+			TArray<TSharedPtr<FConcertSessionClientInfo>> UpdatedClientList;
+			UpdatedClientList.Reserve(OtherConnectedClients.Num() + 1); // +1 is to include this local client (see below).
+			Algo::Transform(OtherConnectedClients, UpdatedClientList, [](FConcertSessionClientInfo InClient)
+			{
+				return MakeShared<FConcertSessionClientInfo>(MoveTemp(InClient));
+			});
+
+			// Add this local client as it is not part of the list returned by GetSessionClients(). (The client connected from this process).
+			if (ClientInfo.IsValid())
+			{
+				UpdatedClientList.Emplace(ClientInfo);
+			}
+
+			// Merge the list used by the list view (the model) with the updated list, removing clients who left and adding the one who joined.
+			ConcertFrontendUtils::SyncArraysByPredicate(Clients, MoveTemp(UpdatedClientList), [](const TSharedPtr<FConcertSessionClientInfo>& ClientToFind)
+			{
+				return [ClientToFind](const TSharedPtr<FConcertSessionClientInfo>& PotentialClient)
+				{
+					return PotentialClient->ClientEndpointId == ClientToFind->ClientEndpointId;
+				};
+			});
+		}
+
+		// Sort the list by display name alphabetically.
+		Clients.StableSort([](const TSharedPtr<FConcertSessionClientInfo>& ClientOne, const TSharedPtr<FConcertSessionClientInfo>& ClientTwo)
+		{
+			return ClientOne->ClientInfo.DisplayName < ClientTwo->ClientInfo.DisplayName;
+		});
+
+		// If a client row was selected, select it back (if still available).
+		if (SelectedItems.Num() > 0)
+		{
+			ClientsListView->SetSelection(SelectedItems[0]);
+		}
+	}
+	else // The session appears to be invalid.
+	{
+		// Clear the list of clients.
+		Clients.Reset();
+	}
+
+	ClientsListView->RequestListRefresh();
+}
+
+const FButtonStyle& SActiveSession::GetConnectionIconStyle() const
+{
+	EConcertActionType ButtonStyle = EConcertActionType::Danger;
+
+	TSharedPtr<IConcertClientSession> ClientSession = WeakSessionPtr.Pin();
+	if (ClientSession.IsValid())
+	{
+		if (ClientSession->GetConnectionStatus() == EConcertConnectionStatus::Connected)
+		{
+			const bool bIsDefault = ClientSession->GetSendReceiveState() == EConcertSendReceiveState::Default;
+			if (bIsDefault)
+			{
+				ButtonStyle = EConcertActionType::Success;
+			}
+			else
+			{
+				ButtonStyle = EConcertActionType::Warning;
+			}
+		}
+	}
+
+	return FAppStyle::Get().GetWidgetStyle<FButtonStyle>(ConcertClientFrontendUtils::ButtonStyleNames[(int32)ButtonStyle]);
+}
+
+FSlateColor SActiveSession::GetConnectionIconColor() const
+{
+	return GetConnectionIconStyle().Normal.TintColor;
+}
+
+FSlateFontInfo SActiveSession::GetConnectionIconFontInfo() const
+{
+	FSlateFontInfo ConnectionIconFontInfo = FAppStyle::Get().GetFontStyle(ConcertClientFrontendUtils::ButtonIconSyle);
+	ConnectionIconFontInfo.OutlineSettings.OutlineSize = 1;
+	ConnectionIconFontInfo.OutlineSettings.OutlineColor = GetConnectionIconStyle().Pressed.TintColor.GetSpecifiedColor();
+
+	return ConnectionIconFontInfo;
+}
+
+void SActiveSession::OnSendConflict(const FConcertConflictDescriptionBase& ConflictDescription)
+{
+	FAsyncTaskNotificationConfig Config;
+	Config.bKeepOpenOnFailure = true;
+	Config.TitleText = ConflictDescription.GetConflictTitle();
+	Config.ExpireDuration = 5.0f;
+
+	FAsyncTaskNotification Notification(Config);
+	Notification.SetHyperlink(FSimpleDelegate(), LOCTEXT("TransactionConflictDetails", "See Details..."));
+	Notification.SetComplete(Config.TitleText, ConflictDescription.GetConflictDetails(), false);
+}
+
+FText SActiveSession::GetConnectionStatusText() const
+{
+	FText StatusText = LOCTEXT("StatusDisconnected", "Disconnected");
+	TSharedPtr<IConcertClientSession> ClientSessionPtr = WeakSessionPtr.Pin();
+	if (ClientSessionPtr.IsValid() && ClientSessionPtr->GetConnectionStatus() == EConcertConnectionStatus::Connected)
+	{
+		const FText SessionDisplayName = FText::FromString(ClientSessionPtr->GetSessionInfo().SessionName);
+
+		switch (ClientSessionPtr->GetSendReceiveState())
+		{
+		case EConcertSendReceiveState::SendOnly:
+			StatusText = FText::Format(LOCTEXT("StatusSendOnlyFmt", "Send Only: {0}"), SessionDisplayName);
+			break;
+		case EConcertSendReceiveState::ReceiveOnly:
+			StatusText = FText::Format(LOCTEXT("StatusReceiveOnlyFmt", "Receive Only: {0}"), SessionDisplayName);
+			break;
+		case EConcertSendReceiveState::Default:
+			StatusText = FText::Format(LOCTEXT("StatusConnectedFmt", "Connected: {0}"), SessionDisplayName);
+			break;
+		};
+	}
+
+	return StatusText;
+}
+
+bool SActiveSession::IsStatusBarLeaveSessionVisible() const
+{
+	if (TSharedPtr<IConcertClientSession> ClientSession = WeakSessionPtr.Pin())
+	{
+		return ClientSession->GetConnectionStatus() == EConcertConnectionStatus::Connected;
+	}
+
+	return false;
+}
+
+FReply SActiveSession::OnClickLeaveSession()
+{
+	IMultiUserClientModule::Get().DisconnectSession();
+	return FReply::Handled();
+}
+
+void SActiveSession::SetSelectedClient(const FGuid& InClientEndpointId, ESelectInfo::Type SelectInfo)
+{
+	if (ClientsListView.IsValid())
+	{
+		TSharedPtr<FConcertSessionClientInfo> NewSelectedClient = FindAvailableClient(InClientEndpointId);
+
+		if (NewSelectedClient.IsValid())
+		{
+			ClientsListView->SetSelection(NewSelectedClient, SelectInfo);
+		}
+		else
+		{
+			ClientsListView->ClearSelection();
+		}
+	}
+}
+
+TSharedPtr<FConcertSessionClientInfo> SActiveSession::FindAvailableClient(const FGuid& InClientEndpointId) const
+{
+	const TSharedPtr<FConcertSessionClientInfo>* FoundClientPtr = Clients.FindByPredicate([&InClientEndpointId](const TSharedPtr<FConcertSessionClientInfo>& PotentialClient)
+	{
+		return PotentialClient->ClientEndpointId == InClientEndpointId;
+	});
+
+	return FoundClientPtr ? *FoundClientPtr : nullptr;
+}
+
+bool SActiveSession::CanProcessPendingPackages() const
+{
+	return bCanHotReload;
+}
+
+void SActiveSession::Tick(const FGeometry& AllottedGeometry, const double InCurrentTime, const float InDeltaTime)
+{
+	if (bCanHotReload && CanReloadDialog)
+	{
+		CanReloadDialog = {};
+	}
+}
+
+namespace UE::MultiUserClient::Private
+{
+static bool IsMapAsset(FName PackageName)
+{
+	FString FullPath;
+	if (FPackageName::TryConvertLongPackageNameToFilename(PackageName.ToString(), FullPath, FPackageName::GetMapPackageExtension()))
+	{
+		if (FPlatformFileManager::Get().GetPlatformFile().FileExists(*FullPath))
+		{
+			return true;
+		}
+	}
+	return false;
+}
+
+static TSharedPtr<SMessageDialog> MakeReloadDialog(FName PackageName)
+{
+	static const FText DialogTitle = LOCTEXT("ActiveSessionPendingHotReloadTitle", "A level hot reload is pending.");
+
+	static const FText ReloadButton = LOCTEXT("ActiveSessionPendingHotReloadButton", "Reload");
+
+	const FText ReloadMessage =
+		FText::Format(
+			LOCTEXT("ActiveSessionPendingHotReload","A hot reload on {0} is pending. Please dismiss this notification to continue with reload. All transactions will be queued until this occurs."),
+			 FText::FromName(PackageName)
+			);
+
+	return SNew(SMessageDialog)
+		.Title(DialogTitle)
+		.Message(ReloadMessage)
+		.Buttons( { SCustomDialog::FButton(ReloadButton) } );
+}
+
+}
+void SActiveSession::ActivityUpdated(const FConcertClientInfo& InClientInfo, const FConcertSyncActivity& InActivity, const FStructOnScope& ActivitySummary)
+{
+	if (InActivity.EventType == EConcertSyncActivityEventType::Package)
+	{
+		TSharedPtr<IConcertSyncClient> SyncClientPin = WeakConcertSyncClient.Pin();
+		if (!WeakSessionPtr.IsValid() || !SyncClientPin || CanReloadDialog)
+		{
+			return;
+		}
+
+		if (InActivity.EndpointId == WeakSessionPtr.Pin()->GetSessionClientEndpointId())
+		{
+			return;
+		}
+
+		check(ActivitySummary.GetStruct()->IsChildOf(FConcertSyncPackageActivitySummary::StaticStruct()));
+		const FConcertSyncPackageActivitySummary* ActivitySummaryPtr =
+			(FConcertSyncPackageActivitySummary*)ActivitySummary.GetStructMemory();
+
+		// User changed its settings.
+		const UConcertClientConfig* ConcertClientConfig = SyncClientPin->GetConcertClient()->GetConfiguration();
+		if (ConcertClientConfig->bShouldPromptForHotReloadOnLevel)
+		{
+			const bool bIsAMap = UE::MultiUserClient::Private::IsMapAsset(ActivitySummaryPtr->PackageName);
+			if (bIsAMap)
+			{
+				bCanHotReload = false;
+				CanReloadDialog = UE::MultiUserClient::Private::MakeReloadDialog(ActivitySummaryPtr->PackageName);
+				CanReloadDialog->GetOnWindowClosedEvent().AddLambda([this](const TSharedRef<SWindow>& ClosedWindow)
+					{
+						bCanHotReload = true;
+					});
+
+				CanReloadDialog->Show();
+			}
+		}
+	}
+}
+
+#undef LOCTEXT_NAMESPACE /* SActiveSession */
diff --git a/Engine/Plugins/Developer/Concert/ConcertApp/MultiUserClient/Source/MultiUserClient/Private/Widgets/SActiveSession.h b/Engine/Plugins/Developer/Concert/ConcertApp/MultiUserClient/Source/MultiUserClient/Private/Widgets/SActiveSession.h
new file mode 100644
--- /dev/null
+++ b/Engine/Plugins/Developer/Concert/ConcertApp/MultiUserClient/Source/MultiUserClient/Private/Widgets/SActiveSession.h
@@ -0,0 +1,167 @@
+// Copyright Epic Games, Inc. All Rights Reserved.
+
+#pragma once
+
+#include "CoreMinimal.h"
+#include "ClientSessionHistoryController.h"
+#include "Widgets/DeclarativeSyntaxSupport.h"
+#include "Widgets/SCompoundWidget.h"
+#include "Widgets/Views/SListView.h"
+#include "Widgets/Input/SComboBox.h"
+#include "ConcertMessages.h"
+
+class IConcertClientSession;
+class IConcertSyncClient;
+
+class FAsyncTaskNotification;
+struct FConcertConflictDescriptionBase;
+struct FConcertSessionClientInfo;
+
+class SCustomDialog;
+class SDockTab;
+class SExpandableArea;
+
+
+/**
+ * Displays the multi-users active session clients and activity, enables the client
+ * to leave the session, open the same level as another client or teleport to another
+ * client presence.
+ */
+class SActiveSession : public SCompoundWidget
+{
+public:
+	/** Struct to store the current send / receive state. */
+	struct FSendReceiveComboItem
+	{
+		FSendReceiveComboItem(FText InName, FText InToolTip, EConcertSendReceiveState InState) :
+			Name(MoveTemp(InName)), ToolTip(MoveTemp(InToolTip)), State(InState) {};
+
+		FText Name;
+		FText ToolTip;
+		EConcertSendReceiveState State;
+	};
+
+	SLATE_BEGIN_ARGS(SActiveSession) { }
+	SLATE_END_ARGS();
+
+	/**
+	* Constructs the active session tab.
+	*
+	* @param InArgs The Slate argument list.
+	*/
+	void Construct(const FArguments& InArgs, TSharedPtr<IConcertSyncClient> InConcertSyncClient);
+
+	virtual void Tick(const FGeometry& AllottedGeometry, const double InCurrentTime, const float InDeltaTime) override;
+
+private:
+
+	/** Generate a new client row */
+	TSharedRef<ITableRow> HandleGenerateRow(TSharedPtr<FConcertSessionClientInfo> InClientInfo, const TSharedRef<STableViewBase>& OwnerTable) const;
+
+	/** Handle a session startup */
+	void HandleSessionStartup(TSharedRef<IConcertClientSession> InClientSession);
+
+	/** Handle a session shutdown */
+	void HandleSessionShutdown(TSharedRef<IConcertClientSession> InClientSession);
+
+	/** Handle a session client change */
+	void HandleSessionClientChanged(IConcertClientSession&, EConcertClientStatus ClientStatus, const FConcertSessionClientInfo& ClientInfo);
+
+	/** Polls the local client info and detect if it changed in order to update its representation in real time.*/
+	EActiveTimerReturnType HandleLocalClientInfoChangePollingTimer(double InCurrentTime, float InDeltaTime);
+
+	/** Update the list of clients while keeping the alphabetical sorting */
+	void UpdateSessionClientListView(const FConcertSessionClientInfo* Client = nullptr, EConcertClientStatus Status = EConcertClientStatus::Updated);
+
+	/** Set the selected client in the clients list view*/
+	void SetSelectedClient(const FGuid& InClientEndpointId, ESelectInfo::Type SelectInfo = ESelectInfo::Direct);
+
+	/** Find a client with its endpoint id */
+	TSharedPtr<FConcertSessionClientInfo> FindAvailableClient(const FGuid& InClientEndpointId) const;
+
+	/** Handling for the status icon and text */
+	const FButtonStyle& GetConnectionIconStyle() const;
+	FSlateColor GetConnectionIconColor() const;
+	FSlateFontInfo GetConnectionIconFontInfo() const;
+	FText GetConnectionStatusText() const;
+
+	/** Handling for leave session button */
+	bool IsStatusBarLeaveSessionVisible() const;
+	FReply OnClickLeaveSession();
+
+	/** Handles how much space the 'Clients' area uses with respect to its expansion state. */
+	SSplitter::ESizeRule GetClientAreaSizeRule() const { return bClientAreaExpanded ? SSplitter::ESizeRule::FractionOfParent : SSplitter::ESizeRule::SizeToContent; }
+	void OnClientAreaExpansionChanged(bool bExpanded) { bClientAreaExpanded = bExpanded; }
+
+	/** Handles how much space the 'History' area uses with respect to its expansion state. */
+	SSplitter::ESizeRule GetHistoryAreaSizeRule() const { return bHistoryAreaExpanded ? SSplitter::ESizeRule::FractionOfParent : SSplitter::ESizeRule::SizeToContent; }
+	void OnHistoryAreaExpansionChanged(bool bExpanded) { bHistoryAreaExpanded = bExpanded; }
+
+	/** Delegate handler when a conflict between inbound transaction and those stored pending outbound transactions. */
+	void OnSendConflict(const FConcertConflictDescriptionBase& ConflictMsg);
+
+	/** Delegate handler on inbound transactions that focus specifically on packages. */
+	void OnPackageChangeActivity();
+
+	/** Delegate handler to indicate if we can hot reload an inbound package. */
+	bool CanProcessPendingPackages() const;
+
+	/** Delegate handler invoked when an activity has been received by multi-user. */
+	void ActivityUpdated(const FConcertClientInfo& InClientInfo, const FConcertSyncActivity& InActivity, const FStructOnScope& /*unused*/);
+
+private:
+
+	/** Get the text object for the send/receive combo box. */
+	FText GetRequestedSendReceiveComboText() const;
+
+	/** Generate the SWidget for the given item. */
+	TSharedRef<SWidget> GenerateSendReceiveComboItem(TSharedPtr<FSendReceiveComboItem> InItem);
+
+	/** Return the initially selected combo item. */
+	int32 GetInitialSendReceiveComboIndex();
+
+	/** Handle a send receiver state change from the combo box. */
+	void HandleSendReceiveChanged(TSharedPtr<FSendReceiveComboItem> Item, ESelectInfo::Type SelectInfo);
+
+	/** Pointer on the client sync. */
+	TWeakPtr<IConcertSyncClient> WeakConcertSyncClient;
+
+	/** Holds a concert client session. */
+	TWeakPtr<IConcertClientSession> WeakSessionPtr;
+
+	/** List view for AvailableClients. */
+	TSharedPtr<SListView<TSharedPtr<FConcertSessionClientInfo>>> ClientsListView;
+
+	/** List of clients for the current session. */
+	TArray<TSharedPtr<FConcertSessionClientInfo>> Clients;
+
+	/** Information about the machine's client. */
+	TSharedPtr<FConcertSessionClientInfo> ClientInfo;
+
+	/** Holds a concert activity log. */
+	TSharedPtr<FClientSessionHistoryController> SessionHistoryController;
+
+	/** The 'Clients' expandable area. */
+	TSharedPtr<SExpandableArea> ClientArea;
+
+	/** The 'History' expandable area. */
+	TSharedPtr<SExpandableArea> HistoryArea;
+
+	/** Holds a pointer to the preset combo box widget. */
+	TSharedPtr< SComboBox< TSharedPtr<FSendReceiveComboItem> > > SendReceiveComboBox;
+
+	/** Available states for the SendReceiveComboBox */
+	TArray< TSharedPtr< FSendReceiveComboItem > > SendReceiveComboList;
+
+	/** Notification handler for hot reload. */
+	TSharedPtr<SCustomDialog> CanReloadDialog;
+
+	/** Flag to indicate if it is OK to hotreload the packages. */
+	bool bCanHotReload = true;
+
+	/** Keeps the status of 'Clients' area expansion. */
+	bool bClientAreaExpanded = true;
+
+	/** Keep the status of 'History' area expansion. */
+	bool bHistoryAreaExpanded = true;
+};
diff --git a/Engine/Plugins/Developer/Concert/ConcertApp/MultiUserClient/Source/MultiUserClient/Private/Widgets/SConcertBrowser.cpp b/Engine/Plugins/Developer/Concert/ConcertApp/MultiUserClient/Source/MultiUserClient/Private/Widgets/SConcertBrowser.cpp
--- a/Engine/Plugins/Developer/Concert/ConcertApp/MultiUserClient/Source/MultiUserClient/Private/Widgets/SConcertBrowser.cpp
+++ b/Engine/Plugins/Developer/Concert/ConcertApp/MultiUserClient/Source/MultiUserClient/Private/Widgets/SConcertBrowser.cpp
@@ -5,8 +5,7 @@
 #include "IConcertClient.h"
 #include "IConcertSyncClient.h"
 #include "MultiUserClientUtils.h"
-#include "ActiveSession/SActiveSessionRoot.h"
-#include "Widgets/ActiveSession/Overview/SActiveSessionOverviewTab.h"
+#include "SActiveSession.h"
 #include "Widgets/Disconnected/SConcertClientSessionBrowser.h"
 #include "Widgets/Disconnected/SConcertNoAvailability.h"
 
@@ -18,7 +17,7 @@
 
 #define LOCTEXT_NAMESPACE "SConcertBrowser"
 
-void SConcertBrowser::Construct(const FArguments& InArgs, TSharedRef<SDockTab> InConstructUnderMajorTab, TWeakPtr<IConcertSyncClient> InSyncClient)
+void SConcertBrowser::Construct(const FArguments& InArgs, const TSharedRef<SDockTab>& ConstructUnderMajorTab, const TSharedPtr<SWindow>& ConstructUnderWindow, TWeakPtr<IConcertSyncClient> InSyncClient)
 {
 	if (!MultiUserClientUtils::HasServerCompatibleCommunicationPluginEnabled())
 	{
@@ -33,7 +32,6 @@ void SConcertBrowser::Construct(const FArguments& InArgs, TSharedRef<SDockTab> I
 	}
 
 	WeakConcertSyncClient = InSyncClient;
-	ConstructedUnderMajorTab = MoveTemp(InConstructUnderMajorTab);
 	if (TSharedPtr<IConcertSyncClient> ConcertSyncClient = WeakConcertSyncClient.Pin())
 	{
 		SearchedText = MakeShared<FText>(); // Will keep in memory the session browser search text between join/leave UI transitions.
@@ -59,7 +57,7 @@ void SConcertBrowser::AttachChildWidget(EConcertConnectionStatus ConnectionStatu
 	{
 		if (ConnectionStatus == EConcertConnectionStatus::Connected)
 		{
-			ChildSlot.AttachWidget(SNew(UE::MultiUserClient::SActiveSessionRoot, ConstructedUnderMajorTab.ToSharedRef(), ConcertSyncClient));
+			ChildSlot.AttachWidget(SNew(SActiveSession, ConcertSyncClient));
 		}
 		else if (ConnectionStatus == EConcertConnectionStatus::Disconnected)
 		{
diff --git a/Engine/Plugins/Developer/Concert/ConcertApp/MultiUserClient/Source/MultiUserClient/Private/Widgets/SConcertBrowser.h b/Engine/Plugins/Developer/Concert/ConcertApp/MultiUserClient/Source/MultiUserClient/Private/Widgets/SConcertBrowser.h
--- a/Engine/Plugins/Developer/Concert/ConcertApp/MultiUserClient/Source/MultiUserClient/Private/Widgets/SConcertBrowser.h
+++ b/Engine/Plugins/Developer/Concert/ConcertApp/MultiUserClient/Source/MultiUserClient/Private/Widgets/SConcertBrowser.h
@@ -28,26 +28,22 @@ public:
 	* Constructs the Browser.
 	*
 	* @param InArgs The Slate argument list.
-	* @param InConstructUnderMajorTab The major tab which will contain the session front-end.
-	* @param InConstructUnderWindow The window in which this widget is being constructed.
+	* @param ConstructUnderMajorTab The major tab which will contain the session front-end.
+	* @param ConstructUnderWindow The window in which this widget is being constructed.
 	* @param InSyncClient The sync client.
 	*/
-	void Construct(const FArguments& InArgs, TSharedRef<SDockTab> InConstructUnderMajorTab, TWeakPtr<IConcertSyncClient> InSyncClient);
+	void Construct(const FArguments& InArgs, const TSharedRef<SDockTab>& ConstructUnderMajorTab, const TSharedPtr<SWindow>& ConstructUnderWindow, TWeakPtr<IConcertSyncClient> InSyncClient);
 
 private:
+	/** Invoked when the session connection state is changed. */
+	void HandleSessionConnectionChanged(IConcertClientSession& InSession, EConcertConnectionStatus ConnectionStatus);
+
+	/** Attaches the child widgets according to the connection status. */
+	void AttachChildWidget(EConcertConnectionStatus ConnectionStatus);
 
 	/** Keeps the sync client interface. */
 	TWeakPtr<IConcertSyncClient> WeakConcertSyncClient;
 
-	/** Kept so it can be passed on to SActiveSessionRoot. */
-	TSharedPtr<SDockTab> ConstructedUnderMajorTab;
-
 	/** Keeps the session browser searched text in memory to reapply it when a user leaves a session and goes back to the session browser. */
 	TSharedPtr<FText> SearchedText;
-	
-	/** Invoked when the session connection state is changed. */
-	void HandleSessionConnectionChanged(IConcertClientSession& InSession, EConcertConnectionStatus ConnectionStatus);
-
-	/** Attaches the child widgets according to the connection status. */
-	void AttachChildWidget(EConcertConnectionStatus ConnectionStatus);
 };
diff --git a/Engine/Plugins/Developer/Concert/ConcertApp/MultiUserClient/Source/MultiUserReplication/MultiUserReplication.Build.cs b/Engine/Plugins/Developer/Concert/ConcertApp/MultiUserClient/Source/MultiUserReplication/MultiUserReplication.Build.cs
--- a/Engine/Plugins/Developer/Concert/ConcertApp/MultiUserClient/Source/MultiUserReplication/MultiUserReplication.Build.cs
+++ b/Engine/Plugins/Developer/Concert/ConcertApp/MultiUserClient/Source/MultiUserReplication/MultiUserReplication.Build.cs
@@ -20,8 +20,7 @@ namespace UnrealBuildTool.Rules
 				new string[]
 				{
 					"Concert",
-					"ConcertClient",
-					"ConcertSyncClient"
+					"ConcertClient"
 				}
 			);
 		}
diff --git a/Engine/Plugins/Developer/Concert/ConcertApp/MultiUserClient/Source/MultiUserReplication/Private/MultiUserReplicationClientProfileAsset.cpp b/Engine/Plugins/Developer/Concert/ConcertApp/MultiUserClient/Source/MultiUserReplication/Private/MultiUserReplicationClientProfileAsset.cpp
deleted file mode 100644
diff --git a/Engine/Plugins/Developer/Concert/ConcertApp/MultiUserClient/Source/MultiUserReplication/Private/MultiUserReplicationStreamAsset.cpp b/Engine/Plugins/Developer/Concert/ConcertApp/MultiUserClient/Source/MultiUserReplication/Private/MultiUserReplicationStreamAsset.cpp
--- a/Engine/Plugins/Developer/Concert/ConcertApp/MultiUserClient/Source/MultiUserReplication/Private/MultiUserReplicationStreamAsset.cpp
+++ b/Engine/Plugins/Developer/Concert/ConcertApp/MultiUserClient/Source/MultiUserReplication/Private/MultiUserReplicationStreamAsset.cpp
@@ -6,14 +6,10 @@ UMultiUserReplicationStreamAsset::UMultiUserReplicationStreamAsset()
 {
 	ReplicationList = CreateDefaultSubobject<UMultiUserPropertyReplicationSelection>(TEXT("ReplicationList"));
 	ReplicationList->SetFlags(RF_Transactional);
-
-	// To avoid issues with delta serialization the CDO is always 0
-	StreamId = HasAnyFlags(RF_ClassDefaultObject)
-		? FGuid{}
-		: FGuid::NewGuid();
 }
 
 FReplicationStreamDescription UMultiUserReplicationStreamAsset::GenerateDescription() const
 {
-	return { StreamId, ReplicationList->ReplicationMap };
+	// TODO:
+	return {};
 }
diff --git a/Engine/Plugins/Developer/Concert/ConcertApp/MultiUserClient/Source/MultiUserReplication/Public/MultiUserReplicationClientProfileAsset.h b/Engine/Plugins/Developer/Concert/ConcertApp/MultiUserClient/Source/MultiUserReplication/Public/MultiUserReplicationClientProfileAsset.h
--- a/Engine/Plugins/Developer/Concert/ConcertApp/MultiUserClient/Source/MultiUserReplication/Public/MultiUserReplicationClientProfileAsset.h
+++ b/Engine/Plugins/Developer/Concert/ConcertApp/MultiUserClient/Source/MultiUserReplication/Public/MultiUserReplicationClientProfileAsset.h
@@ -7,13 +7,6 @@
 #include "MultiUserReplicationStreamAsset.h"
 #include "MultiUserReplicationClientProfileAsset.generated.h"
 
-namespace UE::ConcertSyncClient::Replication
-{
-	struct FJoinReplicatedSessionArgs;
-}
-
-enum class EJoinReplicationErrorCode : uint8;
-
 /** Asset for users to describe a client to the server. */
 UCLASS()
 class MULTIUSERREPLICATION_API UMultiUserReplicationClientProfileAsset : public UObject
@@ -28,7 +21,4 @@ public:
 	/** The streams this client will send to the server. It takes authority over the listed properties. */
 	UPROPERTY(EditAnywhere, Category = "Replication")
 	TArray<TObjectPtr<UMultiUserReplicationStreamAsset>> Streams;
-
-	/** Converts this asset's data to the args expected by IConcertClientReplicationManager for joining. */
-	UE::ConcertSyncClient::Replication::FJoinReplicatedSessionArgs ToJoinArgs() const;
-};
+};
\ No newline at end of file
diff --git a/Engine/Plugins/Developer/Concert/ConcertApp/MultiUserClient/Source/MultiUserReplication/Public/MultiUserReplicationStreamAsset.h b/Engine/Plugins/Developer/Concert/ConcertApp/MultiUserClient/Source/MultiUserReplication/Public/MultiUserReplicationStreamAsset.h
--- a/Engine/Plugins/Developer/Concert/ConcertApp/MultiUserClient/Source/MultiUserReplication/Public/MultiUserReplicationStreamAsset.h
+++ b/Engine/Plugins/Developer/Concert/ConcertApp/MultiUserClient/Source/MultiUserReplication/Public/MultiUserReplicationStreamAsset.h
@@ -15,10 +15,6 @@ class MULTIUSERREPLICATION_API UMultiUserReplicationStreamAsset : public UObject
 	GENERATED_BODY()
 public:
 
-	/** Unique ID for this stream. CDO has value 0 to not cause issues with delta serialization. */
-	UPROPERTY()
-	FGuid StreamId;
-	
 	UPROPERTY(Instanced)
 	TObjectPtr<UMultiUserPropertyReplicationSelection> ReplicationList;
 
diff --git a/Engine/Plugins/Developer/Concert/ConcertApp/MultiUserClient/Source/MultiUserReplicationEditor/MultiUserReplicationEditor.Build.cs b/Engine/Plugins/Developer/Concert/ConcertApp/MultiUserClient/Source/MultiUserReplicationEditor/MultiUserReplicationEditor.Build.cs
--- a/Engine/Plugins/Developer/Concert/ConcertApp/MultiUserClient/Source/MultiUserReplicationEditor/MultiUserReplicationEditor.Build.cs
+++ b/Engine/Plugins/Developer/Concert/ConcertApp/MultiUserClient/Source/MultiUserReplicationEditor/MultiUserReplicationEditor.Build.cs
@@ -37,7 +37,8 @@ namespace UnrealBuildTool.Rules
 					// Concert
 					"Concert",
 					"ConcertClient",
-					"ConcertSharedSlate"
+					"ConcertSharedSlate",
+					"MultiUserClient"
 				}
 			);
 		}
diff --git a/Engine/Plugins/Developer/Concert/ConcertApp/MultiUserClient/Source/MultiUserReplicationEditor/Private/MultiUserReplicationEditorModule.cpp b/Engine/Plugins/Developer/Concert/ConcertApp/MultiUserClient/Source/MultiUserReplicationEditor/Private/MultiUserReplicationEditorModule.cpp
--- a/Engine/Plugins/Developer/Concert/ConcertApp/MultiUserClient/Source/MultiUserReplicationEditor/Private/MultiUserReplicationEditorModule.cpp
+++ b/Engine/Plugins/Developer/Concert/ConcertApp/MultiUserClient/Source/MultiUserReplicationEditor/Private/MultiUserReplicationEditorModule.cpp
@@ -3,6 +3,7 @@
 #include "MultiUserReplicationEditorModule.h"
 
 #include "ConcertSyncSessionFlags.h"
+#include "IMultiUserClientModule.h"
 
 #include "ISettingsModule.h"
 #include "MultiUserReplicationEditorStyle.h"
@@ -29,18 +30,12 @@ namespace UE::MultiUserReplicationEditor
 		UnregisterSettings();
 	}
 
-	IMultiUserReplicationEditorModule::FSettingPath FMultiUserReplicationEditorModule::GetReplicationSettingsInfo() const
-	{
-		return { "Project", "Plugins", "Multi-User Replication" };
-	}
-
 	void FMultiUserReplicationEditorModule::RegisterSettings()
 	{
 		ISettingsModule* SettingsModule = FModuleManager::GetModulePtr<ISettingsModule>("Settings");
 		if (ensure(SettingsModule))
 		{
-			const FSettingPath Path = GetReplicationSettingsInfo();
-			SettingsSection = SettingsModule->RegisterSettings(Path.ContainerName, Path.CategoryName, Path.SectionName,
+			SettingsModule->RegisterSettings("Project", "Plugins", "Multi-User Replication",
 				LOCTEXT("MultiUserReplicationSettingsName", "Multi-User Replication"),
 				LOCTEXT("MultiUserReplicationSettingsDescription", "Configure the Multi-User Replication settings."),
 				UMultiUserReplicationSettings::Get()
@@ -53,11 +48,8 @@ namespace UE::MultiUserReplicationEditor
 		ISettingsModule* SettingsModule = FModuleManager::GetModulePtr<ISettingsModule>("Settings");
 		if (SettingsModule)
 		{
-			const FSettingPath Path = GetReplicationSettingsInfo();
-			SettingsModule->UnregisterSettings(Path.ContainerName, Path.CategoryName, Path.SectionName);
+			SettingsModule->UnregisterSettings("Project", "Plugins", "Multi-User Replication");
 		}
-
-		SettingsSection = nullptr;
 	}
 };
 
diff --git a/Engine/Plugins/Developer/Concert/ConcertApp/MultiUserClient/Source/MultiUserReplicationEditor/Private/MultiUserReplicationEditorModule.h b/Engine/Plugins/Developer/Concert/ConcertApp/MultiUserClient/Source/MultiUserReplicationEditor/Private/MultiUserReplicationEditorModule.h
--- a/Engine/Plugins/Developer/Concert/ConcertApp/MultiUserClient/Source/MultiUserReplicationEditor/Private/MultiUserReplicationEditorModule.h
+++ b/Engine/Plugins/Developer/Concert/ConcertApp/MultiUserClient/Source/MultiUserReplicationEditor/Private/MultiUserReplicationEditorModule.h
@@ -20,15 +20,12 @@ namespace UE::MultiUserReplicationEditor
 
 		//~ Begin IMultiUserReplicationEditor Interface
 		virtual FAssetCategoryPath GetMultiUserReplicationCategory() const override { return MultiUserReplicationCategory; }
-		virtual FSettingPath GetReplicationSettingsInfo() const override;
 		//~ End IMultiUserReplicationEditor Interface
 
 	private:
 
 		FAssetCategoryPath MultiUserReplicationCategory;
 
-		TSharedPtr<ISettingsSection> SettingsSection;
-
 		void RegisterSettings();
 		void UnregisterSettings();
 	};
diff --git a/Engine/Plugins/Developer/Concert/ConcertApp/MultiUserClient/Source/MultiUserReplicationEditor/Public/IMultiUserReplicationEditorModule.h b/Engine/Plugins/Developer/Concert/ConcertApp/MultiUserClient/Source/MultiUserReplicationEditor/Public/IMultiUserReplicationEditorModule.h
--- a/Engine/Plugins/Developer/Concert/ConcertApp/MultiUserClient/Source/MultiUserReplicationEditor/Public/IMultiUserReplicationEditorModule.h
+++ b/Engine/Plugins/Developer/Concert/ConcertApp/MultiUserClient/Source/MultiUserReplicationEditor/Public/IMultiUserReplicationEditorModule.h
@@ -6,7 +6,6 @@
 #include "Modules/ModuleManager.h"
 #include "UObject/NameTypes.h"
 
-class ISettingsSection;
 struct FAssetCategoryPath;
 
 namespace UE::MultiUserReplicationEditor
@@ -40,14 +39,5 @@ namespace UE::MultiUserReplicationEditor
 
 		/** Gets the category under which assets for Multi-User Replication are nested. */
 		virtual FAssetCategoryPath GetMultiUserReplicationCategory() const = 0;
-
-		struct FSettingPath
-		{
-			const FName ContainerName;
-			const FName CategoryName;
-			const FName SectionName;
-		};
-		/** Gets the sections the replication settings are stored under. */
-		virtual FSettingPath GetReplicationSettingsInfo() const = 0;
 	};
 }
\ No newline at end of file
diff --git a/Engine/Plugins/Developer/Concert/ConcertApp/MultiUserClient/Source/MultiUserReplicationEditor/Public/Settings/MultiUserReplicationSettings.h b/Engine/Plugins/Developer/Concert/ConcertApp/MultiUserClient/Source/MultiUserReplicationEditor/Public/Settings/MultiUserReplicationSettings.h
--- a/Engine/Plugins/Developer/Concert/ConcertApp/MultiUserClient/Source/MultiUserReplicationEditor/Public/Settings/MultiUserReplicationSettings.h
+++ b/Engine/Plugins/Developer/Concert/ConcertApp/MultiUserClient/Source/MultiUserReplicationEditor/Public/Settings/MultiUserReplicationSettings.h
@@ -4,9 +4,11 @@
 
 #include "UObject/Object.h"
 #include "MultiUserReplicationClientProfileAsset.h"
+#if WITH_EDITORONLY_DATA
 #include "Settings/MultiUserDefaultPropertySelection.h"
 #include "Settings/MultiUserDefaultSubobjectSelection.h"
 #include "Templates/Function.h"
+#endif
 #include "MultiUserReplicationSettings.generated.h"
 
 // Architecturally speaking FDefaultReplicationSessionSettings should be in the MultiUserClient module.
@@ -37,6 +39,7 @@ public:
 	UPROPERTY(EditAnywhere, Config, Category = "Replication")
 	FDefaultReplicationSessionSettings DefaultSessionSettings;
 
+#if WITH_EDITORONLY_DATA
 	/** Properties you want selected by default when you add a new replicated object in the editor */
 	UPROPERTY(EditAnywhere, Config, Category = "Replication|Editor")
 	TMap<FSoftClassPath, FMultiUserDefaultPropertySelection> DefaultPropertySelection;
@@ -49,12 +52,15 @@ public:
 	 */
 	UPROPERTY(EditAnywhere, Config, Category = "Replication|Editor")
 	TMap<FSoftClassPath, FMultiUserDefaultSubobjectSelection> DefaultSubobjectSelection;
+#endif
 
 	static UMultiUserReplicationSettings* Get() { return GetMutableDefault<UMultiUserReplicationSettings>(); }
 
+#if WITH_EDITORONLY_DATA
 	/** Reads DefaultPropertySelection and applies any default property selections to Info based on the Class just added. */
 	void AddDefaultPropertiesFromSettings(FReplicatedObjectInfo& Info, UClass& Class);
 
 	/** Reads DefaultComponentSelection and calls FurtherObjectsCallback on any further objects that should also be added. */
 	void AddAdditionalObjectsFromSettings(UObject& AddedObject, TFunctionRef<void(UObject&)> FurtherObjectsCallback);
+#endif
 };
\ No newline at end of file
