[Concert] Adding required QuicServer command line arguments when starting MultiUserSlateServer

diff --git a/Engine/Plugins/Developer/Concert/ConcertApp/MultiUserClient/Source/MultiUserClient/Private/MultiUserClientModule.cpp b/Engine/Plugins/Developer/Concert/ConcertApp/MultiUserClient/Source/MultiUserClient/Private/MultiUserClientModule.cpp
--- a/Engine/Plugins/Developer/Concert/ConcertApp/MultiUserClient/Source/MultiUserClient/Private/MultiUserClientModule.cpp
+++ b/Engine/Plugins/Developer/Concert/ConcertApp/MultiUserClient/Source/MultiUserClient/Private/MultiUserClientModule.cpp
@@ -1318,7 +1318,7 @@ private:
 private:
 	FString GenerateConcertServerCommandLine(TOptional<FServerLaunchOverrides> LaunchOverrides) const
 	{
-		// @note FH: Those settings are UDP Messaging specific
+		// @note: Those settings are UDP/QUIC Messaging specific
 
 		FString CmdLine;
 		FConfigFile* EngineConfig = GConfig ? GConfig->FindConfigFileWithBaseName(FName(TEXT("Engine"))) : nullptr;
@@ -1365,6 +1365,11 @@ private:
 				}
 			}
 
+			// Set the Quic unicast endpoint
+			FString QuicUnicastEndpoint;
+			EngineConfig->GetString(TEXT("/Script/QuicMessaging.QuicMessagingSettings"), TEXT("UnicastEndpoint"), QuicUnicastEndpoint);
+			CmdLine += TEXT(" -QUICMESSAGING_TRANSPORT_UNICAST=") + QuicUnicastEndpoint;
+
 			if (LaunchOverrides)
 			{
 				if (!LaunchOverrides->ServerName.IsEmpty())
@@ -1373,6 +1378,10 @@ private:
 				}
 			}
 		}
+
+		// Set Quic to operate in server mode
+		CmdLine += TEXT(" -QUICMESSAGING_TRANSPORT_ISCLIENT=False");
+
 		return CmdLine;
 	}
 
