Check for engine exit before trying to rebuild the session browser.

diff --git a/Engine/Plugins/Developer/Concert/ConcertApp/MultiUserClient/Source/MultiUserClient/Private/Widgets/Disconnected/SConcertClientSessionBrowser.cpp b/Engine/Plugins/Developer/Concert/ConcertApp/MultiUserClient/Source/MultiUserClient/Private/Widgets/Disconnected/SConcertClientSessionBrowser.cpp
--- a/Engine/Plugins/Developer/Concert/ConcertApp/MultiUserClient/Source/MultiUserClient/Private/Widgets/Disconnected/SConcertClientSessionBrowser.cpp
+++ b/Engine/Plugins/Developer/Concert/ConcertApp/MultiUserClient/Source/MultiUserClient/Private/Widgets/Disconnected/SConcertClientSessionBrowser.cpp
@@ -45,7 +45,7 @@
 
 void SConcertClientSessionBrowser::Construct(const FArguments& InArgs, IConcertClientPtr InConcertClient, TSharedPtr<FText> InSearchText)
 {
-	if (!InConcertClient.IsValid())
+	if (!InConcertClient.IsValid() || IsEngineExitRequested())
 	{
 		return; // Don't build the UI if ConcertClient is not available.
 	}
@@ -791,4 +791,4 @@ void SConcertClientSessionBrowser::Tick(const FGeometry& AllottedGeometry, const
 	}
 }
 
-#undef LOCTEXT_NAMESPACE
\ No newline at end of file
+#undef LOCTEXT_NAMESPACE
