Move discovery process to a task thread.

diff --git a/Engine/Plugins/Developer/Concert/ConcertApp/MultiUserClient/Source/MultiUserClient/Private/Widgets/Disconnected/SConcertClientSessionBrowser.cpp b/Engine/Plugins/Developer/Concert/ConcertApp/MultiUserClient/Source/MultiUserClient/Private/Widgets/Disconnected/SConcertClientSessionBrowser.cpp
--- a/Engine/Plugins/Developer/Concert/ConcertApp/MultiUserClient/Source/MultiUserClient/Private/Widgets/Disconnected/SConcertClientSessionBrowser.cpp
+++ b/Engine/Plugins/Developer/Concert/ConcertApp/MultiUserClient/Source/MultiUserClient/Private/Widgets/Disconnected/SConcertClientSessionBrowser.cpp
@@ -116,7 +116,8 @@ void SConcertClientSessionBrowser::Construct(const FArguments& InArgs, IConcertC
 	// Create a timer to periodically poll the server for sessions and session clients at a lower frequency than the normal tick.
 	RegisterActiveTimer(1.0f, FWidgetActiveTimerDelegate::CreateSP(this, &SConcertClientSessionBrowser::TickDiscovery));
 
-	bLocalServerRunning = IMultiUserClientModule::Get().IsConcertServerRunning();
+	const bool bForceTaskStart = true;
+	ScheduleDiscoveryTaskIfPossible(bForceTaskStart);
 }
 
 TSharedRef<SWidget> SConcertClientSessionBrowser::MakeBrowserContent(TSharedPtr<FText> InSearchText)
@@ -344,11 +345,21 @@ TSharedRef<SWidget> SConcertClientSessionBrowser::MakeOverlayedTableView(const T
 		];
 }
 
-EActiveTimerReturnType SConcertClientSessionBrowser::TickDiscovery(double InCurrentTime, float InDeltaTime)
+void SConcertClientSessionBrowser::ScheduleDiscoveryTaskIfPossible(bool bForceTaskStart)
 {
-	// Cache the result of this function because it is very expensive. It kills the framerate if polled every frame.
-	bLocalServerRunning = IMultiUserClientModule::Get().IsConcertServerRunning();
+	if (bForceTaskStart || DiscoveryTask.IsCompleted())
+	{
+		bLocalServerRunning = bForceTaskStart ? false : DiscoveryTask.GetResult();
+		DiscoveryTask = UE::Tasks::Launch( UE_SOURCE_LOCATION, []
+		{
+			return IMultiUserClientModule::Get().IsConcertServerRunning();
+		});
+	}
+}
 
+EActiveTimerReturnType SConcertClientSessionBrowser::TickDiscovery(double InCurrentTime, float InDeltaTime)
+{
+	ScheduleDiscoveryTaskIfPossible(false);
 	UpdateDiscovery();
 	return EActiveTimerReturnType::Continue;
 }
@@ -724,7 +735,7 @@ FReply SConcertClientSessionBrowser::OnNewButtonClicked()
 FReply SConcertClientSessionBrowser::OnLaunchServerButtonClicked()
 {
 	IMultiUserClientModule::Get().LaunchConcertServer();
-	bLocalServerRunning = IMultiUserClientModule::Get().IsConcertServerRunning(); // Immediately update the cache state to avoid showing buttons enabled for a split second.
+	ScheduleDiscoveryTaskIfPossible(false);
 	return FReply::Handled();
 }
 
diff --git a/Engine/Plugins/Developer/Concert/ConcertApp/MultiUserClient/Source/MultiUserClient/Private/Widgets/Disconnected/SConcertClientSessionBrowser.h b/Engine/Plugins/Developer/Concert/ConcertApp/MultiUserClient/Source/MultiUserClient/Private/Widgets/Disconnected/SConcertClientSessionBrowser.h
--- a/Engine/Plugins/Developer/Concert/ConcertApp/MultiUserClient/Source/MultiUserClient/Private/Widgets/Disconnected/SConcertClientSessionBrowser.h
+++ b/Engine/Plugins/Developer/Concert/ConcertApp/MultiUserClient/Source/MultiUserClient/Private/Widgets/Disconnected/SConcertClientSessionBrowser.h
@@ -3,6 +3,7 @@
 #pragma once
 
 #include "IConcertClientModule.h"
+#include "Tasks/Task.h"
 #include "Widgets/DeclarativeSyntaxSupport.h"
 #include "Widgets/SCompoundWidget.h"
 
@@ -104,7 +105,9 @@ private:
 	bool ConfirmDeleteSessionWithDialog(const TArray<TSharedPtr<FConcertSessionTreeItem>>& SessionItems) const;
 
 	// Update server/session/clients lists.
+	void ScheduleDiscoveryTaskIfPossible(bool bForceTaskStart);
 	EActiveTimerReturnType TickDiscovery(double InCurrentTime, float InDeltaTime);
 	void UpdateDiscovery();
 	void RefreshClientList(const TArray<FConcertSessionClientInfo>& LastestClientList);
+	UE::Tasks::TTask<bool> DiscoveryTask;
 };
