Allow MU clients to change their replicated properties in the MU active session view

diff --git a/Engine/Plugins/Developer/Concert/ConcertApp/MultiUserClient/Source/MultiUserClient/Private/Replication/Authority/AuthorityPolicy.cpp b/Engine/Plugins/Developer/Concert/ConcertApp/MultiUserClient/Source/MultiUserClient/Private/Replication/Authority/AuthorityPolicy.cpp
new file mode 100644
--- /dev/null
+++ b/Engine/Plugins/Developer/Concert/ConcertApp/MultiUserClient/Source/MultiUserClient/Private/Replication/Authority/AuthorityPolicy.cpp
@@ -0,0 +1,83 @@
+// Copyright Epic Games, Inc. All Rights Reserved.
+
+#include "AuthorityPolicy.h"
+
+#include "Replication/Stream/LocalClientStreamSynchronizer.h"
+
+#include "IConcertSyncClient.h"
+
+#include "Templates/Function.h"
+
+namespace UE::MultiUserClient
+{
+	namespace Private
+	{
+		static void ForEachObjectAddedByRequest(
+			const FObjectReplicationMap& OldState,
+			const ConcertSyncClient::Replication::FChangeStreamRequest& AcceptedRequest,
+			TFunctionRef<void(const FObjectInStreamID& Object)> Callback
+			)
+		{
+			// Call for every added stream
+			for (const FReplicationStreamDescription_NetPacked& AddedStream : AcceptedRequest.StreamsToAdd)
+			{
+				for (const TPair<FSoftObjectPath, FReplicatedObjectInfo>& Pair : AddedStream.BaseDescription.ReplicationMap.ReplicatedObjects)
+				{
+					Callback({ AddedStream.BaseDescription.Identifier , Pair.Key });
+				}
+			}
+
+			// Call for put object requests 
+			for (const TPair<FObjectInStreamID, FConcertChangeStream_PutObject>& Pair : AcceptedRequest.ObjectsToPut)
+			{
+				const bool bObjectWasAddedWithPut = !OldState.ReplicatedObjects.Contains(Pair.Key.Object);
+				if (bObjectWasAddedWithPut)
+				{
+					Callback(Pair.Key);
+				}
+			}
+		}
+	}
+	
+	FAuthorityPolicy::FAuthorityPolicy(FLocalClientStreamSynchronizer& InLocalClientManager, TSharedRef<IConcertSyncClient> InClient)
+		: LocalClientManager(InLocalClientManager)
+		, Client(MoveTemp(InClient))
+	{
+		// After submitting new objects, immediately request authority over it.
+		LocalClientManager.OnChangesAccepted_AnyThread().AddRaw(this, &FAuthorityPolicy::OnChangesAccepted_AnyThread);
+	}
+
+	FAuthorityPolicy::~FAuthorityPolicy()
+	{
+		LocalClientManager.OnChangesAccepted_AnyThread().RemoveAll(this);
+	}
+
+	void FAuthorityPolicy::OnChangesAccepted_AnyThread(
+		const FObjectReplicationMap& OldState,
+		const ConcertSyncClient::Replication::FChangeStreamRequest& AcceptedRequest
+		)
+	{
+		IConcertClientReplicationManager* ReplicationManager = Client->GetReplicationManager();
+		if (!ensure(ReplicationManager))
+		{
+			return;
+		}
+
+		using namespace ConcertSyncClient::Replication;
+		FAuthorityChangeRequest AuthorityChangeRequest;
+		Private::ForEachObjectAddedByRequest(OldState, AcceptedRequest, [&AuthorityChangeRequest](const FObjectInStreamID& Object)
+		{
+			AuthorityChangeRequest.TakeAuthority.FindOrAdd(Object.Object).StreamIds.AddUnique(Object.StreamId);
+		});
+		
+		OnAuthorityRequestSentDelegate.Broadcast(AuthorityChangeRequest);
+		ReplicationManager->RequestAuthorityChange(AuthorityChangeRequest)
+			.Next([this, DestructionDetection = LifetimeToken->AsWeak(), AuthorityChangeRequest](FAuthorityChangeResponse&& Response)
+			{
+				if (DestructionDetection.IsValid())
+				{
+					OnAuthorityResponseReceivedDelegate.Broadcast(AuthorityChangeRequest, Response);
+				}
+			});
+	}
+}
diff --git a/Engine/Plugins/Developer/Concert/ConcertApp/MultiUserClient/Source/MultiUserClient/Private/Replication/Authority/AuthorityPolicy.h b/Engine/Plugins/Developer/Concert/ConcertApp/MultiUserClient/Source/MultiUserClient/Private/Replication/Authority/AuthorityPolicy.h
new file mode 100644
--- /dev/null
+++ b/Engine/Plugins/Developer/Concert/ConcertApp/MultiUserClient/Source/MultiUserClient/Private/Replication/Authority/AuthorityPolicy.h
@@ -0,0 +1,69 @@
+// Copyright Epic Games, Inc. All Rights Reserved.
+
+#pragma once
+
+#include "Replication/IConcertClientReplicationManager.h"
+#include "Replication/IToken.h"
+#include "Replication/Data/ObjectReplicationMap.h"
+#include "Templates/UnrealTemplate.h"
+
+class IConcertSyncClient;
+
+namespace UE::ConcertSyncClient::Replication
+{
+	struct FAuthorityChangeResponse;
+	struct FAuthorityChangeRequest;
+}
+namespace UE::MultiUserClient
+{
+	class FLocalClientStreamSynchronizer;
+}
+
+namespace UE::MultiUserClient
+{
+	/**
+	 * Manages authority for the local client in the Multi User application.
+	 * Automatically requests authority over newly added objects.
+	 *
+	 * FMultiUserReplicationManager ensures that the passed in arguments, like LocalClientManager, are kept alive for
+	 * the duration of FAuthorityPolicy's life.
+	 */
+	class FAuthorityPolicy : public FNoncopyable
+	{
+	public:
+		
+		FAuthorityPolicy(FLocalClientStreamSynchronizer& InLocalClientManager, TSharedRef<IConcertSyncClient> InClient);
+		~FAuthorityPolicy();
+
+		DECLARE_MULTICAST_DELEGATE_OneParam(FOnAuthorityRequestSent,
+			const ConcertSyncClient::Replication::FAuthorityChangeRequest& Request
+			);
+		/** Called when an authority request is sent to the server. */
+		FOnAuthorityRequestSent& OnAuthorityRequestSent_AnyThread() { return OnAuthorityRequestSentDelegate; }
+
+		DECLARE_MULTICAST_DELEGATE_TwoParams(FOnAuthorityResponseReceived,
+			const ConcertSyncClient::Replication::FAuthorityChangeRequest& Request,
+			const ConcertSyncClient::Replication::FAuthorityChangeResponse& Response
+			);
+		/** Called when an authority response is received from the server. */
+		FOnAuthorityResponseReceived& OnAuthorityResponseReceived_AnyThread() { return OnAuthorityResponseReceivedDelegate; }
+	
+	private:
+
+		/** Referenced by the authority requests to detect destruction of FAuthorityPolicy */
+		const TSharedRef<FToken> LifetimeToken = FToken::Make();
+
+		/** Informs us when changes have been successfully submitted to the server. */
+		FLocalClientStreamSynchronizer& LocalClientManager;
+
+		/** Used to send authority requests to the server. */
+		const TSharedRef<IConcertSyncClient> Client;
+
+		FOnAuthorityRequestSent OnAuthorityRequestSentDelegate;
+		FOnAuthorityResponseReceived OnAuthorityResponseReceivedDelegate;
+
+		/** After changes have been accepted by the server, requests authority over said objects & properties. */
+		void OnChangesAccepted_AnyThread(const FObjectReplicationMap& OldState, const ConcertSyncClient::Replication::FChangeStreamRequest& AcceptedRequest);
+	};
+}
+
diff --git a/Engine/Plugins/Developer/Concert/ConcertApp/MultiUserClient/Source/MultiUserClient/Private/Replication/IToken.h b/Engine/Plugins/Developer/Concert/ConcertApp/MultiUserClient/Source/MultiUserClient/Private/Replication/IToken.h
new file mode 100644
--- /dev/null
+++ b/Engine/Plugins/Developer/Concert/ConcertApp/MultiUserClient/Source/MultiUserClient/Private/Replication/IToken.h
@@ -0,0 +1,25 @@
+// Copyright Epic Games, Inc. All Rights Reserved.
+
+#pragma once
+
+#include "Templates/SharedPointer.h"
+
+namespace UE::MultiUserClient
+{
+	/**
+	 * The token "pattern" used by systems that start non-cancellable, latent operations (like TFutures from Concert responses):
+	 * 1. The system creates a TSharedRef<FToken> and holds on to it for the system's entire lifetime.
+	 * 2. Latent operation capture a TWeakPtr to the token and the this pointer of the system
+	 * 3. When the latent operation finishes, it checks whether the token is still valid.
+	 *	- If yes, that means the system is valid, too, so the captured this can be used safely.
+	 *	- If no, discard the operation result.
+	 *
+	 * This pattern offers an alternative to forcing the system to be wrapped in a shared pointer itself and forcing it to inherit from TSharedFromThis.
+	 * Often systems should not care about whether they are heap allocated or not - in such cases using TSharedFromThis would be an anti-pattern.
+	 */
+	class FToken : public TSharedFromThis<FToken>
+	{
+	public:
+		static TSharedRef<FToken> Make() { return MakeShared<FToken>(); }
+	};
+}
\ No newline at end of file
diff --git a/Engine/Plugins/Developer/Concert/ConcertApp/MultiUserClient/Source/MultiUserClient/Private/Replication/MultiUserReplicationManager.cpp b/Engine/Plugins/Developer/Concert/ConcertApp/MultiUserClient/Source/MultiUserClient/Private/Replication/MultiUserReplicationManager.cpp
--- a/Engine/Plugins/Developer/Concert/ConcertApp/MultiUserClient/Source/MultiUserClient/Private/Replication/MultiUserReplicationManager.cpp
+++ b/Engine/Plugins/Developer/Concert/ConcertApp/MultiUserClient/Source/MultiUserClient/Private/Replication/MultiUserReplicationManager.cpp
@@ -3,16 +3,15 @@
 #include "MultiUserReplicationManager.h"
 
 #include "IConcertSyncClient.h"
-#include "ReplicationUtils.h"
+#include "Replication/IConcertClientReplicationManager.h"
 
+#include "Containers/Ticker.h"
 #include "UObject/Package.h"
 
 namespace UE::MultiUserClient
 {
 	FMultiUserReplicationManager::FMultiUserReplicationManager(TSharedRef<IConcertSyncClient> InClient)
 		: Client(MoveTemp(InClient))
-		, SessionContent(NewObject<UMultiUserReplicationSessionPreset>(GetTransientPackage(), NAME_None, RF_Transient))
-		, LocalClientContent(SessionContent->AddClient())
 	{
 		Client->GetConcertClient()->OnSessionConnectionChanged().AddRaw(
 			this,
@@ -25,10 +24,30 @@ namespace UE::MultiUserClient
 		Client->GetConcertClient()->OnSessionConnectionChanged().RemoveAll(this);
 	}
 
-	void FMultiUserReplicationManager::AddReferencedObjects(FReferenceCollector& Collector)
+	void FMultiUserReplicationManager::JoinReplicationSession()
 	{
-		Collector.AddReferencedObject(SessionContent);
-		Collector.AddReferencedObject(LocalClientContent);
+		IConcertClientReplicationManager* Manager = Client->GetReplicationManager();
+		if (!ensure(ConnectionState == EMultiUserReplicationConnectionState::Disconnected)
+			|| !ensure(Manager))
+		{
+			return;
+		}
+
+		ConnectionState = EMultiUserReplicationConnectionState::Connecting;
+		// For now we join without any initial data - this will likely change in the future (5.5+)
+		Manager->JoinReplicationSession({})
+			.Next([WeakThis = AsWeak()](ConcertSyncClient::Replication::FJoinReplicatedSessionResult&& JoinSessionResult)
+			{
+				// The future can execute on any thread
+				ExecuteOnGameThread(TEXT("JoinReplicationSession"), [WeakThis, JoinSessionResult = MoveTemp(JoinSessionResult)]()
+				{
+					// Shutting down engine?
+					if (const TSharedPtr<FMultiUserReplicationManager> ThisPin = WeakThis.Pin())
+					{
+						ThisPin->HandleReplicationSessionJoined(JoinSessionResult);
+					}
+				});
+			});
 	}
 
 	void FMultiUserReplicationManager::OnSessionConnectionChanged(
@@ -41,7 +60,7 @@ namespace UE::MultiUserClient
 		case EConcertConnectionStatus::Connecting:
 			break;
 		case EConcertConnectionStatus::Connected:
-			OnJoinSession(ConcertClientSession);
+			JoinReplicationSession();
 			break;
 		case EConcertConnectionStatus::Disconnecting:
 			break;
@@ -52,24 +71,36 @@ namespace UE::MultiUserClient
 		}
 	}
 
-	void FMultiUserReplicationManager::OnJoinSession(IConcertClientSession& ConcertClientSession)
+	void FMultiUserReplicationManager::OnLeaveSession(IConcertClientSession& ConcertClientSession)
 	{
-		Replication::JoinSessionForMultiUser(Client, {});
+		// Keep in mind that FClientStreamRepository::GetLocalClientEditModel is referenced by the UI ...
+		ConnectedState.Reset();
+		// ... and after this broadcast the model should no longer be referenced by anyone
+		SetConnectionStateAndBroadcast(EMultiUserReplicationConnectionState::Disconnected);
 	}
 
-	void FMultiUserReplicationManager::OnLeaveSession(IConcertClientSession& ConcertClientSession)
+	void FMultiUserReplicationManager::HandleReplicationSessionJoined(const ConcertSyncClient::Replication::FJoinReplicatedSessionResult& JoinSessionResult)
 	{
-		ClearSessionData();
-		
-		if (IConcertClientReplicationManager* ReplicationManager = Client->GetReplicationManager())
+		const bool bSuccess = JoinSessionResult.ErrorCode == EJoinReplicationErrorCode::Success;
+		if (bSuccess)
+		{
+			ConnectedState.Emplace(Client);
+			SetConnectionStateAndBroadcast(EMultiUserReplicationConnectionState::Connected);
+		}
+		else
 		{
-			ReplicationManager->LeaveReplicationSession();
+			SetConnectionStateAndBroadcast(EMultiUserReplicationConnectionState::Disconnected);
 		}
 	}
 
-	void FMultiUserReplicationManager::ClearSessionData()
+	void FMultiUserReplicationManager::SetConnectionStateAndBroadcast(EMultiUserReplicationConnectionState NewState)
 	{
-		SessionContent->ClearClients();
-		LocalClientContent = SessionContent->AddClient();
+		ConnectionState = NewState;
+		OnReplicationConnectionStateChangedDelegate.Broadcast(ConnectionState);
 	}
+
+	FMultiUserReplicationManager::FConnectedState::FConnectedState(TSharedRef<IConcertSyncClient> InClient)
+		: StreamSynchronizer(InClient)
+		, AuthorityPolicy(StreamSynchronizer.GetDiffer(), InClient)
+	{}
 }
diff --git a/Engine/Plugins/Developer/Concert/ConcertApp/MultiUserClient/Source/MultiUserClient/Private/Replication/MultiUserReplicationManager.h b/Engine/Plugins/Developer/Concert/ConcertApp/MultiUserClient/Source/MultiUserClient/Private/Replication/MultiUserReplicationManager.h
--- a/Engine/Plugins/Developer/Concert/ConcertApp/MultiUserClient/Source/MultiUserClient/Private/Replication/MultiUserReplicationManager.h
+++ b/Engine/Plugins/Developer/Concert/ConcertApp/MultiUserClient/Source/MultiUserClient/Private/Replication/MultiUserReplicationManager.h
@@ -2,61 +2,119 @@
 
 #pragma once
 
-#include "Assets/MultiUserReplicationSessionPreset.h"
+#include "Replication/Authority/AuthorityPolicy.h"
+#include "Replication/Stream/ClientStreamRepository.h"
 #include "IConcertSession.h"
+
+#include "Misc/Optional.h"
 #include "Templates/SharedPointer.h"
-#include "UObject/GCObject.h"
+#include "Templates/UnrealTemplate.h"
 
 class IConcertClientSession;
 class IConcertSyncClient;
-
 enum class EConcertConnectionStatus : uint8;
 
+namespace UE::ConcertSyncClient::Replication
+{
+	struct FJoinReplicatedSessionResult;
+}
+namespace UE::ConcertClientSharedSlate
+{
+	class IEditableObjectToPropertiesModel;
+}
+
 namespace UE::MultiUserClient
 {
+	namespace Replication
+	{
+		struct FJoinSessionResult;
+	}
+
+	enum class EMultiUserReplicationConnectionState : uint8
+	{
+		Connecting,
+		Connected,
+		Disconnected
+	};
+	
 	/**
 	 * Interacts with the replication system on behalf of Multi-User to execute actions specific to Multi-User workflows;
 	 * this is opposed to other uses of the replication API, e.g. users using the system in a shipped game.
 	 *
-	 * This class will
-	 *  - be used as a model for the MU control views, such as displaying the streams in the current session (TODO DP UE-193541).
-	 *  - implement auto-join behavior in response to joining a concert session (TODO DP UE-193538)
-	 *
 	 * This class implements the Fence design pattern. All knowledge Multi-User might need should be encapsulated by this class.
 	 */
-	class FMultiUserReplicationManager : public FGCObject
+	class FMultiUserReplicationManager
+		: public TSharedFromThis<FMultiUserReplicationManager>
+		, public FNoncopyable
 	{
 	public:
 		
 		FMultiUserReplicationManager(TSharedRef<IConcertSyncClient> InClient);
-		virtual ~FMultiUserReplicationManager() override;
+		~FMultiUserReplicationManager();
+
+		/**
+		 * Joins the replication session.
+		 *
+		 * Joining occurs automatically after successful connection to the Concert session. However the request can be
+		 * rejected by the server. In that case, the user can manually attempt to connect again, which is what this
+		 * is exposed publicly for.
+		 */
+		void JoinReplicationSession();
+
+		/** @note You're not supposed to keep any reference to StreamSynchronizer since it can become invalid depending on connection state. */
+		FClientStreamRepository* GetStreamSynchronizer() { return ConnectedState ? &ConnectedState->StreamSynchronizer : nullptr; }
+		const FClientStreamRepository* GetStreamSynchronizer() const { return ConnectedState ? &ConnectedState->StreamSynchronizer : nullptr; }
 
-		UMultiUserReplicationSessionPreset* GetSessionContent() const { return SessionContent; }
-		UMultiUserReplicationClientPreset* GetLocalClientContent() const { return LocalClientContent; }
+		/** @note You're not supposed to keep any reference to AuthorityPolicy since it can become invalid depending on connection state. */
+		FAuthorityPolicy* GetAuthorityPolicy() { return ConnectedState ? & ConnectedState->AuthorityPolicy : nullptr; }
+		const FAuthorityPolicy* GetAuthorityPolicy() const { return ConnectedState ? & ConnectedState->AuthorityPolicy : nullptr; }
 
-		//~ Begin FGCObject Interface
-		virtual void AddReferencedObjects(FReferenceCollector& Collector) override;
-		virtual FString GetReferencerName() const override { return TEXT("FMultiUserReplicationManager"); }
-		//~ End FGCObject Interface
+		/** Called when the connection to the replication system changes. */
+		DECLARE_MULTICAST_DELEGATE_OneParam(FOnReplicationConnectionStateChanged, EMultiUserReplicationConnectionState /*NewState*/);
+		FOnReplicationConnectionStateChanged& OnReplicationConnectionStateChanged() { return OnReplicationConnectionStateChangedDelegate; }
+		EMultiUserReplicationConnectionState GetConnectionState() const { return ConnectionState; }
 
 	private:
 
 		/** Client through which the replication bridge is accessed. */
-		TSharedRef<IConcertSyncClient> Client;
-
-		/** The state of the server is synched up with this object and displayed in the UI. */
-		TObjectPtr<UMultiUserReplicationSessionPreset> SessionContent;
-		/** Data for the local client. Also part of SessionContent->ClientPresets. */
-		TObjectPtr<UMultiUserReplicationClientPreset> LocalClientContent;
+		const TSharedRef<IConcertSyncClient> Client;
 		
+		/** Reflects the current connection state to the replication system (note: this does not reflect the state to the concert session). */
+		EMultiUserReplicationConnectionState ConnectionState = EMultiUserReplicationConnectionState::Disconnected;
+
+		struct FConnectedState
+		{
+			// The order of the below members matters so they are destroyed in the right order!
+			// Rule: Lower systems can only reference higher systems (reminder: C++ destroys in reverse declaration order).
+			
+			/**
+			 * Creates UMultiUserReplicationSessionPreset which is displayed by UI.
+			 * Keeps the preset in sync with the state on the server.
+			 *
+			 * Only valid when ConnectionState == EMultiUserReplicationConnectionState::Connected.
+			 */
+			FClientStreamRepository StreamSynchronizer;
+
+			/** Manages the client's authority like automatically taking ownership over newly submitted objects. */
+			FAuthorityPolicy AuthorityPolicy;
+			
+			FConnectedState(TSharedRef<IConcertSyncClient> InClient);
+		};
+		/** Set when connected to a replication session. */
+		TOptional<FConnectedState> ConnectedState;
+
+		/** Called when ConnectionState changes. */
+		FOnReplicationConnectionStateChanged OnReplicationConnectionStateChangedDelegate;
+
+		/** Callback into Concert for when client connection has changed. */
 		void OnSessionConnectionChanged(IConcertClientSession& ConcertClientSession, EConcertConnectionStatus ConcertConnectionStatus);
-		
-		/** Joins a replication session. */
-		void OnJoinSession(IConcertClientSession& ConcertClientSession);
 		/** Leaves the current replication session */
 		void OnLeaveSession(IConcertClientSession& ConcertClientSession);
+
+		/** Handles server response for joining replication session */
+		void HandleReplicationSessionJoined(const ConcertSyncClient::Replication::FJoinReplicatedSessionResult& JoinSessionResult);
 		
-		void ClearSessionData();
+		void SetConnectionStateAndBroadcast(EMultiUserReplicationConnectionState NewState);
 	};
 }
 
diff --git a/Engine/Plugins/Developer/Concert/ConcertApp/MultiUserClient/Source/MultiUserClient/Private/Replication/ReplicationUtils.cpp b/Engine/Plugins/Developer/Concert/ConcertApp/MultiUserClient/Source/MultiUserClient/Private/Replication/ReplicationUtils.cpp
deleted file mode 100644
diff --git a/Engine/Plugins/Developer/Concert/ConcertApp/MultiUserClient/Source/MultiUserClient/Private/Replication/ReplicationUtils.h b/Engine/Plugins/Developer/Concert/ConcertApp/MultiUserClient/Source/MultiUserClient/Private/Replication/ReplicationUtils.h
deleted file mode 100644
diff --git a/Engine/Plugins/Developer/Concert/ConcertApp/MultiUserClient/Source/MultiUserClient/Private/Replication/Stream/ClientStreamRepository.cpp b/Engine/Plugins/Developer/Concert/ConcertApp/MultiUserClient/Source/MultiUserClient/Private/Replication/Stream/ClientStreamRepository.cpp
new file mode 100644
--- /dev/null
+++ b/Engine/Plugins/Developer/Concert/ConcertApp/MultiUserClient/Source/MultiUserClient/Private/Replication/Stream/ClientStreamRepository.cpp
@@ -0,0 +1,56 @@
+// Copyright Epic Games, Inc. All Rights Reserved.
+
+#include "ClientStreamRepository.h"
+
+#include "Assets/MultiUserReplicationSessionPreset.h"
+#include "LocalClientStreamSynchronizer.h"
+#include "Replication/ReplicationWidgetFactories.h"
+#include "Replication/Editor/Model/IEditableObjectToPropertiesModel.h"
+
+#include "UObject/Package.h"
+#include "UObject/UObjectGlobals.h"
+
+namespace UE::MultiUserClient
+{
+	FClientStreamRepository::FClientStreamRepository(TSharedRef<IConcertSyncClient> InClient)
+		: SessionContent(NewObject<UMultiUserReplicationSessionPreset>(GetTransientPackage(), NAME_None, RF_Transient))
+		, LocalClientContent(SessionContent->AddClient())
+		, LocalClientEditModel(ConcertClientSharedSlate::CreatePropertySelectionModel(*LocalClientContent->Stream, LocalClientContent->Stream->MakeReplicationMapGetterAttribute()))
+		, LocalClientStreamDiffer(MakeShared<FLocalClientStreamSynchronizer>(
+			MoveTemp(InClient),
+			LocalClientContent->Stream->StreamId,
+			LocalClientContent->Stream->MakeReplicationMapGetterAttribute(),
+			FLocalClientStreamSynchronizer::FOnModifyReplicationMap::CreateLambda([this](){ LocalClientContent->Stream->Modify(); })
+			))
+	{
+		LocalClientEditModel->OnObjectsChanged().AddRaw(this, &FClientStreamRepository::OnObjectsChanged);
+		LocalClientEditModel->OnPropertiesChanged().AddRaw(this, &FClientStreamRepository::OnPropertiesChanged);
+		
+		LocalClientStreamDiffer->OnChangesReverted_GameThread().AddLambda([this]()
+		{
+			OnModelChangedDelegate.Broadcast();
+		});
+	}
+
+	void FClientStreamRepository::AddReferencedObjects(FReferenceCollector& Collector)
+	{
+		Collector.AddReferencedObject(SessionContent);
+		Collector.AddReferencedObject(LocalClientContent);
+	}
+
+	void FClientStreamRepository::OnObjectsChanged(
+		TArrayView<UObject*> Objects,
+		TArrayView<FSoftObjectPath> SoftObjectPaths,
+		ConcertClientSharedSlate::EReplicatedObjectChangeReason ReplicatedObjectChangeReason
+		)
+	{
+		// Could improve performance by just considering what actually changed instead of doing a full rebuild
+		LocalClientStreamDiffer->RefreshChangesCache();
+	}
+
+	void FClientStreamRepository::OnPropertiesChanged()
+	{
+		// Could improve performance by just considering what actually changed instead of doing a full rebuild
+		LocalClientStreamDiffer->RefreshChangesCache();
+	}
+}
diff --git a/Engine/Plugins/Developer/Concert/ConcertApp/MultiUserClient/Source/MultiUserClient/Private/Replication/Stream/ClientStreamRepository.h b/Engine/Plugins/Developer/Concert/ConcertApp/MultiUserClient/Source/MultiUserClient/Private/Replication/Stream/ClientStreamRepository.h
new file mode 100644
--- /dev/null
+++ b/Engine/Plugins/Developer/Concert/ConcertApp/MultiUserClient/Source/MultiUserClient/Private/Replication/Stream/ClientStreamRepository.h
@@ -0,0 +1,81 @@
+// Copyright Epic Games, Inc. All Rights Reserved.
+
+#pragma once
+
+#include "Assets/MultiUserReplicationSessionPreset.h"
+#include "Replication/Editor/Model/IEditableObjectToPropertiesModel.h"
+#include "UObject/GCObject.h"
+#include "Templates/UnrealTemplate.h"
+
+class IConcertSyncClient;
+
+namespace UE::ConcertClientSharedSlate
+{
+	class IEditableObjectToPropertiesModel;
+}
+
+namespace UE::MultiUserClient
+{
+	class FLocalClientStreamSynchronizer;
+
+	/** Stores all clients' streams in a UMultiUserReplicationSessionPreset and keeps it in sync with the asset. */
+	class FClientStreamRepository
+		: public FGCObject
+		, public FNoncopyable
+	{
+	public:
+
+		FClientStreamRepository(TSharedRef<IConcertSyncClient> InClient);
+		
+		UMultiUserReplicationSessionPreset* GetSessionContent() const { return SessionContent; }
+		UMultiUserReplicationClientPreset* GetLocalClientContent() const { return LocalClientContent; }
+		TSharedRef<ConcertClientSharedSlate::IEditableObjectToPropertiesModel> GetLocalClientEditModel() const { return LocalClientEditModel; }
+		
+		const FLocalClientStreamSynchronizer& GetDiffer() const { return LocalClientStreamDiffer.Get(); }
+		FLocalClientStreamSynchronizer& GetDiffer() { return LocalClientStreamDiffer.Get(); }
+
+		/**
+		 * Called when the data underlying the model has changed externally. Since the change was not caused by the model,
+		 * its events, like IEditableObjectToPropertiesModel::OnObjectsChanged, were not called.
+		 * 
+		 * Subscribers are intended to call IReplicationEditorView::Refresh() in response.
+		 * 
+		 * Examples: Remote client changed their stream, remote client joined with streams,
+		 * or the local client reverted a change that was rejected by the server.
+		 */
+		DECLARE_MULTICAST_DELEGATE(FOnModelChanged);
+		FOnModelChanged& OnModelChanged_GameThread() { return OnModelChangedDelegate; }
+		
+		//~ Begin FGCObject Interface
+		virtual void AddReferencedObjects(FReferenceCollector& Collector) override;
+		virtual FString GetReferencerName() const override { return TEXT("FReplicationStreamSynchronizer"); }
+		//~ End FGCObject Interface
+		
+	private:
+		
+		/** The state of the server is synched up with this object and displayed in the UI. */
+		TObjectPtr<UMultiUserReplicationSessionPreset> SessionContent;
+		/** Data for the local client. Also part of SessionContent->ClientPresets. */
+		TObjectPtr<UMultiUserReplicationClientPreset> LocalClientContent;
+
+		/**
+		 * Used to detect changes made to the local client.
+		 
+		 * The SessionContent UObject must be kept alive as long as this model lives. 
+		 * It should be noted that the UI keeps a strong reference to this model.
+		 * The UI is destroyed right after this FClientStreamRepository.
+		 * @see FMultiUserReplicationManager::OnLeaveSession
+		 */
+		TSharedRef<ConcertClientSharedSlate::IEditableObjectToPropertiesModel> LocalClientEditModel;
+
+		/** Keeps track of the server's state of the local client's streams and handles changing them. */
+		TSharedRef<FLocalClientStreamSynchronizer> LocalClientStreamDiffer;
+
+		/** Called when the data underlying the model has changed and the UI needs to be refreshed. */
+		FOnModelChanged OnModelChangedDelegate;
+		
+		void OnObjectsChanged(TArrayView<UObject*> Objects, TArrayView<FSoftObjectPath> SoftObjectPaths, ConcertClientSharedSlate::EReplicatedObjectChangeReason ReplicatedObjectChangeReason);
+		void OnPropertiesChanged();
+	};
+}
+
diff --git a/Engine/Plugins/Developer/Concert/ConcertApp/MultiUserClient/Source/MultiUserClient/Private/Replication/Stream/LocalClientStreamSynchronizer.cpp b/Engine/Plugins/Developer/Concert/ConcertApp/MultiUserClient/Source/MultiUserClient/Private/Replication/Stream/LocalClientStreamSynchronizer.cpp
new file mode 100644
--- /dev/null
+++ b/Engine/Plugins/Developer/Concert/ConcertApp/MultiUserClient/Source/MultiUserClient/Private/Replication/Stream/LocalClientStreamSynchronizer.cpp
@@ -0,0 +1,258 @@
+// Copyright Epic Games, Inc. All Rights Reserved.
+
+#include "LocalClientStreamSynchronizer.h"
+
+#include "IConcertSyncClient.h"
+#include "Replication/ChangeStreamSharedUtils.h"
+#include "Replication/IConcertClientReplicationManager.h"
+
+#include "ScopedTransaction.h"
+
+#define LOCTEXT_NAMESPACE "FLocalClientStreamDiffer"
+
+namespace UE::MultiUserClient
+{
+	FLocalClientStreamSynchronizer::FLocalClientStreamSynchronizer(
+		TSharedRef<IConcertSyncClient> InLocalClient,
+		const FGuid& InLocalClientStreamId,
+		TAttribute<FObjectReplicationMap*> InStreamWithInProgressChangesAttribute,
+		FOnModifyReplicationMap InOnModifyReplicationMapDelegate
+		)
+		: LocalClient(MoveTemp(InLocalClient))
+		, LocalClientStreamId(InLocalClientStreamId)
+		, StreamWithInProgressChangesAttribute(MoveTemp(InStreamWithInProgressChangesAttribute))
+		, OnModifyReplicationMapDelegate(MoveTemp(InOnModifyReplicationMapDelegate))
+	{
+		check(InStreamWithInProgressChangesAttribute.IsBound() || InStreamWithInProgressChangesAttribute.IsSet());
+	}
+
+	void FLocalClientStreamSynchronizer::RefreshChangesCache()
+	{
+		CachedDeltaChange = DiffChanges();
+		RefreshWarnings();
+	}
+
+	TFuture<FLocalClientStreamSynchronizer::FSubmissionResult> FLocalClientStreamSynchronizer::SubmitChanges(bool bRefreshChanges)
+	{
+		using namespace ConcertSyncClient::Replication;
+
+		IConcertClientReplicationManager* ReplicationManager = LocalClient->GetReplicationManager();
+		if (!ensure(ReplicationManager) || !CanSubmitChanges())
+		{
+			return MakeFulfilledPromise<FSubmissionResult>(FSubmissionResult{}).GetFuture();
+		}
+
+		if (bRefreshChanges)
+		{
+			RefreshChangesCache();
+		}
+		ChangeInTransit = CachedDeltaChange;
+		FChangeStreamRequest Request = BuildChangeRequest();
+		
+		return ReplicationManager->ChangeStream(Request)
+			.Next([this, DestructionDetection = LifetimeToken->AsWeak(), Request](FChangeStreamResponse&& Response)
+			{
+				const FSubmissionResult SubmissionResult { Request, Response };
+				// The request might execute after we're destroyed, e.g. by leaving session while request is on the way.
+				// In that case, the Concert session triggers the OnSessionConnectionChanged which destroys us. Only after that, all the requests are timed out.
+				if (!DestructionDetection.IsValid())
+				{
+					return SubmissionResult;
+				}
+
+				const FChangelist ChangeThatWasInTransit = MoveTemp(ChangeInTransit.GetValue());
+				ChangeInTransit.Reset();
+				
+				if (Response.IsSuccess())
+				{
+					// Tell everybody who cares old and new implicit state. They can e.g. decide to request authority for the new objects.
+					OnChangesAcceptedDelegate.Broadcast(ConfirmedServerState, Request);
+					UpdateConfirmedServerState(Request, ChangeThatWasInTransit);
+				}
+				
+				return SubmissionResult;
+			});
+	}
+
+	void FLocalClientStreamSynchronizer::RevertCachedChanges()
+	{
+		if (HasRevertableLocalChanges())
+		{
+			FScopedTransaction Transaction(LOCTEXT("RevertCachedChanges", "Revert replication changes"));
+			OnModifyReplicationMapDelegate.ExecuteIfBound();
+			
+			*StreamWithInProgressChangesAttribute.Get() = ConfirmedServerState;
+			CachedDeltaChange = {};
+
+			check(IsInGameThread());
+			OnChangesRevertedDelegate.Broadcast();
+			RefreshChangesCache();
+		}
+	}
+
+	bool FLocalClientStreamSynchronizer::HasSubmittableLocalChanges() const
+	{
+		return !CachedDeltaChange.ObjectsToPut.IsEmpty() || !CachedDeltaChange.ObjectsToRemove.IsEmpty();
+	}
+
+	bool FLocalClientStreamSynchronizer::HasRevertableLocalChanges() const
+	{
+		return HasSubmittableLocalChanges() || HasAnyWarnings();
+	}
+
+	bool FLocalClientStreamSynchronizer::IsChangeRequestInTransit() const
+	{
+		return ChangeInTransit.IsSet();
+	}
+
+	EObjectWarningFlags FLocalClientStreamSynchronizer::GetObjectWarningFlags(const FSoftObjectPath& ObjectPath)
+	{
+		const EObjectWarningFlags* Flags = ObjectsWithWarnings.Find(ObjectPath);
+		return Flags ? *Flags : EObjectWarningFlags::Ok;
+	}
+
+	bool FLocalClientStreamSynchronizer::HasAnyWarnings() const
+	{
+		return !ObjectsWithWarnings.IsEmpty();
+	}
+
+	FLocalClientStreamSynchronizer::EObjectChangeType FLocalClientStreamSynchronizer::GetObjectChanges(const FSoftObjectPath& Object) const
+	{
+		// TODO DP:
+		return EObjectChangeType::NoChange;
+	}
+
+	FLocalClientStreamSynchronizer::EPropertyChangeType FLocalClientStreamSynchronizer::GetPropertyChanges(const FSoftObjectPath& Object, const FConcertPropertyChain& PropertyChain) const
+	{
+		// TODO DP:
+		return EPropertyChangeType::NoChange;
+	}
+	
+	FLocalClientStreamSynchronizer::FChangelist FLocalClientStreamSynchronizer::DiffChanges(
+		const FGuid& StreamId,
+		const FObjectReplicationMap& Base,
+		const FObjectReplicationMap& Changed
+		)
+	{
+		// BuildRequestFromDiff does not tolerate any invalid entries (like empty properties, which the UI generates right after you add an object to the list)
+		FObjectReplicationMap Cleansed = Changed;
+		ConcertSyncCore::Replication::ChangeStreamUtils::IterateInvalidEntries(Changed, [&Cleansed](const FSoftObjectPath& InvalidObject, const FReplicatedObjectInfo&)
+		{
+			Cleansed.ReplicatedObjects.Remove(InvalidObject);
+			return EBreakBehavior::Continue;
+		});
+
+		FConcertChangeStream_Request Request = ConcertSyncCore::Replication::ChangeStreamUtils::BuildRequestFromDiff(StreamId, Base, Cleansed);
+		return { MoveTemp(Request.ObjectsToRemove), MoveTemp(Request.ObjectsToPut) };
+	}
+
+	void FLocalClientStreamSynchronizer::RefreshWarnings()
+	{
+		ObjectsWithWarnings.Empty();
+		const FObjectReplicationMap* Map = StreamWithInProgressChangesAttribute.Get();
+		if (!Map)
+		{
+			return;
+		}
+		
+		ConcertSyncCore::Replication::ChangeStreamUtils::IterateInvalidEntries(*Map, [this](const FSoftObjectPath& InvalidObject, const FReplicatedObjectInfo& Info)
+		{
+			EObjectWarningFlags Flags = EObjectWarningFlags::Ok;
+			if (Info.PropertySelection.ReplicatedProperties.IsEmpty())
+			{
+				Flags |= EObjectWarningFlags::MissingProperties;
+			}
+			
+			ObjectsWithWarnings.Add(InvalidObject, Flags);
+			return EBreakBehavior::Continue;
+		});
+	}
+
+	ConcertSyncClient::Replication::FChangeStreamRequest FLocalClientStreamSynchronizer::BuildChangeRequest() const
+	{
+		switch (ComputeNextRequestType())
+		{
+		case EChangeRequestType::CreateNewStream: return BuildChangeRequest_CreateNewStream(LocalClientStreamId, CachedDeltaChange);
+		case EChangeRequestType::UpdateExistingStream: return BuildChangeRequest_UpdateExistingStream(CachedDeltaChange);
+		default: checkNoEntry(); return {};
+		}
+	}
+
+	ConcertSyncClient::Replication::FChangeStreamRequest FLocalClientStreamSynchronizer::BuildChangeRequest_CreateNewStream(const FGuid& StreamId, const FChangelist& FromChangelist)
+	{
+		const TMap<FObjectInStreamID, FConcertChangeStream_PutObject>& ObjectsToPut = FromChangelist.ObjectsToPut;
+		
+		ConcertSyncClient::Replication::FChangeStreamRequest Request;
+		Request.StreamsToAdd.Emplace();
+		FReplicationStreamDescription_NetPacked& NewStream = Request.StreamsToAdd[0];
+		NewStream.BaseDescription.Identifier = StreamId;
+			
+		// If creating a new stream, the objects must be supplied in the description instead of in PutObjects!
+		FObjectReplicationMap& ReplicationMap = NewStream.BaseDescription.ReplicationMap;
+		ReplicationMap.ReplicatedObjects.Reserve(ObjectsToPut.Num());
+		for (const TPair<FObjectInStreamID, FConcertChangeStream_PutObject>& PutObjectPair : ObjectsToPut)
+		{
+			const TOptional<FReplicatedObjectInfo> NewObjectInfo = PutObjectPair.Value.MakeObjectInfoIfValid();
+			// The editing UI allows adding objects without properties (to make UX easier) - do not submit those to the server.
+			if (!NewObjectInfo)
+			{
+				continue;
+			}
+			checkf(PutObjectPair.Key.StreamId == StreamId, TEXT("BuildChangeRequest should always use LocalClientStreamId ID!"));
+				
+			ReplicationMap.ReplicatedObjects.Add(PutObjectPair.Key.Object, *NewObjectInfo);
+		}
+		
+		return Request;
+	}
+
+	ConcertSyncClient::Replication::FChangeStreamRequest FLocalClientStreamSynchronizer::BuildChangeRequest_UpdateExistingStream(FChangelist FromChangelist)
+	{
+		return { MoveTemp(FromChangelist.ObjectsToRemove), MoveTemp(FromChangelist.ObjectsToPut) };
+	}
+	
+	FLocalClientStreamSynchronizer::EChangeRequestType FLocalClientStreamSynchronizer::ComputeNextRequestType() const
+	{
+		return ConfirmedServerState.ReplicatedObjects.IsEmpty()
+			? EChangeRequestType::CreateNewStream
+			: EChangeRequestType::UpdateExistingStream;
+	}
+
+	void FLocalClientStreamSynchronizer::UpdateConfirmedServerState(
+		const ConcertSyncClient::Replication::FChangeStreamRequest& Request,
+		const FChangelist& ChangeThatWasInTransit
+		)
+	{
+		const bool bCreatedNewStream = !Request.StreamsToAdd.IsEmpty(); 
+		if (bCreatedNewStream)
+		{
+			ConfirmedServerState = Request.StreamsToAdd[0].BaseDescription.ReplicationMap;
+		}
+		else
+		{
+			// The goal here is to leverage ApplyValidatedRequest for updating ConfirmedServerState.
+			TArray<FReplicationStreamDescription> FakeDescriptions;
+			FakeDescriptions.Emplace();
+			FReplicationStreamDescription& FakeStreamDescription = FakeDescriptions[0];
+			FakeStreamDescription.BaseDescription.Identifier = LocalClientStreamId;
+			// Cheaply move ConfirmedServerState into the FReplicationStreamDescription so it satisfies the ApplyValidatedRequest API...
+			FakeStreamDescription.BaseDescription.ReplicationMap = MoveTemp(ConfirmedServerState);
+			ConcertSyncCore::Replication::ChangeStreamUtils::ApplyValidatedRequest(
+				BuildChangeRequest_UpdateExistingStream(ChangeThatWasInTransit),
+				FakeDescriptions
+				);
+						
+			// Careful: ApplyValidatedRequest may have emptied the array!
+			FObjectReplicationMap NewServerState = FakeDescriptions.IsEmpty()
+				? FObjectReplicationMap{}
+			: FakeStreamDescription.BaseDescription.ReplicationMap;
+						
+			// ... and cheaply move the updated map back into our ConfirmedServerState
+			ConfirmedServerState = MoveTemp(NewServerState);
+		}
+					
+		RefreshChangesCache();
+	}
+}
+
+#undef LOCTEXT_NAMESPACE
\ No newline at end of file
diff --git a/Engine/Plugins/Developer/Concert/ConcertApp/MultiUserClient/Source/MultiUserClient/Private/Replication/Stream/LocalClientStreamSynchronizer.h b/Engine/Plugins/Developer/Concert/ConcertApp/MultiUserClient/Source/MultiUserClient/Private/Replication/Stream/LocalClientStreamSynchronizer.h
new file mode 100644
--- /dev/null
+++ b/Engine/Plugins/Developer/Concert/ConcertApp/MultiUserClient/Source/MultiUserClient/Private/Replication/Stream/LocalClientStreamSynchronizer.h
@@ -0,0 +1,194 @@
+// Copyright Epic Games, Inc. All Rights Reserved.
+
+#pragma once
+
+#include "Replication/IToken.h"
+#include "Replication/Data/ObjectIds.h"
+#include "Replication/Data/ObjectReplicationMap.h"
+#include "Replication/Messages/ConcertReplicationEvents.h"
+
+#include "Async/Future.h"
+#include "Misc/Attribute.h"
+#include "Templates/UnrealTemplate.h"
+
+class IConcertSyncClient;
+
+namespace UE::ConcertSyncClient::Replication
+{
+	struct FChangeStreamResponse;
+	struct FChangeStreamRequest;
+}
+
+namespace UE::MultiUserClient
+{
+	enum class EObjectWarningFlags
+	{
+		Ok = 0,
+		MissingProperties = 1 << 0
+	};
+	ENUM_CLASS_FLAGS(EObjectWarningFlags);
+	
+	/**
+	 * Knows of the local client's registered replication streams and builds a changelist. The changelist tracks the
+	 * unconfirmed changes to the client's streams and is updates when the server confirms the change.
+	 */
+	class FLocalClientStreamSynchronizer : public FNoncopyable, public TSharedFromThis<FLocalClientStreamSynchronizer>
+	{
+	public:
+
+		DECLARE_DELEGATE(FOnModifyReplicationMap);
+		FLocalClientStreamSynchronizer(
+			TSharedRef<IConcertSyncClient> InLocalClient,
+			const FGuid& InLocalClientStreamId,
+			TAttribute<FObjectReplicationMap*> InStreamWithInProgressChangesAttribute,
+			FOnModifyReplicationMap InOnModifyReplicationMapDelegate
+			);
+
+		/**
+		 * Clears and rebuilds CachedDeltaChange in response to the underlying StreamWithInProgressChangesAttribute having been changed.
+		 * 
+		 * This iterates through all elements, which is suboptimal. However, this is needed e.g. after an undo & redo command.
+		 * In the future, we could implement a more efficient version that just responds to the relevant events (like adding an object).
+		 * While this function does redundant work, it should really not be a big performance bog after all. Also it is called
+		 * infrequently e.g. in response to a user pressing a button so we should be fine.
+		 */
+		void RefreshChangesCache();
+
+		using FSubmissionResult = TPair<ConcertSyncClient::Replication::FChangeStreamRequest, ConcertSyncClient::Replication::FChangeStreamResponse>;
+		/**
+		 * Synchronizes the server with the locally made changes to the client stream.
+		 * @note This future can execute on any thread.
+		 */
+		TFuture<FSubmissionResult> SubmitChanges(bool bRefreshChanges = false);
+		
+		/** Reverts the external stream StreamWithInProgressChangesAttribute to be the version registered on the server. */
+		void RevertCachedChanges();
+
+		/** @return Whether there are any changes that that can be submitted to the server (excludes those with warnings). */
+		bool HasSubmittableLocalChanges() const;
+		/** @return Whether there are any changes including those with warnings. */
+		bool HasRevertableLocalChanges() const;
+		
+		/** @return Whether we're currently awaiting a response from the server for updating the stream. */
+		bool IsChangeRequestInTransit() const;
+		/** @return Whether it is ok to call SubmitChanges. */
+		bool CanSubmitChanges() const { return HasSubmittableLocalChanges() && !IsChangeRequestInTransit(); }
+
+		/** @return The warnings flags for the given ObjectPath. */
+		EObjectWarningFlags GetObjectWarningFlags(const FSoftObjectPath& ObjectPath);
+		/** @return Whether there are any warnings for any object. */
+		bool HasAnyWarnings() const;
+
+		enum class EObjectChangeType
+		{
+			NoChange,
+			Added,
+			Removed,
+			PropertiesModified
+		};
+		/** @return Whether this object's local configuration differs from the server version. */
+		EObjectChangeType GetObjectChanges(const FSoftObjectPath& Object) const;
+		
+		enum class EPropertyChangeType
+		{
+			NoChange,
+			Added,
+			Removed
+		};
+		/** @return Gets how the given property config differs from what is on the server. */
+		EPropertyChangeType GetPropertyChanges(const FSoftObjectPath& Object, const FConcertPropertyChain& PropertyChain) const;
+
+		/** Called when RevertCachedChanges is called. The UI should be refreshed. */
+		DECLARE_MULTICAST_DELEGATE(FOnChangesReverted);
+		FOnChangesReverted& OnChangesReverted_GameThread() { return OnChangesRevertedDelegate; }
+		
+		/** Called when a change request that was in transit was accepted by the server. */
+		DECLARE_MULTICAST_DELEGATE_TwoParams(FOnChangesAccepted, const FObjectReplicationMap& OldServerState, const ConcertSyncClient::Replication::FChangeStreamRequest& AcceptedRequest);
+		FOnChangesAccepted& OnChangesAccepted_AnyThread() { return OnChangesAcceptedDelegate; }
+		
+	private:
+		
+		/** Referenced by the change requests to detect destruction of FLocalClientStreamSynchronizer */
+		const TSharedRef<FToken> LifetimeToken = FToken::Make();
+		
+		/** Owning client. Used to send change requests to the server. */
+		const TSharedRef<IConcertSyncClient> LocalClient;
+		/** The ID of the local client's stream this FLocalClientStreamDiffer is managing. */
+		const FGuid LocalClientStreamId;
+
+		/**
+		 * Represents ConfirmedServerState with changes made to it. These changes have not been sent to the server, yet.
+		 * This is cleared when a change is sent to the server. 
+		 */
+		const TAttribute<FObjectReplicationMap*> StreamWithInProgressChangesAttribute;
+		/**
+		 * Represents what the local client thinks the replication map on the server currently looks like.
+		 * This is updated every time the server confirms a change.
+		 */
+		FObjectReplicationMap ConfirmedServerState;
+
+		struct FChangelist
+		{
+			TSet<FObjectInStreamID> ObjectsToRemove;
+			TMap<FObjectInStreamID, FConcertChangeStream_PutObject> ObjectsToPut;
+		};
+		
+		/**
+		 * Represents a change to ConfirmedServerState that is currently in transit to the server.
+		 * When a change is submitted to the server, ChangeInTransit becomes the difference between ConfirmedServerState and StreamWithInProgressChangesAttribute.
+		 * When the server confirms these changes, they are applied to ConfirmedServerState.
+		 */
+		TOptional<FChangelist> ChangeInTransit;
+		/**
+		 * The changes that if applied to ConfirmedServerState would result in StreamWithInProgressChangesAttribute.
+		 * This is updated every time the local client makes changes.
+		 */
+		FChangelist CachedDeltaChange;
+
+		/** These objects have warnings that prevent them from being set to the server (missing data, unsupported class, etc.). */
+		TMap<FSoftObjectPath, EObjectWarningFlags> ObjectsWithWarnings;
+		
+		/** Called when StreamWithInProgressChangesAttribute is about to be modified. */
+		FOnModifyReplicationMap OnModifyReplicationMapDelegate;
+		/** Called when StreamWithInProgressChangesAttribute was modified by this FLocalClientStreamDiffer. */
+		FOnChangesReverted OnChangesRevertedDelegate;
+		/** Called when a change request that was in transit was accepted by the server. */
+		FOnChangesAccepted OnChangesAcceptedDelegate;
+
+		FChangelist DiffChanges() const
+		{
+			const FObjectReplicationMap* Map = StreamWithInProgressChangesAttribute.Get();
+			return Map ? DiffChanges(LocalClientStreamId, ConfirmedServerState, *Map) : FChangelist{};
+		}
+		/** Builds the changelist to get from Base to Changed. */
+		static FChangelist DiffChanges(const FGuid& StreamId, const FObjectReplicationMap& Base, const FObjectReplicationMap& Changed);
+
+		/** Updates ObjectsWithWarnings. */
+		void RefreshWarnings();
+		
+		/** Builds a change request based on what's registered on the server. */
+		ConcertSyncClient::Replication::FChangeStreamRequest BuildChangeRequest() const;
+		enum class EChangeRequestType
+		{
+			/** When NeedsToCreateNewStreamNextRequest returns true */
+			CreateNewStream,
+			UpdateExistingStream
+		};
+		/** Converts a changelist to. */
+		static ConcertSyncClient::Replication::FChangeStreamRequest BuildChangeRequest_CreateNewStream(const FGuid& StreamId, const FChangelist& FromChangelist);
+		static ConcertSyncClient::Replication::FChangeStreamRequest BuildChangeRequest_UpdateExistingStream(FChangelist FromChangelist);
+		/**
+		 * The stream may have to be created first. This is the case either if
+		 *	1. this is the first time sending to the server, or
+		 *	2. the client created objects and then removed them all again (the server destroys the stream in ths case).
+		 */
+		EChangeRequestType ComputeNextRequestType() const;
+
+		/** Updates the server state after the server has answered a request. */
+		void UpdateConfirmedServerState(
+			const ConcertSyncClient::Replication::FChangeStreamRequest& Request,
+			const FChangelist& ChangeThatWasInTransit
+			);
+	};
+}
+
diff --git a/Engine/Plugins/Developer/Concert/ConcertApp/MultiUserClient/Source/MultiUserClient/Private/Widgets/ActiveSession/Replication/Joined/SReplicationJoinedToolbar.cpp b/Engine/Plugins/Developer/Concert/ConcertApp/MultiUserClient/Source/MultiUserClient/Private/Widgets/ActiveSession/Replication/Joined/SReplicationJoinedToolbar.cpp
new file mode 100644
--- /dev/null
+++ b/Engine/Plugins/Developer/Concert/ConcertApp/MultiUserClient/Source/MultiUserClient/Private/Widgets/ActiveSession/Replication/Joined/SReplicationJoinedToolbar.cpp
@@ -0,0 +1,193 @@
+// Copyright Epic Games, Inc. All Rights Reserved.
+
+#include "SReplicationJoinedToolbar.h"
+
+#include "ConcertLogGlobal.h"
+#include "Replication/IConcertClientReplicationManager.h"
+#include "Replication/MultiUserReplicationManager.h"
+#include "Replication/Stream/ClientStreamRepository.h"
+#include "Replication/Stream/LocalClientStreamSynchronizer.h"
+
+#include "Framework/Notifications/NotificationManager.h"
+#include "SNegativeActionButton.h"
+#include "Styling/AppStyle.h"
+#include "Styling/StyleColors.h"
+#include "Widgets/Images/SImage.h"
+#include "Widgets/Input/SButton.h"
+#include "Widgets/Layout/SBox.h"
+#include "Widgets/Layout/SSpacer.h"
+#include "Widgets/Notifications/SNotificationList.h"
+#include "Widgets/SBoxPanel.h"
+#include "Widgets/Text/STextBlock.h"
+
+#define LOCTEXT_NAMESPACE "SReplicationJoinedWidget"
+
+namespace UE::MultiUserClient
+{
+	void SReplicationJoinedToolbar::Construct(const FArguments& InArgs, TSharedRef<FMultiUserReplicationManager> InReplicationManager)
+	{
+		ReplicationManager = InReplicationManager;
+		
+		ChildSlot
+		[
+			SNew(SHorizontalBox)
+
+			+SHorizontalBox::Slot()
+			.FillWidth(1.f)
+			[
+				SNew(SSpacer)
+			]
+
+			// Upload button
+			+SHorizontalBox::Slot()
+			.AutoWidth()
+			[
+				BuildUploadButton()
+			]
+
+			// Revert changes
+			+SHorizontalBox::Slot()
+			.AutoWidth()
+			[
+				BuildRevertButton()
+			]
+		];
+	}
+
+	TSharedRef<SWidget> SReplicationJoinedToolbar::BuildUploadButton()
+	{
+		return SNew(SButton)
+			.ForegroundColor(FSlateColor::UseStyle())
+			.IsEnabled_Lambda([this](){ return ReplicationManager->GetStreamSynchronizer()->GetDiffer().CanSubmitChanges(); })
+			.ToolTipText(this, &SReplicationJoinedToolbar::GetUploadButtonToolTipText)
+			.HAlign(HAlign_Center)
+			.VAlign(VAlign_Center)
+			.OnClicked(this, &SReplicationJoinedToolbar::OnUploadButtonClicked)
+			[
+				SNew(SHorizontalBox)
+				+ SHorizontalBox::Slot()
+				.HAlign(HAlign_Center)
+				.VAlign(VAlign_Center)
+				.AutoWidth()
+				[
+					SNew(SImage)
+					.Image_Lambda([this]()
+					{
+						const FLocalClientStreamSynchronizer& Differ = ReplicationManager->GetStreamSynchronizer()->GetDiffer();
+						return Differ.HasAnyWarnings()
+							? FAppStyle::Get().GetBrush("Icons.Warning")
+							: FAppStyle::Get().GetBrush("Icons.Plus");
+					})
+					.ColorAndOpacity_Lambda([this]()
+					{
+						const FLocalClientStreamSynchronizer& Differ = ReplicationManager->GetStreamSynchronizer()->GetDiffer();
+						return Differ.HasAnyWarnings()
+							? FStyleColors::AccentYellow
+							: FStyleColors::AccentGreen;
+					})
+				]
+				+ SHorizontalBox::Slot()
+				.Padding(FMargin(3, 0, 0, 0))
+				.VAlign(VAlign_Center)
+				.AutoWidth()
+				[
+					SNew(STextBlock)
+					.TextStyle(FAppStyle::Get(), "SmallButtonText")
+					.Text(LOCTEXT("UploadChanges.Label", "Upload"))
+				]
+			];
+	}
+
+	FText SReplicationJoinedToolbar::GetUploadButtonToolTipText() const
+	{
+		const FLocalClientStreamSynchronizer& Differ = ReplicationManager->GetStreamSynchronizer()->GetDiffer();
+		const FText BaseToolTipText = [this, &Differ]()
+		{
+			if (Differ.HasSubmittableLocalChanges() && !Differ.IsChangeRequestInTransit())
+			{
+				return LOCTEXT("UploadChanges.ToolTip.Upload", "Upload your local changes to the server.");
+			}
+			if (Differ.IsChangeRequestInTransit())
+			{
+				return LOCTEXT("UploadChanges.ToolTip.InTransit", "Awaiting server response.");
+			}
+			return LOCTEXT("UploadChanges.ToolTip.NoChanges", "No submittable changes.");
+		}();
+		if (!Differ.HasAnyWarnings())
+		{
+			return BaseToolTipText;
+		}
+
+		return FText::Format(LOCTEXT("UploadChanges.ToolTip.BaseFmt", "{0}\n\n{1}"),
+			BaseToolTipText,
+			LOCTEXT("UploadChanges.ToolTip", "Warning: Objects with warnings will not be submitted. There are objects with warnings.")
+			);
+	}
+
+	FReply SReplicationJoinedToolbar::OnUploadButtonClicked() const
+	{
+		FLocalClientStreamSynchronizer& Differ = ReplicationManager->GetStreamSynchronizer()->GetDiffer();
+		if (Differ.CanSubmitChanges())
+		{
+			FNotificationInfo NotificationInfo(LOCTEXT("Uploading.InProgress", "Uploading stream changes."));
+			NotificationInfo.bUseThrobber = true;
+			NotificationInfo.bUseSuccessFailIcons = true;
+			NotificationInfo.bFireAndForget = false;
+			const TSharedPtr<SNotificationItem> Notification = FSlateNotificationManager::Get().AddNotification(NotificationInfo);
+			Differ.SubmitChanges()
+				.Next([Notification](FLocalClientStreamSynchronizer::FSubmissionResult&& Result)
+				{
+					const ConcertSyncClient::Replication::FChangeStreamResponse& Response = Result.Value;
+					if (Response.IsSuccess())
+					{
+						Notification->SetCompletionState(SNotificationItem::CS_Success);
+						Notification->SetText(LOCTEXT("Uploading.Success", "Changes accepted."));
+									
+						const ConcertSyncClient::Replication::FChangeStreamRequest& Request = Result.Key;
+						// Multi-User adds only one stream so we only need to look at index 0
+						const int32 NumAdditions = Request.StreamsToAdd.IsEmpty() ? 0 : Request.StreamsToAdd[0].BaseDescription.ReplicationMap.ReplicatedObjects.Num();
+						const int32 NumPuts = Request.ObjectsToPut.Num() + NumAdditions;
+						Notification->SetSubText(FText::Format(LOCTEXT("Uploading.Success.SubText", "{0} puts, {1} removals"),
+							NumPuts,
+							Request.ObjectsToRemove.Num()
+							));
+					}
+					else
+					{
+						Notification->SetCompletionState(SNotificationItem::CS_Fail);
+						Notification->SetText(LOCTEXT("Uploading.Failure", "Changes rejected."));
+						Notification->SetSubText(LOCTEXT("Uploading.Failure.SubText", "See log for details."));
+
+						FStringOutputDevice Errors;
+						Response.LogErrors(Errors);
+						UE_LOG(LogConcert, Error, TEXT("Errors uploading stream changes:\n%s"), *Errors);
+					}
+								
+					Notification->ExpireAndFadeout();
+				});
+		}
+		return FReply::Handled();
+	}
+
+	TSharedRef<SWidget> SReplicationJoinedToolbar::BuildRevertButton()
+	{
+		return SNew(SNegativeActionButton)
+			.Text(LOCTEXT("RevertChanges.Label", "Revert"))
+			.ToolTipText_Lambda([this]()
+			{
+				const FLocalClientStreamSynchronizer& Differ = ReplicationManager->GetStreamSynchronizer()->GetDiffer();
+				return Differ.HasRevertableLocalChanges()
+					? LOCTEXT("RevertChanges.ToolTip.Revert", "Reverts all local changes you've made to the state that's on the server.")
+					: LOCTEXT("RevertChanges.ToolTip.NothingToUpload", "Once you make changes to the stream, you can use this button to revert them.");
+			})
+			.IsEnabled_Lambda([this](){ return ReplicationManager->GetStreamSynchronizer()->GetDiffer().HasRevertableLocalChanges(); })
+			.OnClicked_Lambda([this]()
+			{
+				FLocalClientStreamSynchronizer& Differ = ReplicationManager->GetStreamSynchronizer()->GetDiffer();
+				Differ.RevertCachedChanges();
+				return FReply::Handled();
+			});
+	}
+}
+
+#undef LOCTEXT_NAMESPACE
\ No newline at end of file
diff --git a/Engine/Plugins/Developer/Concert/ConcertApp/MultiUserClient/Source/MultiUserClient/Private/Widgets/ActiveSession/Replication/Joined/SReplicationJoinedToolbar.h b/Engine/Plugins/Developer/Concert/ConcertApp/MultiUserClient/Source/MultiUserClient/Private/Widgets/ActiveSession/Replication/Joined/SReplicationJoinedToolbar.h
new file mode 100644
--- /dev/null
+++ b/Engine/Plugins/Developer/Concert/ConcertApp/MultiUserClient/Source/MultiUserClient/Private/Widgets/ActiveSession/Replication/Joined/SReplicationJoinedToolbar.h
@@ -0,0 +1,35 @@
+// Copyright Epic Games, Inc. All Rights Reserved.
+
+#pragma once
+
+#include "Widgets/DeclarativeSyntaxSupport.h"
+#include "Widgets/SCompoundWidget.h"
+
+namespace UE::MultiUserClient
+{
+	class FMultiUserReplicationManager;
+
+	/** Contains a bunch of actions that can be performed while connected to replication. */
+	class SReplicationJoinedToolbar : public SCompoundWidget
+	{
+	public:
+		
+		SLATE_BEGIN_ARGS(SReplicationJoinedToolbar)
+		{}
+		SLATE_END_ARGS()
+
+		void Construct(const FArguments& InArgs, TSharedRef<FMultiUserReplicationManager> InReplicationManager);
+
+	private:
+
+		/** Acts as the model of this view */
+		TSharedPtr<FMultiUserReplicationManager> ReplicationManager;
+		
+		TSharedRef<SWidget> BuildUploadButton();
+		FText GetUploadButtonToolTipText() const;
+		FReply OnUploadButtonClicked() const;
+		
+		TSharedRef<SWidget> BuildRevertButton();
+	};
+}
+
diff --git a/Engine/Plugins/Developer/Concert/ConcertApp/MultiUserClient/Source/MultiUserClient/Private/Widgets/ActiveSession/Replication/Joined/SReplicationJoinedView.cpp b/Engine/Plugins/Developer/Concert/ConcertApp/MultiUserClient/Source/MultiUserClient/Private/Widgets/ActiveSession/Replication/Joined/SReplicationJoinedView.cpp
new file mode 100644
--- /dev/null
+++ b/Engine/Plugins/Developer/Concert/ConcertApp/MultiUserClient/Source/MultiUserClient/Private/Widgets/ActiveSession/Replication/Joined/SReplicationJoinedView.cpp
@@ -0,0 +1,135 @@
+// Copyright Epic Games, Inc. All Rights Reserved.
+
+#include "SReplicationJoinedView.h"
+
+#include "ConcertLogGlobal.h"
+#include "Replication/Editor/Model/Object/EditorObjectSelectionSourceModel.h"
+#include "Replication/Editor/Model/Property/SelectPropertyFromUClassModel.h"
+#include "Replication/Editor/View/IReplicationEditorView.h"
+#include "Replication/MultiUserReplicationManager.h"
+#include "Replication/ReplicationWidgetFactories.h"
+#include "Replication/Stream/ClientStreamRepository.h"
+#include "SReplicationJoinedToolbar.h"
+
+#include "Framework/Notifications/NotificationManager.h"
+#include "Widgets/Layout/SBox.h"
+#include "Widgets/SBoxPanel.h"
+#include "Widgets/Notifications/SNotificationList.h"
+
+#define LOCTEXT_NAMESPACE "SReplicationJoinedWidget"
+
+namespace UE::MultiUserClient
+{
+	void SReplicationJoinedView::Construct(const FArguments& InArgs, TSharedRef<FMultiUserReplicationManager> InReplicationManager)
+	{
+		ReplicationManager = InReplicationManager;
+		
+		using namespace ConcertClientSharedSlate;
+		const FCreateEditorParams ReplicationEditorCreationParams
+		{
+			ReplicationManager->GetStreamSynchronizer()->GetLocalClientEditModel(),
+			MakeShared<FEditorObjectSelectionSourceModel>(),
+			MakeShared<FSelectPropertyFromUClassModel>()
+		};
+
+		const TSharedRef<IReplicationEditorView> NewEditorView = CreateEditor(ReplicationEditorCreationParams);
+		WeakEditorView = NewEditorView;
+		ChildSlot
+		[
+			SNew(SVerticalBox)
+
+			// Toolbar
+			+SVerticalBox::Slot()
+			.AutoHeight()
+			.Padding(2.f)
+			[
+				SNew(SReplicationJoinedToolbar, InReplicationManager)
+			]
+
+			// Editor
+			+SVerticalBox::Slot()
+			.FillHeight(1.f)
+			[
+				NewEditorView
+			]
+		];
+		
+		// Refresh UI if streams change externally, e.g. a remote client changed what they sent
+		ReplicationManager->GetStreamSynchronizer()->OnModelChanged_GameThread().AddSP(this, &SReplicationJoinedView::OnModelChanged);
+
+		// Show notifications about changing authority
+		ReplicationManager->GetAuthorityPolicy()->OnAuthorityRequestSent_AnyThread().AddSP(this, &SReplicationJoinedView::OnAuthorityRequestSent_AnyThread);
+		ReplicationManager->GetAuthorityPolicy()->OnAuthorityResponseReceived_AnyThread().AddSP(this, &SReplicationJoinedView::OnAuthorityResponseReceived_AnyThread);
+	}
+
+	void SReplicationJoinedView::OnModelChanged() const
+	{
+		if (const TSharedPtr<ConcertClientSharedSlate::IReplicationEditorView> EditorView = WeakEditorView.Pin())
+		{
+			EditorView->Refresh();
+		}
+	}
+
+	void SReplicationJoinedView::OnAuthorityRequestSent_AnyThread(
+		const ConcertSyncClient::Replication::FAuthorityChangeRequest& Request
+		)
+	{
+		ExecuteOnGameThread(TEXT("OnAuthorityRequestSent"), [WeakThis = TWeakPtr<SReplicationJoinedView>(SharedThis(this)), Request]()
+		{
+			if (const TSharedPtr<SReplicationJoinedView> ThisPin = WeakThis.Pin())
+			{
+				FNotificationInfo NotificationInfo(LOCTEXT("Authority.InProgress.Text", "Requesting authority."));
+				NotificationInfo.bUseThrobber = true;
+				NotificationInfo.bUseSuccessFailIcons = true;
+				NotificationInfo.bFireAndForget = false;
+				ThisPin->AuthorityChangeNotification = FSlateNotificationManager::Get().AddNotification(NotificationInfo);
+				ThisPin->AuthorityChangeNotification->SetSubText(
+					FText::Format(LOCTEXT("Authority.InProgress.SubTextFmt", "Taking {0}|plural(one=object,other\nReleasing {1}|plural(one=object,other=objects)"),
+						Request.TakeAuthority.Num(),
+						Request.ReleaseAuthority.Num()
+					));
+			}
+		});
+	}
+
+	void SReplicationJoinedView::OnAuthorityResponseReceived_AnyThread(
+		const ConcertSyncClient::Replication::FAuthorityChangeRequest& Request,
+		const ConcertSyncClient::Replication::FAuthorityChangeResponse& Response
+		)
+	{
+		ExecuteOnGameThread(TEXT("OnAuthorityResponseReceived"), [WeakThis = TWeakPtr<SReplicationJoinedView>(SharedThis(this)), Request, Response]()
+		{
+			if (const TSharedPtr<SReplicationJoinedView> ThisPin = WeakThis.Pin()
+				; ThisPin && ensure(ThisPin->AuthorityChangeNotification))
+			{
+				const bool bSuccess = Response.RejectedObjects.IsEmpty();
+				const FText Text = bSuccess
+					? LOCTEXT("Authority.Text.Accepted", "Authority change accepted.")
+					: LOCTEXT("Authority.Text.Rejection", "Authority change had rejections.");
+				const FText SubText = bSuccess
+					? FText::GetEmpty()
+					: FText::Format(LOCTEXT("Authority.SubText.RejectionFmt", "{0} accepted {1} rejected\nSee logs for details."),
+						Request.TakeAuthority.Num() - Response.RejectedObjects.Num(),
+						Response.RejectedObjects.Num()
+						);
+				
+				ThisPin->AuthorityChangeNotification->SetCompletionState(bSuccess ? SNotificationItem::CS_Success : SNotificationItem::CS_Fail);
+				ThisPin->AuthorityChangeNotification->SetText(Text);
+				ThisPin->AuthorityChangeNotification->SetSubText(SubText);
+
+				if (!bSuccess)
+				{
+					const FString Error = FString::JoinBy(Response.RejectedObjects, TEXT("\n"), [](const TPair<FSoftObjectPath, FConcertStreamArray>& Pair)
+					{
+						return Pair.Key.ToString();
+					});
+					UE_LOG(LogConcert, Error, TEXT("Authority change rejected objects:\n%s"), *Error);
+				}
+				
+				ThisPin->AuthorityChangeNotification->ExpireAndFadeout();
+			}
+		});
+	}
+}
+
+#undef LOCTEXT_NAMESPACE
\ No newline at end of file
diff --git a/Engine/Plugins/Developer/Concert/ConcertApp/MultiUserClient/Source/MultiUserClient/Private/Widgets/ActiveSession/Replication/Joined/SReplicationJoinedView.h b/Engine/Plugins/Developer/Concert/ConcertApp/MultiUserClient/Source/MultiUserClient/Private/Widgets/ActiveSession/Replication/Joined/SReplicationJoinedView.h
new file mode 100644
--- /dev/null
+++ b/Engine/Plugins/Developer/Concert/ConcertApp/MultiUserClient/Source/MultiUserClient/Private/Widgets/ActiveSession/Replication/Joined/SReplicationJoinedView.h
@@ -0,0 +1,55 @@
+// Copyright Epic Games, Inc. All Rights Reserved.
+
+#pragma once
+
+#include "Replication/IConcertClientReplicationManager.h"
+#include "Widgets/DeclarativeSyntaxSupport.h"
+#include "Widgets/SCompoundWidget.h"
+
+class SNotificationItem;
+
+namespace UE::ConcertClientSharedSlate
+{
+	class IReplicationEditorView;
+}
+
+namespace UE::MultiUserClient
+{
+	class FMultiUserReplicationManager;
+
+	/** This widget is displayed by SReplicationRootWidget when the client has joined replication. */
+	class SReplicationJoinedView : public SCompoundWidget
+	{
+	public:
+
+		SLATE_BEGIN_ARGS(SReplicationJoinedView)
+		{}
+		SLATE_END_ARGS()
+
+		void Construct(const FArguments& InArgs, TSharedRef<FMultiUserReplicationManager> InReplicationManager);
+
+	private:
+
+		/** Acts as the model of this view */
+		TSharedPtr<FMultiUserReplicationManager> ReplicationManager;
+		
+		/**
+		 * The editor view if one is being displayed.
+		 * ChildSlot keeps the reference alive if we're in state EMultiUserReplicationConnectionState::Connected.
+		 * @see ShowWidget_Connected
+		 */
+		TWeakPtr<ConcertClientSharedSlate::IReplicationEditorView> WeakEditorView;
+
+		/** Notification about in progress authority change, if any. */
+		TSharedPtr<SNotificationItem> AuthorityChangeNotification;
+		
+		/** Called when any of the streams change. */
+		void OnModelChanged() const;
+
+		void OnAuthorityRequestSent_AnyThread(const ConcertSyncClient::Replication::FAuthorityChangeRequest& Request);
+		void OnAuthorityResponseReceived_AnyThread(
+			const ConcertSyncClient::Replication::FAuthorityChangeRequest& Request,
+			const ConcertSyncClient::Replication::FAuthorityChangeResponse& Response
+			);
+	};
+}
diff --git a/Engine/Plugins/Developer/Concert/ConcertApp/MultiUserClient/Source/MultiUserClient/Private/Widgets/ActiveSession/Replication/SReplicationControlsTab.cpp b/Engine/Plugins/Developer/Concert/ConcertApp/MultiUserClient/Source/MultiUserClient/Private/Widgets/ActiveSession/Replication/SReplicationControlsTab.cpp
deleted file mode 100644
diff --git a/Engine/Plugins/Developer/Concert/ConcertApp/MultiUserClient/Source/MultiUserClient/Private/Widgets/ActiveSession/Replication/SReplicationControlsTab.h b/Engine/Plugins/Developer/Concert/ConcertApp/MultiUserClient/Source/MultiUserClient/Private/Widgets/ActiveSession/Replication/SReplicationControlsTab.h
deleted file mode 100644
diff --git a/Engine/Plugins/Developer/Concert/ConcertApp/MultiUserClient/Source/MultiUserClient/Private/Widgets/ActiveSession/Replication/SReplicationRootWidget.cpp b/Engine/Plugins/Developer/Concert/ConcertApp/MultiUserClient/Source/MultiUserClient/Private/Widgets/ActiveSession/Replication/SReplicationRootWidget.cpp
new file mode 100644
--- /dev/null
+++ b/Engine/Plugins/Developer/Concert/ConcertApp/MultiUserClient/Source/MultiUserClient/Private/Widgets/ActiveSession/Replication/SReplicationRootWidget.cpp
@@ -0,0 +1,129 @@
+// Copyright Epic Games, Inc. All Rights Reserved.
+
+#include "SReplicationRootWidget.h"
+
+#include "Widgets/ActiveSession/Replication/Joined/SReplicationJoinedView.h"
+#include "Replication/MultiUserReplicationManager.h"
+
+#include "Styling/AppStyle.h"
+#include "Widgets/Images/SThrobber.h"
+#include "Widgets/Input/SButton.h"
+#include "Widgets/Layout/SExpandableArea.h"
+#include "Widgets/Layout/SBorder.h"
+#include "Widgets/Layout/SBox.h"
+#include "Widgets/Layout/SSpacer.h"
+#include "Widgets/Layout/SSplitter.h"
+#include "Widgets/SBoxPanel.h"
+#include "Widgets/Text/STextBlock.h"
+
+#define LOCTEXT_NAMESPACE "SReplicationControlsTab"
+
+namespace UE::MultiUserClient
+{
+	void SReplicationRootWidget::Construct(const FArguments& InArgs, TSharedRef<FMultiUserReplicationManager> InReplicationManager)
+	{
+		ReplicationManager = MoveTemp(InReplicationManager);
+		
+		// Show different widget based on whether we've joined replication.
+		// We'll join replication automatically on session start but this may be rejected (in practice rejection should not happen though)
+		ReplicationManager->OnReplicationConnectionStateChanged().AddSP(this, &SReplicationRootWidget::OnReplicationConnectionStateChanged);
+		OnReplicationConnectionStateChanged(ReplicationManager->GetConnectionState());
+	}
+
+	void SReplicationRootWidget::OnReplicationConnectionStateChanged(EMultiUserReplicationConnectionState NewState)
+	{
+		switch (NewState)
+		{
+		case EMultiUserReplicationConnectionState::Connecting:
+			ShowWidget_Connecting();
+			break;
+		case EMultiUserReplicationConnectionState::Connected:
+			ShowWidget_Connected();
+			break;
+		case EMultiUserReplicationConnectionState::Disconnected:
+			ShowWidget_Disconnected();
+			break;
+		default: checkNoEntry();
+		}
+	}
+
+	void SReplicationRootWidget::ShowWidget_Connecting()
+	{
+		ChildSlot
+		[
+			SNew(SBox)
+			.Padding(0.f, 0.f, 0.f, 10.f)
+			.HAlign(HAlign_Center)
+			.VAlign(VAlign_Center)
+			.Clipping(EWidgetClipping::ClipToBounds)
+			[
+				SNew(SHorizontalBox)
+				
+				+SHorizontalBox::Slot()
+				[
+					SNew(SThrobber)
+				]
+
+				+SHorizontalBox::Slot()
+				.AutoWidth()
+				.HAlign(HAlign_Center)
+				[
+					SNew(STextBlock)
+					.Text(LOCTEXT("FailedToConnect", "Joining replication session"))
+					.Justification(ETextJustify::Center)
+				]
+			]
+		];
+	}
+
+	void SReplicationRootWidget::ShowWidget_Connected()
+	{
+		ChildSlot
+		[
+			SNew(SReplicationJoinedView, ReplicationManager.ToSharedRef())
+		];
+	}
+
+	void SReplicationRootWidget::ShowWidget_Disconnected()
+	{
+		ChildSlot
+		[
+			SNew(SBox)
+			.Padding(0.f, 0.f, 0.f, 10.f)
+			.HAlign(HAlign_Center)
+			.VAlign(VAlign_Center)
+			.Clipping(EWidgetClipping::ClipToBounds)
+			[
+				SNew(SVerticalBox)
+
+				+SVerticalBox::Slot()
+				.AutoHeight()
+				.HAlign(HAlign_Center)
+				[
+					SNew(STextBlock)
+					.Text(LOCTEXT("FailedToConnect", "Failed to join replication session."))
+					.Justification(ETextJustify::Center)
+				]
+
+				+SVerticalBox::Slot()
+				.Padding(0, 4, 0, 0)
+				.AutoHeight()
+				.HAlign(HAlign_Center)
+				[
+					SNew(SButton)
+					.OnClicked_Lambda([this]()
+					{
+						ReplicationManager->JoinReplicationSession();
+						return FReply::Handled();
+					})
+					[
+						SNew(STextBlock)
+						.Text(LOCTEXT("EnableLoggingVisibility.Button.Text", "Join replication session"))
+					]
+				]
+			]
+		];
+	}
+}
+
+#undef LOCTEXT_NAMESPACE
\ No newline at end of file
diff --git a/Engine/Plugins/Developer/Concert/ConcertApp/MultiUserClient/Source/MultiUserClient/Private/Widgets/ActiveSession/Replication/SReplicationRootWidget.h b/Engine/Plugins/Developer/Concert/ConcertApp/MultiUserClient/Source/MultiUserClient/Private/Widgets/ActiveSession/Replication/SReplicationRootWidget.h
new file mode 100644
--- /dev/null
+++ b/Engine/Plugins/Developer/Concert/ConcertApp/MultiUserClient/Source/MultiUserClient/Private/Widgets/ActiveSession/Replication/SReplicationRootWidget.h
@@ -0,0 +1,39 @@
+// Copyright Epic Games, Inc. All Rights Reserved.
+
+#pragma once
+
+#include "Widgets/DeclarativeSyntaxSupport.h"
+#include "Widgets/SCompoundWidget.h"
+#include "Widgets/Layout/SSplitter.h"
+
+namespace UE::MultiUserClient
+{
+	class FMultiUserReplicationManager;
+	
+	enum class EMultiUserReplicationConnectionState : uint8;
+
+	/** Root widget for replication in Multi-User session. */
+	class SReplicationRootWidget : public SCompoundWidget
+	{
+	public:
+
+		SLATE_BEGIN_ARGS(SReplicationRootWidget)
+		{}
+		SLATE_END_ARGS()
+
+		void Construct(const FArguments& InArgs, TSharedRef<FMultiUserReplicationManager> InReplicationManager);
+
+	private:
+
+		/** Manages the business logic which we represent. */
+		TSharedPtr<FMultiUserReplicationManager> ReplicationManager;
+		
+		/** Called when joining replication session state changes. */
+		void OnReplicationConnectionStateChanged(EMultiUserReplicationConnectionState);
+
+		void ShowWidget_Connecting();
+		void ShowWidget_Connected();
+		void ShowWidget_Disconnected();
+	};
+}
+
diff --git a/Engine/Plugins/Developer/Concert/ConcertApp/MultiUserClient/Source/MultiUserClient/Private/Widgets/ActiveSession/Replication/SReplicationTabWithWarningOverlay.cpp b/Engine/Plugins/Developer/Concert/ConcertApp/MultiUserClient/Source/MultiUserClient/Private/Widgets/ActiveSession/Replication/SReplicationTabWithWarningOverlay.cpp
--- a/Engine/Plugins/Developer/Concert/ConcertApp/MultiUserClient/Source/MultiUserClient/Private/Widgets/ActiveSession/Replication/SReplicationTabWithWarningOverlay.cpp
+++ b/Engine/Plugins/Developer/Concert/ConcertApp/MultiUserClient/Source/MultiUserClient/Private/Widgets/ActiveSession/Replication/SReplicationTabWithWarningOverlay.cpp
@@ -2,7 +2,7 @@
 
 #include "SReplicationTabWithWarningOverlay.h"
 
-#include "SReplicationControlsTab.h"
+#include "SReplicationRootWidget.h"
 
 #include "Widgets/Input/SButton.h"
 #include "Widgets/Layout/SBox.h"
@@ -23,7 +23,7 @@ namespace UE::MultiUserClient
 			// The real content
 			+SOverlay::Slot()
 			[
-				SNew(SReplicationControlsTab, InReplicationManager)
+				SNew(SReplicationRootWidget, InReplicationManager)
 			]
 
 			//  
diff --git a/Engine/Plugins/Developer/Concert/ConcertApp/MultiUserClient/Source/MultiUserReplicationEditor/Private/StreamEditor/ReplicationSessionPresetEditorToolkit.cpp b/Engine/Plugins/Developer/Concert/ConcertApp/MultiUserClient/Source/MultiUserReplicationEditor/Private/StreamEditor/ReplicationSessionPresetEditorToolkit.cpp
--- a/Engine/Plugins/Developer/Concert/ConcertApp/MultiUserClient/Source/MultiUserReplicationEditor/Private/StreamEditor/ReplicationSessionPresetEditorToolkit.cpp
+++ b/Engine/Plugins/Developer/Concert/ConcertApp/MultiUserClient/Source/MultiUserReplicationEditor/Private/StreamEditor/ReplicationSessionPresetEditorToolkit.cpp
@@ -4,9 +4,10 @@
 
 #include "Assets/MultiUserReplicationClientPreset.h"
 #include "Assets/MultiUserReplicationSessionPreset.h"
-#include "Replication/ReplicationWidgetFactories.h"
 #include "Replication/Editor/Model/Property/SelectPropertyFromUClassModel.h"
 #include "Replication/Editor/Model/Object/EditorObjectSelectionSourceModel.h"
+#include "Replication/Editor/View/IReplicationEditorView.h"
+#include "Replication/ReplicationWidgetFactories.h"
 
 #include "Framework/Docking/TabManager.h"
 #include "Framework/Docking/WorkspaceItem.h"
@@ -83,10 +84,11 @@ namespace UE::MultiUserReplicationEditor
 			}));
 		const TSharedRef<FEditorObjectSelectionSourceModel> ObjectSourceModel = MakeShared<FEditorObjectSelectionSourceModel>();
 		const TSharedRef<FSelectPropertyFromUClassModel> PropertySourceModel = MakeShared<FSelectPropertyFromUClassModel>();
+		const TSharedRef<IReplicationEditorView> EditorView = CreateEditor(FCreateEditorParams{ AssetReadWriteModel, ObjectSourceModel, PropertySourceModel });
 		return SNew(SDockTab)
 			.Label(LOCTEXT("BaseDetailsTitle", "Details"))
 			[
-				CreateEditor(FCreateEditorParams{ AssetReadWriteModel, ObjectSourceModel, PropertySourceModel })
+				EditorView
 			];
 	}
 }
diff --git a/Engine/Plugins/Developer/Concert/ConcertUI/ConcertClientSharedSlate/Source/ConcertClientSharedSlate/Private/Replication/Editor/View/SReplicationStreamEditor.cpp b/Engine/Plugins/Developer/Concert/ConcertUI/ConcertClientSharedSlate/Source/ConcertClientSharedSlate/Private/Replication/Editor/View/SReplicationStreamEditor.cpp
--- a/Engine/Plugins/Developer/Concert/ConcertUI/ConcertClientSharedSlate/Source/ConcertClientSharedSlate/Private/Replication/Editor/View/SReplicationStreamEditor.cpp
+++ b/Engine/Plugins/Developer/Concert/ConcertUI/ConcertClientSharedSlate/Source/ConcertClientSharedSlate/Private/Replication/Editor/View/SReplicationStreamEditor.cpp
@@ -17,10 +17,16 @@ namespace UE::ConcertClientSharedSlate
 
 			+SVerticalBox::Slot()
 			[
-				SNew(SPropertyReplicationSelectionEditor, Params.DataModel, Params.ObjectSource, Params.PropertySource)
+				SAssignNew(Editor, SPropertyReplicationSelectionEditor, Params.DataModel, Params.ObjectSource, Params.PropertySource)
 			]
 		];
 	}
+
+	void SReplicationStreamEditor::Refresh()
+	{
+		Editor->RefreshObjectData();
+		Editor->RefreshPropertyData();
+	}
 }
 
 
diff --git a/Engine/Plugins/Developer/Concert/ConcertUI/ConcertClientSharedSlate/Source/ConcertClientSharedSlate/Private/Replication/Editor/View/SReplicationStreamEditor.h b/Engine/Plugins/Developer/Concert/ConcertUI/ConcertClientSharedSlate/Source/ConcertClientSharedSlate/Private/Replication/Editor/View/SReplicationStreamEditor.h
--- a/Engine/Plugins/Developer/Concert/ConcertUI/ConcertClientSharedSlate/Source/ConcertClientSharedSlate/Private/Replication/Editor/View/SReplicationStreamEditor.h
+++ b/Engine/Plugins/Developer/Concert/ConcertUI/ConcertClientSharedSlate/Source/ConcertClientSharedSlate/Private/Replication/Editor/View/SReplicationStreamEditor.h
@@ -2,6 +2,7 @@
 
 #pragma once
 
+#include "Replication/Editor/View/IReplicationEditorView.h"
 #include "Widgets/DeclarativeSyntaxSupport.h"
 #include "Widgets/SCompoundWidget.h"
 
@@ -13,14 +14,23 @@ namespace UE::ConcertClientSharedSlate
 	
 	struct FCreateEditorParams;
 	
-	/** Root widget for editing UMultiUserReplicationStreamAsset. */
-	class CONCERTCLIENTSHAREDSLATE_API SReplicationStreamEditor : public SCompoundWidget
+	/** Root widget for editing a set of replication streams. */
+	class CONCERTCLIENTSHAREDSLATE_API SReplicationStreamEditor : public IReplicationEditorView
 	{
 	public:
+		
 		SLATE_BEGIN_ARGS(SReplicationStreamEditor)
 		{}
 		SLATE_END_ARGS()
 
 		void Construct(const FArguments& InArgs, const FCreateEditorParams& Params);
+
+		//~ Begin IReplicationEditorView Interface
+		virtual void Refresh() override;
+		//~ End IReplicationEditorView Interface
+
+	private:
+		
+		TSharedPtr<SPropertyReplicationSelectionEditor> Editor;
 	};
 }
diff --git a/Engine/Plugins/Developer/Concert/ConcertUI/ConcertClientSharedSlate/Source/ConcertClientSharedSlate/Private/Replication/ReplicationWidgetFactories.cpp b/Engine/Plugins/Developer/Concert/ConcertUI/ConcertClientSharedSlate/Source/ConcertClientSharedSlate/Private/Replication/ReplicationWidgetFactories.cpp
--- a/Engine/Plugins/Developer/Concert/ConcertUI/ConcertClientSharedSlate/Source/ConcertClientSharedSlate/Private/Replication/ReplicationWidgetFactories.cpp
+++ b/Engine/Plugins/Developer/Concert/ConcertUI/ConcertClientSharedSlate/Source/ConcertClientSharedSlate/Private/Replication/ReplicationWidgetFactories.cpp
@@ -7,7 +7,7 @@
 
 namespace UE::ConcertClientSharedSlate
 {
-	TSharedRef<SWidget> CreateEditor(FCreateEditorParams Params)
+	TSharedRef<IReplicationEditorView> CreateEditor(FCreateEditorParams Params)
 	{
 		return SNew(SReplicationStreamEditor, MoveTemp(Params));
 	}
diff --git a/Engine/Plugins/Developer/Concert/ConcertUI/ConcertClientSharedSlate/Source/ConcertClientSharedSlate/Public/Replication/Editor/View/IReplicationEditorView.h b/Engine/Plugins/Developer/Concert/ConcertUI/ConcertClientSharedSlate/Source/ConcertClientSharedSlate/Public/Replication/Editor/View/IReplicationEditorView.h
new file mode 100644
--- /dev/null
+++ b/Engine/Plugins/Developer/Concert/ConcertUI/ConcertClientSharedSlate/Source/ConcertClientSharedSlate/Public/Replication/Editor/View/IReplicationEditorView.h
@@ -0,0 +1,23 @@
+// Copyright Epic Games, Inc. All Rights Reserved.
+
+#pragma once
+
+#include "Templates/SharedPointer.h"
+#include "Widgets/SCompoundWidget.h"
+
+namespace UE::ConcertClientSharedSlate
+{
+	/**
+	 * Widget which edits a set of replication streams.
+	 * @see ReplicationWidgetFactories.h
+	 */
+	class CONCERTCLIENTSHAREDSLATE_API IReplicationEditorView : public SCompoundWidget
+	{
+	public:
+
+		/** Call after the data underlying the model was externally changed and needs to be redisplayed in the UI. */
+		virtual void Refresh() = 0;
+
+		virtual ~IReplicationEditorView() = default;
+	};
+}
\ No newline at end of file
diff --git a/Engine/Plugins/Developer/Concert/ConcertUI/ConcertClientSharedSlate/Source/ConcertClientSharedSlate/Public/Replication/ReplicationWidgetFactories.h b/Engine/Plugins/Developer/Concert/ConcertUI/ConcertClientSharedSlate/Source/ConcertClientSharedSlate/Public/Replication/ReplicationWidgetFactories.h
--- a/Engine/Plugins/Developer/Concert/ConcertUI/ConcertClientSharedSlate/Source/ConcertClientSharedSlate/Public/Replication/ReplicationWidgetFactories.h
+++ b/Engine/Plugins/Developer/Concert/ConcertUI/ConcertClientSharedSlate/Source/ConcertClientSharedSlate/Public/Replication/ReplicationWidgetFactories.h
@@ -14,18 +14,28 @@ struct FObjectReplicationMap;
 namespace UE::ConcertClientSharedSlate
 {
 	class IEditableObjectToPropertiesModel;
+	class IReplicationEditorView;
 	class IObjectSelectionSourceModel;
 	class IPropertySelectionSourceModel;
 		
 	struct FCreateEditorParams
 	{
-		/** The model that the editor is displaying */
+		/**
+		 * The model that the editor is displaying.
+		 * @note The view will keep a strong reference to this.
+		 */
 		TSharedRef<IEditableObjectToPropertiesModel> DataModel;
 
-		/** Determines the objects that can be added to the object list. */
+		/**
+		 * Determines the objects that can be added to the object list. 
+		 * @note The view will keep a strong reference to this.
+		 */
 		TSharedRef<IObjectSelectionSourceModel> ObjectSource;
 
-		/** Determines the properties that can be added to the property list. */
+		/**
+		 * Determines the properties that can be added to the property list.
+		 * @note The view will keep a strong reference to this.
+		 */
 		TSharedRef<IPropertySelectionSourceModel> PropertySource;
 
 		// TODO DP: Add way to add more columns
@@ -39,7 +49,7 @@ namespace UE::ConcertClientSharedSlate
 	 * 2. Subobjects, similar to component view (SSubobjectEditor): Shows subobjects of a root objects selected above; typically components.
 	 * 3. Properties, similar to details panel: Shows properties of the selected root object and / or subobjects.
 	 */
-	CONCERTCLIENTSHAREDSLATE_API TSharedRef<SWidget> CreateEditor(FCreateEditorParams Params);
+	CONCERTCLIENTSHAREDSLATE_API TSharedRef<IReplicationEditorView> CreateEditor(FCreateEditorParams Params);
 
 	/**
 	 * Creates a model that can be passed to CreateEditor.
