Fix duplicate loca keys

diff --git a/Engine/Plugins/Developer/Concert/ConcertApp/MultiUserClient/Source/MultiUserClient/Private/Widgets/ActiveSession/Replication/SReplicationRootWidget.cpp b/Engine/Plugins/Developer/Concert/ConcertApp/MultiUserClient/Source/MultiUserClient/Private/Widgets/ActiveSession/Replication/SReplicationRootWidget.cpp
--- a/Engine/Plugins/Developer/Concert/ConcertApp/MultiUserClient/Source/MultiUserClient/Private/Widgets/ActiveSession/Replication/SReplicationRootWidget.cpp
+++ b/Engine/Plugins/Developer/Concert/ConcertApp/MultiUserClient/Source/MultiUserClient/Private/Widgets/ActiveSession/Replication/SReplicationRootWidget.cpp
@@ -16,7 +16,7 @@
 #include "Widgets/SBoxPanel.h"
 #include "Widgets/Text/STextBlock.h"
 
-#define LOCTEXT_NAMESPACE "SReplicationControlsTab"
+#define LOCTEXT_NAMESPACE "SReplicationRootWidget"
 
 namespace UE::MultiUserClient
 {
@@ -69,7 +69,7 @@ namespace UE::MultiUserClient
 				.HAlign(HAlign_Center)
 				[
 					SNew(STextBlock)
-					.Text(LOCTEXT("FailedToConnect", "Joining replication session"))
+					.Text(LOCTEXT("SuccessfullyConnected", "Joining replication session"))
 					.Justification(ETextJustify::Center)
 				]
 			]
