Prepare code for displaying streams from other clients

diff --git a/Engine/Plugins/Developer/Concert/ConcertApp/MultiUserClient/Source/MultiUserClient/Private/Replication/Authority/AuthorityPolicy.cpp b/Engine/Plugins/Developer/Concert/ConcertApp/MultiUserClient/Source/MultiUserClient/Private/Replication/Authority/AuthorityPolicy.cpp
--- a/Engine/Plugins/Developer/Concert/ConcertApp/MultiUserClient/Source/MultiUserClient/Private/Replication/Authority/AuthorityPolicy.cpp
+++ b/Engine/Plugins/Developer/Concert/ConcertApp/MultiUserClient/Source/MultiUserClient/Private/Replication/Authority/AuthorityPolicy.cpp
@@ -2,7 +2,7 @@
 
 #include "AuthorityPolicy.h"
 
-#include "Replication/Stream/LocalClientStreamSynchronizer.h"
+#include "Replication/Stream/LocalStreamChangeTracker.h"
 
 #include "IConcertSyncClient.h"
 
@@ -39,17 +39,17 @@ namespace UE::MultiUserClient
 		}
 	}
 	
-	FAuthorityPolicy::FAuthorityPolicy(FLocalClientStreamSynchronizer& InLocalClientManager, TSharedRef<IConcertSyncClient> InClient)
-		: LocalClientManager(InLocalClientManager)
-		, Client(MoveTemp(InClient))
+	FAuthorityPolicy::FAuthorityPolicy(TSharedRef<IConcertSyncClient> InClient, TSharedRef<IClientStreamSynchronizer> InStreamSynchronizer)
+		: Client(MoveTemp(InClient))
+		, StreamSynchronizer(MoveTemp(InStreamSynchronizer))
 	{
 		// After submitting new objects, immediately request authority over it.
-		LocalClientManager.OnChangesAccepted_AnyThread().AddRaw(this, &FAuthorityPolicy::OnChangesAccepted_AnyThread);
+		StreamSynchronizer->OnChangesAccepted_AnyThread().AddRaw(this, &FAuthorityPolicy::OnChangesAccepted_AnyThread);
 	}
 
 	FAuthorityPolicy::~FAuthorityPolicy()
 	{
-		LocalClientManager.OnChangesAccepted_AnyThread().RemoveAll(this);
+		StreamSynchronizer->OnChangesAccepted_AnyThread().RemoveAll(this);
 	}
 
 	void FAuthorityPolicy::OnChangesAccepted_AnyThread(
diff --git a/Engine/Plugins/Developer/Concert/ConcertApp/MultiUserClient/Source/MultiUserClient/Private/Replication/Authority/AuthorityPolicy.h b/Engine/Plugins/Developer/Concert/ConcertApp/MultiUserClient/Source/MultiUserClient/Private/Replication/Authority/AuthorityPolicy.h
--- a/Engine/Plugins/Developer/Concert/ConcertApp/MultiUserClient/Source/MultiUserClient/Private/Replication/Authority/AuthorityPolicy.h
+++ b/Engine/Plugins/Developer/Concert/ConcertApp/MultiUserClient/Source/MultiUserClient/Private/Replication/Authority/AuthorityPolicy.h
@@ -16,7 +16,8 @@ namespace UE::ConcertSyncClient::Replication
 }
 namespace UE::MultiUserClient
 {
-	class FLocalClientStreamSynchronizer;
+	class IClientStreamSynchronizer;
+	class FLocalStreamChangeTracker;
 }
 
 namespace UE::MultiUserClient
@@ -32,7 +33,7 @@ namespace UE::MultiUserClient
 	{
 	public:
 		
-		FAuthorityPolicy(FLocalClientStreamSynchronizer& InLocalClientManager, TSharedRef<IConcertSyncClient> InClient);
+		FAuthorityPolicy(TSharedRef<IConcertSyncClient> InClient, TSharedRef<IClientStreamSynchronizer> InStreamSynchronizer);
 		~FAuthorityPolicy();
 
 		DECLARE_MULTICAST_DELEGATE_OneParam(FOnAuthorityRequestSent,
@@ -53,11 +54,10 @@ namespace UE::MultiUserClient
 		/** Referenced by the authority requests to detect destruction of FAuthorityPolicy */
 		const TSharedRef<FToken> LifetimeToken = FToken::Make();
 
-		/** Informs us when changes have been successfully submitted to the server. */
-		FLocalClientStreamSynchronizer& LocalClientManager;
-
 		/** Used to send authority requests to the server. */
 		const TSharedRef<IConcertSyncClient> Client;
+		/** Informs us when changes have been successfully submitted to the server. */
+		TSharedRef<IClientStreamSynchronizer> StreamSynchronizer;
 
 		FOnAuthorityRequestSent OnAuthorityRequestSentDelegate;
 		FOnAuthorityResponseReceived OnAuthorityResponseReceivedDelegate;
diff --git a/Engine/Plugins/Developer/Concert/ConcertApp/MultiUserClient/Source/MultiUserClient/Private/Replication/Client/ReplicationClient.cpp b/Engine/Plugins/Developer/Concert/ConcertApp/MultiUserClient/Source/MultiUserClient/Private/Replication/Client/ReplicationClient.cpp
new file mode 100644
--- /dev/null
+++ b/Engine/Plugins/Developer/Concert/ConcertApp/MultiUserClient/Source/MultiUserClient/Private/Replication/Client/ReplicationClient.cpp
@@ -0,0 +1,50 @@
+// Copyright Epic Games, Inc. All Rights Reserved.
+
+#include "ReplicationClient.h"
+
+#include "Assets/MultiUserReplicationClientPreset.h"
+#include "Replication/ReplicationWidgetFactories.h"
+#include "Replication/Editor/Model/IEditableObjectToPropertiesModel.h"
+#include "Replication/Stream/LocalStreamChangeTracker.h"
+
+namespace UE::MultiUserClient
+{
+	FReplicationClient::FReplicationClient(
+		UMultiUserReplicationClientPreset& InSessionContent,
+		TSharedRef<IConcertSyncClient> InClient,
+		TSharedRef<IClientStreamSynchronizer> InStreamSynchronizer
+		)
+		: ClientContentStorage(&InSessionContent)
+		, StreamSynchronizer(MoveTemp(InStreamSynchronizer))
+		, LocalClientEditModel(ConcertClientSharedSlate::CreatePropertySelectionModel(*ClientContentStorage->Stream, ClientContentStorage->Stream->MakeReplicationMapGetterAttribute()))
+		, LocalClientStreamDiffer(
+			StreamSynchronizer,
+			ClientContentStorage->Stream->MakeReplicationMapGetterAttribute(),
+			FLocalStreamChangeTracker::FOnModifyReplicationMap::CreateLambda([this](){ ClientContentStorage->Stream->Modify(); })
+			)
+	{
+		LocalClientEditModel->OnObjectsChanged().AddRaw(this, &FReplicationClient::OnObjectsChanged);
+		LocalClientEditModel->OnPropertiesChanged().AddRaw(this, &FReplicationClient::OnPropertiesChanged);
+		
+		LocalClientStreamDiffer.OnChangesReverted_GameThread().AddLambda([this]()
+		{
+			OnModelExternallyChangedDelegate.Broadcast();
+		});
+	}
+
+	void FReplicationClient::OnObjectsChanged(
+		TConstArrayView<UObject*> Objects,
+		TConstArrayView<FSoftObjectPath> SoftObjectPaths,
+		ConcertClientSharedSlate::EReplicatedObjectChangeReason ReplicatedObjectChangeReason
+		)
+	{
+		// Could improve performance by just considering what actually changed instead of doing a full rebuild
+		LocalClientStreamDiffer.RefreshChangesCache();
+	}
+
+	void FReplicationClient::OnPropertiesChanged()
+	{
+		// Could improve performance by just considering what actually changed instead of doing a full rebuild
+		LocalClientStreamDiffer.RefreshChangesCache();
+	}
+}
diff --git a/Engine/Plugins/Developer/Concert/ConcertApp/MultiUserClient/Source/MultiUserClient/Private/Replication/Client/ReplicationClient.h b/Engine/Plugins/Developer/Concert/ConcertApp/MultiUserClient/Source/MultiUserClient/Private/Replication/Client/ReplicationClient.h
new file mode 100644
--- /dev/null
+++ b/Engine/Plugins/Developer/Concert/ConcertApp/MultiUserClient/Source/MultiUserClient/Private/Replication/Client/ReplicationClient.h
@@ -0,0 +1,79 @@
+// Copyright Epic Games, Inc. All Rights Reserved.
+
+#pragma once
+
+#include "Replication/Stream/LocalStreamChangeTracker.h"
+#include "Templates/SharedPointer.h"
+#include "Templates/UnrealTemplate.h"
+
+class IConcertSyncClient;
+class UMultiUserReplicationClientPreset;
+
+namespace UE::ConcertClientSharedSlate
+{
+	class IEditableObjectToPropertiesModel;
+	enum class EReplicatedObjectChangeReason : uint8;
+}
+
+namespace UE::MultiUserClient
+{
+	/** Holds on to info about a local or remote client. */
+	class FReplicationClient : public FNoncopyable
+	{
+	public:
+
+		FReplicationClient(
+			UMultiUserReplicationClientPreset& InSessionContent,
+			TSharedRef<IConcertSyncClient> InClient,
+			TSharedRef<IClientStreamSynchronizer> InStreamSynchronizer
+			);
+
+		UMultiUserReplicationClientPreset* GetClientContent() const { return ClientContentStorage; }
+		TSharedRef<ConcertClientSharedSlate::IEditableObjectToPropertiesModel> GetClientEditModel() const { return LocalClientEditModel; }
+		TSharedRef<IClientStreamSynchronizer> GetStreamSynchronizer() const { return StreamSynchronizer; } 
+		
+		const FLocalStreamChangeTracker& GetDiffer() const { return LocalClientStreamDiffer; }
+		FLocalStreamChangeTracker& GetDiffer() { return LocalClientStreamDiffer; }
+
+		/**
+		 * Called when the data underlying the model has changed externally. Since the change was not caused by the model,
+		 * its events, like IEditableObjectToPropertiesModel::OnObjectsChanged, were not called.
+		 * 
+		 * Subscribers are intended to call IReplicationEditorView::Refresh() in response.
+		 * 
+		 * Examples: Remote client changed their stream, remote client joined with streams,
+		 * or the local client reverted a change that was rejected by the server.
+		 */
+		DECLARE_MULTICAST_DELEGATE(FOnModelExternallyChanged);
+		FOnModelExternallyChanged& OnModelExternallyChanged_GameThread() { return OnModelExternallyChangedDelegate; }
+		
+	private:
+		
+		/** The state of the server is synched up with this object and displayed in the UI. */
+		TObjectPtr<UMultiUserReplicationClientPreset> ClientContentStorage;
+		
+		/** Keeps the client's state on the server in sync. */
+		TSharedRef<IClientStreamSynchronizer> StreamSynchronizer;
+		
+		/**
+		 * Used to detect changes made to the client's config by the local editor.
+		 * Those changes can later be applied to the client.
+		 * 
+		 * The SessionContent UObject must be kept alive as long as this model lives. 
+		 * It should be noted that the UI keeps a strong reference to this model.
+		 * The UI is destroyed right after this FClientStreamRepository.
+		 * @see FMultiUserReplicationManager::OnLeaveSession
+		 */
+		TSharedRef<ConcertClientSharedSlate::IEditableObjectToPropertiesModel> LocalClientEditModel;
+		/** Tracks changes made to server's state of the client's streams and prepares to upload them using StreamSynchronizer. */
+		FLocalStreamChangeTracker LocalClientStreamDiffer;
+
+		/** Called when the data underlying the model has changed (and the UI needs to be refreshed). */
+		FOnModelExternallyChanged OnModelExternallyChangedDelegate;
+
+		// Respond to model changing
+		void OnObjectsChanged(TConstArrayView<UObject*> Objects, TConstArrayView<FSoftObjectPath> SoftObjectPaths, ConcertClientSharedSlate::EReplicatedObjectChangeReason ReplicatedObjectChangeReason);
+		void OnPropertiesChanged();
+	};
+}
+
diff --git a/Engine/Plugins/Developer/Concert/ConcertApp/MultiUserClient/Source/MultiUserClient/Private/Replication/Client/ReplicationClientManager.cpp b/Engine/Plugins/Developer/Concert/ConcertApp/MultiUserClient/Source/MultiUserClient/Private/Replication/Client/ReplicationClientManager.cpp
new file mode 100644
--- /dev/null
+++ b/Engine/Plugins/Developer/Concert/ConcertApp/MultiUserClient/Source/MultiUserClient/Private/Replication/Client/ReplicationClientManager.cpp
@@ -0,0 +1,26 @@
+// Copyright Epic Games, Inc. All Rights Reserved.
+
+#include "ReplicationClientManager.h"
+
+#include "Assets/MultiUserReplicationSessionPreset.h"
+#include "Replication/Stream/LocalClientStreamSynchronizer.h"
+
+#include "UObject/Package.h"
+#include "UObject/UObjectGlobals.h"
+
+namespace UE::MultiUserClient
+{
+	FReplicationClientManager::FReplicationClientManager(TSharedRef<IConcertSyncClient> InClient)
+		: SessionContent(NewObject<UMultiUserReplicationSessionPreset>(GetTransientPackage(), NAME_None, RF_Transient)) 
+		, LocalClient([this, &InClient]()
+		{
+			UMultiUserReplicationClientPreset* ClientPreset = SessionContent->AddClient();
+			return FReplicationClient(*ClientPreset, InClient, MakeShared<FLocalClientStreamSynchronizer>(InClient, ClientPreset->Stream->StreamId));
+		}())
+	{}
+
+	void FReplicationClientManager::AddReferencedObjects(FReferenceCollector& Collector)
+	{
+		Collector.AddReferencedObject(SessionContent);
+	}
+}
diff --git a/Engine/Plugins/Developer/Concert/ConcertApp/MultiUserClient/Source/MultiUserClient/Private/Replication/Client/ReplicationClientManager.h b/Engine/Plugins/Developer/Concert/ConcertApp/MultiUserClient/Source/MultiUserClient/Private/Replication/Client/ReplicationClientManager.h
new file mode 100644
--- /dev/null
+++ b/Engine/Plugins/Developer/Concert/ConcertApp/MultiUserClient/Source/MultiUserClient/Private/Replication/Client/ReplicationClientManager.h
@@ -0,0 +1,48 @@
+// Copyright Epic Games, Inc. All Rights Reserved.
+
+#pragma once
+
+#include "ReplicationClient.h"
+#include "Assets/MultiUserReplicationSessionPreset.h"
+#include "Replication/Editor/Model/IEditableObjectToPropertiesModel.h"
+#include "UObject/GCObject.h"
+#include "Templates/UnrealTemplate.h"
+
+class IConcertSyncClient;
+
+namespace UE::ConcertClientSharedSlate
+{
+	class IEditableObjectToPropertiesModel;
+}
+
+namespace UE::MultiUserClient
+{
+	class FReplicationClient;
+	class FLocalStreamChangeTracker;
+
+	/** Keeps track of connected clients synchronizing their stream data in a UMultiUserReplicationSessionPreset. */
+	class FReplicationClientManager
+		: public FGCObject
+		, public FNoncopyable
+	{
+	public:
+
+		FReplicationClientManager(TSharedRef<IConcertSyncClient> InClient);
+
+		FReplicationClient& GetLocalClient() { return LocalClient; }
+		
+		//~ Begin FGCObject Interface
+		virtual void AddReferencedObjects(FReferenceCollector& Collector) override;
+		virtual FString GetReferencerName() const override { return TEXT("FReplicationStreamSynchronizer"); }
+		//~ End FGCObject Interface
+		
+	private:
+		
+		/** The state of the server is synched up with this object and displayed in the UI. */
+		TObjectPtr<UMultiUserReplicationSessionPreset> SessionContent;
+
+		/** Manages the local client */
+		FReplicationClient LocalClient;
+	};
+}
+
diff --git a/Engine/Plugins/Developer/Concert/ConcertApp/MultiUserClient/Source/MultiUserClient/Private/Replication/MultiUserReplicationManager.cpp b/Engine/Plugins/Developer/Concert/ConcertApp/MultiUserClient/Source/MultiUserClient/Private/Replication/MultiUserReplicationManager.cpp
--- a/Engine/Plugins/Developer/Concert/ConcertApp/MultiUserClient/Source/MultiUserClient/Private/Replication/MultiUserReplicationManager.cpp
+++ b/Engine/Plugins/Developer/Concert/ConcertApp/MultiUserClient/Source/MultiUserClient/Private/Replication/MultiUserReplicationManager.cpp
@@ -100,7 +100,7 @@ namespace UE::MultiUserClient
 	}
 
 	FMultiUserReplicationManager::FConnectedState::FConnectedState(TSharedRef<IConcertSyncClient> InClient)
-		: StreamSynchronizer(InClient)
-		, AuthorityPolicy(StreamSynchronizer.GetDiffer(), InClient)
+		: ClientManager(InClient)
+		, AuthorityPolicy(InClient, ClientManager.GetLocalClient().GetStreamSynchronizer())
 	{}
 }
diff --git a/Engine/Plugins/Developer/Concert/ConcertApp/MultiUserClient/Source/MultiUserClient/Private/Replication/MultiUserReplicationManager.h b/Engine/Plugins/Developer/Concert/ConcertApp/MultiUserClient/Source/MultiUserClient/Private/Replication/MultiUserReplicationManager.h
--- a/Engine/Plugins/Developer/Concert/ConcertApp/MultiUserClient/Source/MultiUserClient/Private/Replication/MultiUserReplicationManager.h
+++ b/Engine/Plugins/Developer/Concert/ConcertApp/MultiUserClient/Source/MultiUserClient/Private/Replication/MultiUserReplicationManager.h
@@ -3,8 +3,9 @@
 #pragma once
 
 #include "Replication/Authority/AuthorityPolicy.h"
-#include "Replication/Stream/ClientStreamRepository.h"
+
 #include "IConcertSession.h"
+#include "Client/ReplicationClientManager.h"
 
 #include "Misc/Optional.h"
 #include "Templates/SharedPointer.h"
@@ -18,6 +19,7 @@ namespace UE::ConcertSyncClient::Replication
 {
 	struct FJoinReplicatedSessionResult;
 }
+
 namespace UE::ConcertClientSharedSlate
 {
 	class IEditableObjectToPropertiesModel;
@@ -61,9 +63,9 @@ namespace UE::MultiUserClient
 		 */
 		void JoinReplicationSession();
 
-		/** @note You're not supposed to keep any reference to StreamSynchronizer since it can become invalid depending on connection state. */
-		FClientStreamRepository* GetStreamSynchronizer() { return ConnectedState ? &ConnectedState->StreamSynchronizer : nullptr; }
-		const FClientStreamRepository* GetStreamSynchronizer() const { return ConnectedState ? &ConnectedState->StreamSynchronizer : nullptr; }
+		/** @note You're not supposed to keep any reference to the ClientManager since it can become invalid depending on connection state. */
+		FReplicationClientManager* GetClientManager() { return ConnectedState ? &ConnectedState->ClientManager : nullptr; }
+		const FReplicationClientManager* GetClientManager() const { return ConnectedState ? &ConnectedState->ClientManager : nullptr; }
 
 		/** @note You're not supposed to keep any reference to AuthorityPolicy since it can become invalid depending on connection state. */
 		FAuthorityPolicy* GetAuthorityPolicy() { return ConnectedState ? & ConnectedState->AuthorityPolicy : nullptr; }
@@ -93,7 +95,7 @@ namespace UE::MultiUserClient
 			 *
 			 * Only valid when ConnectionState == EMultiUserReplicationConnectionState::Connected.
 			 */
-			FClientStreamRepository StreamSynchronizer;
+			FReplicationClientManager ClientManager;
 
 			/** Manages the client's authority like automatically taking ownership over newly submitted objects. */
 			FAuthorityPolicy AuthorityPolicy;
diff --git a/Engine/Plugins/Developer/Concert/ConcertApp/MultiUserClient/Source/MultiUserClient/Private/Replication/Stream/ClientStreamRepository.cpp b/Engine/Plugins/Developer/Concert/ConcertApp/MultiUserClient/Source/MultiUserClient/Private/Replication/Stream/ClientStreamRepository.cpp
deleted file mode 100644
diff --git a/Engine/Plugins/Developer/Concert/ConcertApp/MultiUserClient/Source/MultiUserClient/Private/Replication/Stream/ClientStreamRepository.h b/Engine/Plugins/Developer/Concert/ConcertApp/MultiUserClient/Source/MultiUserClient/Private/Replication/Stream/ClientStreamRepository.h
deleted file mode 100644
diff --git a/Engine/Plugins/Developer/Concert/ConcertApp/MultiUserClient/Source/MultiUserClient/Private/Replication/Stream/IClientStreamSynchronizer.h b/Engine/Plugins/Developer/Concert/ConcertApp/MultiUserClient/Source/MultiUserClient/Private/Replication/Stream/IClientStreamSynchronizer.h
new file mode 100644
--- /dev/null
+++ b/Engine/Plugins/Developer/Concert/ConcertApp/MultiUserClient/Source/MultiUserClient/Private/Replication/Stream/IClientStreamSynchronizer.h
@@ -0,0 +1,82 @@
+// Copyright Epic Games, Inc. All Rights Reserved.
+
+#pragma once
+
+#include "Replication/Data/ObjectIds.h"
+#include "Replication/Messages/ChangeStream.h"
+
+#include "Async/Future.h"
+#include "Replication/IConcertClientReplicationManager.h"
+
+struct FObjectReplicationMap;
+
+namespace UE::ConcertSyncClient::Replication
+{
+	struct FChangeStreamRequest;
+	struct FChangeStreamResponse;
+}
+
+namespace UE::MultiUserClient
+{
+	struct FStreamChangelist
+	{
+		TSet<FObjectInStreamID> ObjectsToRemove;
+		TMap<FObjectInStreamID, FConcertReplication_ChangeStream_PutObject> ObjectsToPut;
+	};
+	
+	enum class ESubmitChangesErrorCode
+	{
+		/** Changes were accepted by the server */
+		Success,
+		/** Skipped because another request is in progress. */
+		CannotSendRequest
+	};
+
+	struct FCompletedChangeSubmission
+	{
+		ConcertSyncClient::Replication::FChangeStreamRequest Request;
+		ConcertSyncClient::Replication::FChangeStreamResponse Response;
+	};
+
+	struct FSubmitChangesResult
+	{
+		/** Error code of the submission. Determines whether Response is valid. */
+		ESubmitChangesErrorCode ErrorCode;
+
+		/** Valid if ErrorCode != ESubmitChangesErrorCode::AlreadyInProgress. */
+		TOptional<FCompletedChangeSubmission> SubmissionInfo;
+	};
+	
+	/**
+	 * Keeps track of a client's registered streams.
+	 * Exposes ways for observing and mutating streams.
+	 *
+	 * The implementation varies depending on whether the client corresponds to the local or a remote editor.
+	 */
+	class IClientStreamSynchronizer
+	{
+	public:
+
+		/**
+		 * Requests to change the stream. The future executes once the changes have been accepted on the server.
+		 * @note This future can execute on any thread.
+		 */
+		virtual TFuture<FSubmitChangesResult> SubmitChanges(const FStreamChangelist& Changelist) = 0;
+		/** @return Whether it's currently valid to call SubmitChanges. For example, it is not valid to call while there is already request in progress. */
+		virtual bool CanMakeSubmitRequest() const = 0;
+
+		/** @return The managed client's stream ID. */
+		virtual FGuid GetStreamId() const = 0;
+		/** @return What local instance thinks the client's server state is. */
+		virtual const FObjectReplicationMap& GetServerState() const = 0;
+
+		DECLARE_MULTICAST_DELEGATE_TwoParams(FOnChangesAccepted, const FObjectReplicationMap& OldServerState, const ConcertSyncClient::Replication::FChangeStreamRequest& AcceptedRequest);
+		/** @return Event executed when a change request that was in transit was accepted by the server. This is executed before OnServerStateSynched: the local confirmed state is not updated, yet. */
+		virtual FOnChangesAccepted& OnChangesAccepted_AnyThread() = 0;
+		DECLARE_MULTICAST_DELEGATE(FOnServerStateChanged);
+		/** @return Event executed when the result of GetServerState has been synched. Called after OnChangesAccepted_AnyThread. */
+		virtual FOnServerStateChanged& OnServerStateSynched_AnyThread() = 0;
+		
+		virtual ~IClientStreamSynchronizer() = default;
+	};
+}
diff --git a/Engine/Plugins/Developer/Concert/ConcertApp/MultiUserClient/Source/MultiUserClient/Private/Replication/Stream/LocalClientStreamSynchronizer.cpp b/Engine/Plugins/Developer/Concert/ConcertApp/MultiUserClient/Source/MultiUserClient/Private/Replication/Stream/LocalClientStreamSynchronizer.cpp
--- a/Engine/Plugins/Developer/Concert/ConcertApp/MultiUserClient/Source/MultiUserClient/Private/Replication/Stream/LocalClientStreamSynchronizer.cpp
+++ b/Engine/Plugins/Developer/Concert/ConcertApp/MultiUserClient/Source/MultiUserClient/Private/Replication/Stream/LocalClientStreamSynchronizer.cpp
@@ -14,45 +14,29 @@ namespace UE::MultiUserClient
 {
 	FLocalClientStreamSynchronizer::FLocalClientStreamSynchronizer(
 		TSharedRef<IConcertSyncClient> InLocalClient,
-		const FGuid& InLocalClientStreamId,
-		TAttribute<FObjectReplicationMap*> InStreamWithInProgressChangesAttribute,
-		FOnModifyReplicationMap InOnModifyReplicationMapDelegate
+		const FGuid& InLocalClientStreamId
 		)
 		: LocalClient(MoveTemp(InLocalClient))
 		, LocalClientStreamId(InLocalClientStreamId)
-		, StreamWithInProgressChangesAttribute(MoveTemp(InStreamWithInProgressChangesAttribute))
-		, OnModifyReplicationMapDelegate(MoveTemp(InOnModifyReplicationMapDelegate))
-	{
-		check(InStreamWithInProgressChangesAttribute.IsBound() || InStreamWithInProgressChangesAttribute.IsSet());
-	}
+	{}
 
-	void FLocalClientStreamSynchronizer::RefreshChangesCache()
-	{
-		CachedDeltaChange = DiffChanges();
-		RefreshWarnings();
-	}
-
-	TFuture<FLocalClientStreamSynchronizer::FSubmissionResult> FLocalClientStreamSynchronizer::SubmitChanges(bool bRefreshChanges)
+	TFuture<FSubmitChangesResult> FLocalClientStreamSynchronizer::SubmitChanges(const FStreamChangelist& Changelist)
 	{
 		using namespace ConcertSyncClient::Replication;
 
 		IConcertClientReplicationManager* ReplicationManager = LocalClient->GetReplicationManager();
-		if (!ensure(ReplicationManager) || !CanSubmitChanges())
-		{
-			return MakeFulfilledPromise<FSubmissionResult>(FSubmissionResult{}).GetFuture();
-		}
-
-		if (bRefreshChanges)
+		if (!ensure(ReplicationManager) || !CanMakeSubmitRequest())
 		{
-			RefreshChangesCache();
+			return MakeFulfilledPromise<FSubmitChangesResult>(FSubmitChangesResult{ ESubmitChangesErrorCode::CannotSendRequest }).GetFuture();
 		}
-		ChangeInTransit = CachedDeltaChange;
-		FChangeStreamRequest Request = BuildChangeRequest();
+		
+		ChangeInTransit = Changelist;
+		FChangeStreamRequest Request = BuildChangeRequest(Changelist);
 		
 		return ReplicationManager->ChangeStream(Request)
 			.Next([this, DestructionDetection = LifetimeToken->AsWeak(), Request](FChangeStreamResponse&& Response)
 			{
-				const FSubmissionResult SubmissionResult { Request, Response };
+				const FSubmitChangesResult SubmissionResult { ESubmitChangesErrorCode::Success, { FCompletedChangeSubmission{Request, Response } } };
 				// The request might execute after we're destroyed, e.g. by leaving session while request is on the way.
 				// In that case, the Concert session triggers the OnSessionConnectionChanged which destroys us. Only after that, all the requests are timed out.
 				if (!DestructionDetection.IsValid())
@@ -60,7 +44,7 @@ namespace UE::MultiUserClient
 					return SubmissionResult;
 				}
 
-				const FChangelist ChangeThatWasInTransit = MoveTemp(ChangeInTransit.GetValue());
+				const FStreamChangelist ChangeThatWasInTransit = MoveTemp(ChangeInTransit.GetValue());
 				ChangeInTransit.Reset();
 				
 				if (Response.IsSuccess())
@@ -74,111 +58,17 @@ namespace UE::MultiUserClient
 			});
 	}
 
-	void FLocalClientStreamSynchronizer::RevertCachedChanges()
-	{
-		if (HasRevertableLocalChanges())
-		{
-			FScopedTransaction Transaction(LOCTEXT("RevertCachedChanges", "Revert replication changes"));
-			OnModifyReplicationMapDelegate.ExecuteIfBound();
-			
-			*StreamWithInProgressChangesAttribute.Get() = ConfirmedServerState;
-			CachedDeltaChange = {};
-
-			check(IsInGameThread());
-			OnChangesRevertedDelegate.Broadcast();
-			RefreshChangesCache();
-		}
-	}
-
-	bool FLocalClientStreamSynchronizer::HasSubmittableLocalChanges() const
-	{
-		return !CachedDeltaChange.ObjectsToPut.IsEmpty() || !CachedDeltaChange.ObjectsToRemove.IsEmpty();
-	}
-
-	bool FLocalClientStreamSynchronizer::HasRevertableLocalChanges() const
-	{
-		return HasSubmittableLocalChanges() || HasAnyWarnings();
-	}
-
-	bool FLocalClientStreamSynchronizer::IsChangeRequestInTransit() const
-	{
-		return ChangeInTransit.IsSet();
-	}
-
-	EObjectWarningFlags FLocalClientStreamSynchronizer::GetObjectWarningFlags(const FSoftObjectPath& ObjectPath)
-	{
-		const EObjectWarningFlags* Flags = ObjectsWithWarnings.Find(ObjectPath);
-		return Flags ? *Flags : EObjectWarningFlags::Ok;
-	}
-
-	bool FLocalClientStreamSynchronizer::HasAnyWarnings() const
-	{
-		return !ObjectsWithWarnings.IsEmpty();
-	}
-
-	FLocalClientStreamSynchronizer::EObjectChangeType FLocalClientStreamSynchronizer::GetObjectChanges(const FSoftObjectPath& Object) const
-	{
-		// TODO DP:
-		return EObjectChangeType::NoChange;
-	}
-
-	FLocalClientStreamSynchronizer::EPropertyChangeType FLocalClientStreamSynchronizer::GetPropertyChanges(const FSoftObjectPath& Object, const FConcertPropertyChain& PropertyChain) const
-	{
-		// TODO DP:
-		return EPropertyChangeType::NoChange;
-	}
-	
-	FLocalClientStreamSynchronizer::FChangelist FLocalClientStreamSynchronizer::DiffChanges(
-		const FGuid& StreamId,
-		const FObjectReplicationMap& Base,
-		const FObjectReplicationMap& Changed
-		)
-	{
-		// BuildRequestFromDiff does not tolerate any invalid entries (like empty properties, which the UI generates right after you add an object to the list)
-		FObjectReplicationMap Cleansed = Changed;
-		ConcertSyncCore::Replication::ChangeStreamUtils::IterateInvalidEntries(Changed, [&Cleansed](const FSoftObjectPath& InvalidObject, const FReplicatedObjectInfo&)
-		{
-			Cleansed.ReplicatedObjects.Remove(InvalidObject);
-			return EBreakBehavior::Continue;
-		});
-
-		FConcertReplication_ChangeStream_Request Request = ConcertSyncCore::Replication::ChangeStreamUtils::BuildRequestFromDiff(StreamId, Base, Cleansed);
-		return { MoveTemp(Request.ObjectsToRemove), MoveTemp(Request.ObjectsToPut) };
-	}
-
-	void FLocalClientStreamSynchronizer::RefreshWarnings()
-	{
-		ObjectsWithWarnings.Empty();
-		const FObjectReplicationMap* Map = StreamWithInProgressChangesAttribute.Get();
-		if (!Map)
-		{
-			return;
-		}
-		
-		ConcertSyncCore::Replication::ChangeStreamUtils::IterateInvalidEntries(*Map, [this](const FSoftObjectPath& InvalidObject, const FReplicatedObjectInfo& Info)
-		{
-			EObjectWarningFlags Flags = EObjectWarningFlags::Ok;
-			if (Info.PropertySelection.ReplicatedProperties.IsEmpty())
-			{
-				Flags |= EObjectWarningFlags::MissingProperties;
-			}
-			
-			ObjectsWithWarnings.Add(InvalidObject, Flags);
-			return EBreakBehavior::Continue;
-		});
-	}
-
-	ConcertSyncClient::Replication::FChangeStreamRequest FLocalClientStreamSynchronizer::BuildChangeRequest() const
+	ConcertSyncClient::Replication::FChangeStreamRequest FLocalClientStreamSynchronizer::BuildChangeRequest(const FStreamChangelist& Changelist) const
 	{
 		switch (ComputeNextRequestType())
 		{
-		case EChangeRequestType::CreateNewStream: return BuildChangeRequest_CreateNewStream(LocalClientStreamId, CachedDeltaChange);
-		case EChangeRequestType::UpdateExistingStream: return BuildChangeRequest_UpdateExistingStream(CachedDeltaChange);
+		case EChangeRequestType::CreateNewStream: return BuildChangeRequest_CreateNewStream(LocalClientStreamId, Changelist);
+		case EChangeRequestType::UpdateExistingStream: return BuildChangeRequest_UpdateExistingStream(Changelist);
 		default: checkNoEntry(); return {};
 		}
 	}
 
-	ConcertSyncClient::Replication::FChangeStreamRequest FLocalClientStreamSynchronizer::BuildChangeRequest_CreateNewStream(const FGuid& StreamId, const FChangelist& FromChangelist)
+	ConcertSyncClient::Replication::FChangeStreamRequest FLocalClientStreamSynchronizer::BuildChangeRequest_CreateNewStream(const FGuid& StreamId, const FStreamChangelist& FromChangelist)
 	{
 		const TMap<FObjectInStreamID, FConcertReplication_ChangeStream_PutObject>& ObjectsToPut = FromChangelist.ObjectsToPut;
 		
@@ -206,7 +96,7 @@ namespace UE::MultiUserClient
 		return Request;
 	}
 
-	ConcertSyncClient::Replication::FChangeStreamRequest FLocalClientStreamSynchronizer::BuildChangeRequest_UpdateExistingStream(FChangelist FromChangelist)
+	ConcertSyncClient::Replication::FChangeStreamRequest FLocalClientStreamSynchronizer::BuildChangeRequest_UpdateExistingStream(FStreamChangelist FromChangelist)
 	{
 		return { MoveTemp(FromChangelist.ObjectsToRemove), MoveTemp(FromChangelist.ObjectsToPut) };
 	}
@@ -220,7 +110,7 @@ namespace UE::MultiUserClient
 
 	void FLocalClientStreamSynchronizer::UpdateConfirmedServerState(
 		const ConcertSyncClient::Replication::FChangeStreamRequest& Request,
-		const FChangelist& ChangeThatWasInTransit
+		const FStreamChangelist& ChangeThatWasInTransit
 		)
 	{
 		const bool bCreatedNewStream = !Request.StreamsToAdd.IsEmpty(); 
@@ -251,7 +141,7 @@ namespace UE::MultiUserClient
 			ConfirmedServerState = MoveTemp(NewServerState);
 		}
 					
-		RefreshChangesCache();
+		OnServerStateSynchedDelegate.Broadcast();
 	}
 }
 
diff --git a/Engine/Plugins/Developer/Concert/ConcertApp/MultiUserClient/Source/MultiUserClient/Private/Replication/Stream/LocalClientStreamSynchronizer.h b/Engine/Plugins/Developer/Concert/ConcertApp/MultiUserClient/Source/MultiUserClient/Private/Replication/Stream/LocalClientStreamSynchronizer.h
--- a/Engine/Plugins/Developer/Concert/ConcertApp/MultiUserClient/Source/MultiUserClient/Private/Replication/Stream/LocalClientStreamSynchronizer.h
+++ b/Engine/Plugins/Developer/Concert/ConcertApp/MultiUserClient/Source/MultiUserClient/Private/Replication/Stream/LocalClientStreamSynchronizer.h
@@ -2,10 +2,9 @@
 
 #pragma once
 
+#include "IClientStreamSynchronizer.h"
 #include "Replication/IToken.h"
-#include "Replication/Data/ObjectIds.h"
 #include "Replication/Data/ObjectReplicationMap.h"
-#include "Replication/Messages/ChangeStream.h"
 
 #include "Async/Future.h"
 #include "Misc/Attribute.h"
@@ -13,99 +12,27 @@
 
 class IConcertSyncClient;
 
-namespace UE::ConcertSyncClient::Replication
-{
-	struct FChangeStreamResponse;
-	struct FChangeStreamRequest;
-}
-
 namespace UE::MultiUserClient
 {
-	enum class EObjectWarningFlags
-	{
-		Ok = 0,
-		MissingProperties = 1 << 0
-	};
-	ENUM_CLASS_FLAGS(EObjectWarningFlags);
-	
 	/**
 	 * Knows of the local client's registered replication streams and builds a changelist. The changelist tracks the
 	 * unconfirmed changes to the client's streams and is updates when the server confirms the change.
 	 */
-	class FLocalClientStreamSynchronizer : public FNoncopyable, public TSharedFromThis<FLocalClientStreamSynchronizer>
+	class FLocalClientStreamSynchronizer : public IClientStreamSynchronizer, public FNoncopyable
 	{
 	public:
 
-		DECLARE_DELEGATE(FOnModifyReplicationMap);
-		FLocalClientStreamSynchronizer(
-			TSharedRef<IConcertSyncClient> InLocalClient,
-			const FGuid& InLocalClientStreamId,
-			TAttribute<FObjectReplicationMap*> InStreamWithInProgressChangesAttribute,
-			FOnModifyReplicationMap InOnModifyReplicationMapDelegate
-			);
-
-		/**
-		 * Clears and rebuilds CachedDeltaChange in response to the underlying StreamWithInProgressChangesAttribute having been changed.
-		 * 
-		 * This iterates through all elements, which is suboptimal. However, this is needed e.g. after an undo & redo command.
-		 * In the future, we could implement a more efficient version that just responds to the relevant events (like adding an object).
-		 * While this function does redundant work, it should really not be a big performance bog after all. Also it is called
-		 * infrequently e.g. in response to a user pressing a button so we should be fine.
-		 */
-		void RefreshChangesCache();
-
-		using FSubmissionResult = TPair<ConcertSyncClient::Replication::FChangeStreamRequest, ConcertSyncClient::Replication::FChangeStreamResponse>;
-		/**
-		 * Synchronizes the server with the locally made changes to the client stream.
-		 * @note This future can execute on any thread.
-		 */
-		TFuture<FSubmissionResult> SubmitChanges(bool bRefreshChanges = false);
-		
-		/** Reverts the external stream StreamWithInProgressChangesAttribute to be the version registered on the server. */
-		void RevertCachedChanges();
+		FLocalClientStreamSynchronizer(TSharedRef<IConcertSyncClient> InLocalClient, const FGuid& InLocalClientStreamId);
 
-		/** @return Whether there are any changes that that can be submitted to the server (excludes those with warnings). */
-		bool HasSubmittableLocalChanges() const;
-		/** @return Whether there are any changes including those with warnings. */
-		bool HasRevertableLocalChanges() const;
-		
-		/** @return Whether we're currently awaiting a response from the server for updating the stream. */
-		bool IsChangeRequestInTransit() const;
-		/** @return Whether it is ok to call SubmitChanges. */
-		bool CanSubmitChanges() const { return HasSubmittableLocalChanges() && !IsChangeRequestInTransit(); }
+		//~ Begin IClientStreamSynchronizer Interface
+		virtual TFuture<FSubmitChangesResult> SubmitChanges(const FStreamChangelist& Changelist) override;
+		virtual bool CanMakeSubmitRequest() const override { return !ChangeInTransit.IsSet(); }
+		virtual FGuid GetStreamId() const override { return LocalClientStreamId; }
+		virtual const FObjectReplicationMap& GetServerState() const override { return ConfirmedServerState; }
+		virtual FOnChangesAccepted& OnChangesAccepted_AnyThread() override { return OnChangesAcceptedDelegate; }
+		virtual FOnServerStateChanged& OnServerStateSynched_AnyThread() override { return OnServerStateSynchedDelegate; }
+		//~ End IClientStreamSynchronizer Interface
 
-		/** @return The warnings flags for the given ObjectPath. */
-		EObjectWarningFlags GetObjectWarningFlags(const FSoftObjectPath& ObjectPath);
-		/** @return Whether there are any warnings for any object. */
-		bool HasAnyWarnings() const;
-
-		enum class EObjectChangeType
-		{
-			NoChange,
-			Added,
-			Removed,
-			PropertiesModified
-		};
-		/** @return Whether this object's local configuration differs from the server version. */
-		EObjectChangeType GetObjectChanges(const FSoftObjectPath& Object) const;
-		
-		enum class EPropertyChangeType
-		{
-			NoChange,
-			Added,
-			Removed
-		};
-		/** @return Gets how the given property config differs from what is on the server. */
-		EPropertyChangeType GetPropertyChanges(const FSoftObjectPath& Object, const FConcertPropertyChain& PropertyChain) const;
-
-		/** Called when RevertCachedChanges is called. The UI should be refreshed. */
-		DECLARE_MULTICAST_DELEGATE(FOnChangesReverted);
-		FOnChangesReverted& OnChangesReverted_GameThread() { return OnChangesRevertedDelegate; }
-		
-		/** Called when a change request that was in transit was accepted by the server. */
-		DECLARE_MULTICAST_DELEGATE_TwoParams(FOnChangesAccepted, const FObjectReplicationMap& OldServerState, const ConcertSyncClient::Replication::FChangeStreamRequest& AcceptedRequest);
-		FOnChangesAccepted& OnChangesAccepted_AnyThread() { return OnChangesAcceptedDelegate; }
-		
 	private:
 		
 		/** Referenced by the change requests to detect destruction of FLocalClientStreamSynchronizer */
@@ -116,58 +43,26 @@ namespace UE::MultiUserClient
 		/** The ID of the local client's stream this FLocalClientStreamDiffer is managing. */
 		const FGuid LocalClientStreamId;
 
-		/**
-		 * Represents ConfirmedServerState with changes made to it. These changes have not been sent to the server, yet.
-		 * This is cleared when a change is sent to the server. 
-		 */
-		const TAttribute<FObjectReplicationMap*> StreamWithInProgressChangesAttribute;
 		/**
 		 * Represents what the local client thinks the replication map on the server currently looks like.
 		 * This is updated every time the server confirms a change.
 		 */
 		FObjectReplicationMap ConfirmedServerState;
-
-		struct FChangelist
-		{
-			TSet<FObjectInStreamID> ObjectsToRemove;
-			TMap<FObjectInStreamID, FConcertReplication_ChangeStream_PutObject> ObjectsToPut;
-		};
 		
 		/**
 		 * Represents a change to ConfirmedServerState that is currently in transit to the server.
 		 * When a change is submitted to the server, ChangeInTransit becomes the difference between ConfirmedServerState and StreamWithInProgressChangesAttribute.
 		 * When the server confirms these changes, they are applied to ConfirmedServerState.
 		 */
-		TOptional<FChangelist> ChangeInTransit;
-		/**
-		 * The changes that if applied to ConfirmedServerState would result in StreamWithInProgressChangesAttribute.
-		 * This is updated every time the local client makes changes.
-		 */
-		FChangelist CachedDeltaChange;
-
-		/** These objects have warnings that prevent them from being set to the server (missing data, unsupported class, etc.). */
-		TMap<FSoftObjectPath, EObjectWarningFlags> ObjectsWithWarnings;
+		TOptional<FStreamChangelist> ChangeInTransit;
 		
-		/** Called when StreamWithInProgressChangesAttribute is about to be modified. */
-		FOnModifyReplicationMap OnModifyReplicationMapDelegate;
-		/** Called when StreamWithInProgressChangesAttribute was modified by this FLocalClientStreamDiffer. */
-		FOnChangesReverted OnChangesRevertedDelegate;
 		/** Called when a change request that was in transit was accepted by the server. */
 		FOnChangesAccepted OnChangesAcceptedDelegate;
-
-		FChangelist DiffChanges() const
-		{
-			const FObjectReplicationMap* Map = StreamWithInProgressChangesAttribute.Get();
-			return Map ? DiffChanges(LocalClientStreamId, ConfirmedServerState, *Map) : FChangelist{};
-		}
-		/** Builds the changelist to get from Base to Changed. */
-		static FChangelist DiffChanges(const FGuid& StreamId, const FObjectReplicationMap& Base, const FObjectReplicationMap& Changed);
-
-		/** Updates ObjectsWithWarnings. */
-		void RefreshWarnings();
+		/** Event executed when the result of GetServerState has been synched. */
+		FOnServerStateChanged OnServerStateSynchedDelegate;
 		
 		/** Builds a change request based on what's registered on the server. */
-		ConcertSyncClient::Replication::FChangeStreamRequest BuildChangeRequest() const;
+		ConcertSyncClient::Replication::FChangeStreamRequest BuildChangeRequest(const FStreamChangelist& Changelist) const;
 		enum class EChangeRequestType
 		{
 			/** When NeedsToCreateNewStreamNextRequest returns true */
@@ -175,8 +70,8 @@ namespace UE::MultiUserClient
 			UpdateExistingStream
 		};
 		/** Converts a changelist to. */
-		static ConcertSyncClient::Replication::FChangeStreamRequest BuildChangeRequest_CreateNewStream(const FGuid& StreamId, const FChangelist& FromChangelist);
-		static ConcertSyncClient::Replication::FChangeStreamRequest BuildChangeRequest_UpdateExistingStream(FChangelist FromChangelist);
+		static ConcertSyncClient::Replication::FChangeStreamRequest BuildChangeRequest_CreateNewStream(const FGuid& StreamId, const FStreamChangelist& FromChangelist);
+		static ConcertSyncClient::Replication::FChangeStreamRequest BuildChangeRequest_UpdateExistingStream(FStreamChangelist FromChangelist);
 		/**
 		 * The stream may have to be created first. This is the case either if
 		 *	1. this is the first time sending to the server, or
@@ -187,7 +82,7 @@ namespace UE::MultiUserClient
 		/** Updates the server state after the server has answered a request. */
 		void UpdateConfirmedServerState(
 			const ConcertSyncClient::Replication::FChangeStreamRequest& Request,
-			const FChangelist& ChangeThatWasInTransit
+			const FStreamChangelist& ChangeThatWasInTransit
 			);
 	};
 }
diff --git a/Engine/Plugins/Developer/Concert/ConcertApp/MultiUserClient/Source/MultiUserClient/Private/Replication/Stream/LocalStreamChangeTracker.cpp b/Engine/Plugins/Developer/Concert/ConcertApp/MultiUserClient/Source/MultiUserClient/Private/Replication/Stream/LocalStreamChangeTracker.cpp
new file mode 100644
--- /dev/null
+++ b/Engine/Plugins/Developer/Concert/ConcertApp/MultiUserClient/Source/MultiUserClient/Private/Replication/Stream/LocalStreamChangeTracker.cpp
@@ -0,0 +1,143 @@
+// Copyright Epic Games, Inc. All Rights Reserved.
+
+#include "LocalStreamChangeTracker.h"
+
+#include "IConcertSyncClient.h"
+#include "Replication/ChangeStreamSharedUtils.h"
+#include "Replication/IConcertClientReplicationManager.h"
+
+#include "ScopedTransaction.h"
+
+#define LOCTEXT_NAMESPACE "FLocalClientStreamDiffer"
+
+namespace UE::MultiUserClient
+{
+	FLocalStreamChangeTracker::FLocalStreamChangeTracker(
+		TSharedRef<IClientStreamSynchronizer> InStreamSynchronizer,
+		TAttribute<FObjectReplicationMap*> InStreamWithInProgressChangesAttribute,
+		FOnModifyReplicationMap InOnModifyReplicationMapDelegate
+		)
+		: StreamSynchronizer(MoveTemp(InStreamSynchronizer))
+		, StreamWithInProgressChangesAttribute(MoveTemp(InStreamWithInProgressChangesAttribute))
+		, OnModifyReplicationMapDelegate(MoveTemp(InOnModifyReplicationMapDelegate))
+	{
+		check(InStreamWithInProgressChangesAttribute.IsBound() || InStreamWithInProgressChangesAttribute.IsSet());
+		StreamSynchronizer->OnServerStateSynched_AnyThread().AddRaw(this, &FLocalStreamChangeTracker::RefreshChangesCache);
+	}
+
+	FLocalStreamChangeTracker::~FLocalStreamChangeTracker()
+	{
+		StreamSynchronizer->OnServerStateSynched_AnyThread().RemoveAll(this);
+	}
+
+	void FLocalStreamChangeTracker::RefreshChangesCache()
+	{
+		CachedDeltaChange = DiffChanges();
+		RefreshWarnings();
+	}
+
+	TFuture<FSubmitChangesResult> FLocalStreamChangeTracker::SubmitChanges(bool bRefreshChanges)
+	{
+		if (bRefreshChanges)
+		{
+			RefreshChangesCache();
+		}
+		return StreamSynchronizer->SubmitChanges(CachedDeltaChange);
+	}
+
+	void FLocalStreamChangeTracker::RevertCachedChanges()
+	{
+		if (HasRevertableLocalChanges())
+		{
+			FScopedTransaction Transaction(LOCTEXT("RevertCachedChanges", "Revert replication changes"));
+			OnModifyReplicationMapDelegate.ExecuteIfBound();
+			
+			*StreamWithInProgressChangesAttribute.Get() = StreamSynchronizer->GetServerState();
+			CachedDeltaChange = {};
+
+			check(IsInGameThread());
+			OnChangesRevertedDelegate.Broadcast();
+			RefreshChangesCache();
+		}
+	}
+
+	bool FLocalStreamChangeTracker::HasSubmittableLocalChanges() const
+	{
+		return !CachedDeltaChange.ObjectsToPut.IsEmpty() || !CachedDeltaChange.ObjectsToRemove.IsEmpty();
+	}
+
+	bool FLocalStreamChangeTracker::HasRevertableLocalChanges() const
+	{
+		return HasSubmittableLocalChanges() || HasAnyWarnings();
+	}
+
+	bool FLocalStreamChangeTracker::CanMakeSubmitNetworkRequest() const
+	{
+		return StreamSynchronizer->CanMakeSubmitRequest();
+	}
+
+	EObjectWarningFlags FLocalStreamChangeTracker::GetObjectWarningFlags(const FSoftObjectPath& ObjectPath)
+	{
+		const EObjectWarningFlags* Flags = ObjectsWithWarnings.Find(ObjectPath);
+		return Flags ? *Flags : EObjectWarningFlags::Ok;
+	}
+
+	bool FLocalStreamChangeTracker::HasAnyWarnings() const
+	{
+		return !ObjectsWithWarnings.IsEmpty();
+	}
+
+	FLocalStreamChangeTracker::EObjectChangeType FLocalStreamChangeTracker::GetObjectChanges(const FSoftObjectPath& Object) const
+	{
+		// TODO DP:
+		return EObjectChangeType::NoChange;
+	}
+
+	FLocalStreamChangeTracker::EPropertyChangeType FLocalStreamChangeTracker::GetPropertyChanges(const FSoftObjectPath& Object, const FConcertPropertyChain& PropertyChain) const
+	{
+		// TODO DP:
+		return EPropertyChangeType::NoChange;
+	}
+	
+	FStreamChangelist FLocalStreamChangeTracker::DiffChanges(
+		const FGuid& StreamId,
+		const FObjectReplicationMap& Base,
+		const FObjectReplicationMap& Changed
+		)
+	{
+		// BuildRequestFromDiff does not tolerate any invalid entries (like empty properties, which the UI generates right after you add an object to the list)
+		FObjectReplicationMap Cleansed = Changed;
+		ConcertSyncCore::Replication::ChangeStreamUtils::IterateInvalidEntries(Changed, [&Cleansed](const FSoftObjectPath& InvalidObject, const FReplicatedObjectInfo&)
+		{
+			Cleansed.ReplicatedObjects.Remove(InvalidObject);
+			return EBreakBehavior::Continue;
+		});
+
+		FConcertReplication_ChangeStream_Request Request = ConcertSyncCore::Replication::ChangeStreamUtils::BuildRequestFromDiff(StreamId, Base, Cleansed);
+		return { MoveTemp(Request.ObjectsToRemove), MoveTemp(Request.ObjectsToPut) };
+	}
+
+	void FLocalStreamChangeTracker::RefreshWarnings()
+	{
+		ObjectsWithWarnings.Empty();
+		const FObjectReplicationMap* Map = StreamWithInProgressChangesAttribute.Get();
+		if (!Map)
+		{
+			return;
+		}
+		
+		ConcertSyncCore::Replication::ChangeStreamUtils::IterateInvalidEntries(*Map, [this](const FSoftObjectPath& InvalidObject, const FReplicatedObjectInfo& Info)
+		{
+			EObjectWarningFlags Flags = EObjectWarningFlags::Ok;
+			if (Info.PropertySelection.ReplicatedProperties.IsEmpty())
+			{
+				Flags |= EObjectWarningFlags::MissingProperties;
+			}
+			
+			ObjectsWithWarnings.Add(InvalidObject, Flags);
+			return EBreakBehavior::Continue;
+		});
+	}
+}
+
+#undef LOCTEXT_NAMESPACE
\ No newline at end of file
diff --git a/Engine/Plugins/Developer/Concert/ConcertApp/MultiUserClient/Source/MultiUserClient/Private/Replication/Stream/LocalStreamChangeTracker.h b/Engine/Plugins/Developer/Concert/ConcertApp/MultiUserClient/Source/MultiUserClient/Private/Replication/Stream/LocalStreamChangeTracker.h
new file mode 100644
--- /dev/null
+++ b/Engine/Plugins/Developer/Concert/ConcertApp/MultiUserClient/Source/MultiUserClient/Private/Replication/Stream/LocalStreamChangeTracker.h
@@ -0,0 +1,139 @@
+// Copyright Epic Games, Inc. All Rights Reserved.
+
+#pragma once
+
+#include "IClientStreamSynchronizer.h"
+#include "Replication/Data/ObjectReplicationMap.h"
+
+#include "Async/Future.h"
+#include "Misc/Attribute.h"
+#include "Templates/UnrealTemplate.h"
+
+class IConcertSyncClient;
+
+namespace UE::ConcertSyncClient::Replication
+{
+	struct FChangeStreamResponse;
+	struct FChangeStreamRequest;
+}
+
+namespace UE::MultiUserClient
+{
+	enum class EObjectWarningFlags
+	{
+		Ok = 0,
+		MissingProperties = 1 << 0
+	};
+	ENUM_CLASS_FLAGS(EObjectWarningFlags);
+	
+	/**
+	 * Knows of the local client's registered replication streams and builds a changelist. The changelist tracks the
+	 * unconfirmed changes to the client's streams and is updates when the server confirms the change.
+	 */
+	class FLocalStreamChangeTracker : public FNoncopyable
+	{
+	public:
+
+		DECLARE_DELEGATE(FOnModifyReplicationMap);
+		FLocalStreamChangeTracker(
+			TSharedRef<IClientStreamSynchronizer> InStreamSynchronizer,
+			TAttribute<FObjectReplicationMap*> InStreamWithInProgressChangesAttribute,
+			FOnModifyReplicationMap InOnModifyReplicationMapDelegate
+			);
+		~FLocalStreamChangeTracker();
+
+		/**
+		 * Clears and rebuilds CachedDeltaChange in response to the underlying StreamWithInProgressChangesAttribute having been changed.
+		 * 
+		 * This iterates through all elements, which is suboptimal. However, this is needed e.g. after an undo & redo command.
+		 * In the future, we could implement a more efficient version that just responds to the relevant events (like adding an object).
+		 * While this function does redundant work, it should really not be a big performance bog after all. Also it is called
+		 * infrequently e.g. in response to a user pressing a button so we should be fine.
+		 */
+		void RefreshChangesCache();
+
+		/**
+		 * Synchronizes the server with the locally made changes to the client stream.
+		 * @note This future can execute on any thread.
+		 */
+		TFuture<FSubmitChangesResult> SubmitChanges(bool bRefreshChanges = false);
+		
+		/** Reverts the external stream StreamWithInProgressChangesAttribute to be the version registered on the server. */
+		void RevertCachedChanges();
+
+		/** @return Whether there are any changes that that can be submitted to the server (excludes those with warnings). */
+		bool HasSubmittableLocalChanges() const;
+		/** @return Whether there are any changes including those with warnings. */
+		bool HasRevertableLocalChanges() const;
+		
+		/** @return Whether we're currently awaiting a response from the server for updating the stream. */
+		bool CanMakeSubmitNetworkRequest() const;
+		/** @return Whether it is ok to call SubmitChanges. */
+		bool CanSubmitChanges() const { return HasSubmittableLocalChanges() && CanMakeSubmitNetworkRequest(); }
+
+		/** @return The warnings flags for the given ObjectPath. */
+		EObjectWarningFlags GetObjectWarningFlags(const FSoftObjectPath& ObjectPath);
+		/** @return Whether there are any warnings for any object. */
+		bool HasAnyWarnings() const;
+
+		enum class EObjectChangeType
+		{
+			NoChange,
+			Added,
+			Removed,
+			PropertiesModified
+		};
+		/** @return Whether this object's local configuration differs from the server version. */
+		EObjectChangeType GetObjectChanges(const FSoftObjectPath& Object) const;
+		
+		enum class EPropertyChangeType
+		{
+			NoChange,
+			Added,
+			Removed
+		};
+		/** @return Gets how the given property config differs from what is on the server. */
+		EPropertyChangeType GetPropertyChanges(const FSoftObjectPath& Object, const FConcertPropertyChain& PropertyChain) const;
+
+		/** Called when RevertCachedChanges is called. The UI should be refreshed. */
+		DECLARE_MULTICAST_DELEGATE(FOnChangesReverted);
+		FOnChangesReverted& OnChangesReverted_GameThread() { return OnChangesRevertedDelegate; }
+		
+	private:
+
+		/** Keeps track of the edited stream's server state. */
+		TSharedRef<IClientStreamSynchronizer> StreamSynchronizer;
+
+		/**
+		 * Represents ConfirmedServerState with changes made to it. These changes have not been sent to the server, yet.
+		 * This is cleared when a change is sent to the server. 
+		 */
+		const TAttribute<FObjectReplicationMap*> StreamWithInProgressChangesAttribute;
+		
+		/**
+		 * The changes that if applied to ConfirmedServerState would result in StreamWithInProgressChangesAttribute.
+		 * This is updated every time the local client makes changes.
+		 */
+		FStreamChangelist CachedDeltaChange;
+
+		/** These objects have warnings that prevent them from being set to the server (missing data, unsupported class, etc.). */
+		TMap<FSoftObjectPath, EObjectWarningFlags> ObjectsWithWarnings;
+		
+		/** Called when StreamWithInProgressChangesAttribute is about to be modified. */
+		FOnModifyReplicationMap OnModifyReplicationMapDelegate;
+		/** Called when StreamWithInProgressChangesAttribute was modified by this FLocalClientStreamDiffer. */
+		FOnChangesReverted OnChangesRevertedDelegate;
+
+		FStreamChangelist DiffChanges() const
+		{
+			const FObjectReplicationMap* Map = StreamWithInProgressChangesAttribute.Get();
+			return Map ? DiffChanges(StreamSynchronizer->GetStreamId(), StreamSynchronizer->GetServerState(), *Map) : FStreamChangelist{};
+		}
+		/** Builds the changelist to get from Base to Changed. */
+		static FStreamChangelist DiffChanges(const FGuid& StreamId, const FObjectReplicationMap& Base, const FObjectReplicationMap& Changed);
+
+		/** Updates ObjectsWithWarnings. */
+		void RefreshWarnings();
+	};
+}
+
diff --git a/Engine/Plugins/Developer/Concert/ConcertApp/MultiUserClient/Source/MultiUserClient/Private/Widgets/ActiveSession/Replication/Joined/SReplicationClientView.cpp b/Engine/Plugins/Developer/Concert/ConcertApp/MultiUserClient/Source/MultiUserClient/Private/Widgets/ActiveSession/Replication/Joined/SReplicationClientView.cpp
new file mode 100644
--- /dev/null
+++ b/Engine/Plugins/Developer/Concert/ConcertApp/MultiUserClient/Source/MultiUserClient/Private/Widgets/ActiveSession/Replication/Joined/SReplicationClientView.cpp
@@ -0,0 +1,60 @@
+// Copyright Epic Games, Inc. All Rights Reserved.
+
+#include "SReplicationClientView.h"
+
+#include "Replication/Client/ReplicationClient.h"
+#include "Replication/Editor/Model/Object/EditorObjectSelectionSourceModel.h"
+#include "Replication/Editor/Model/Property/SelectPropertyFromUClassModel.h"
+#include "Replication/Editor/View/IReplicationEditorView.h"
+#include "Replication/ReplicationWidgetFactories.h"
+#include "Widgets/ActiveSession/Replication/Joined/SReplicationClientViewToolbar.h"
+
+#include "Widgets/SBoxPanel.h"
+
+namespace UE::MultiUserClient
+{
+	void SReplicationClientView::Construct(const FArguments& InArgs)
+	{
+		GetReplicationClientAttribute = InArgs._GetReplicationClient;
+		FReplicationClient* ReplicationClient = GetReplicationClientAttribute.Get();
+		check(ReplicationClient);
+		
+		using namespace ConcertClientSharedSlate;
+		const FCreateEditorParams ReplicationEditorCreationParams
+		{
+			ReplicationClient->GetClientEditModel(),
+			MakeShared<FEditorObjectSelectionSourceModel>(),
+			MakeShared<FSelectPropertyFromUClassModel>()
+		};
+
+		EditorView = CreateEditorForUnrealEditor(ReplicationEditorCreationParams);
+		ChildSlot
+		[
+			SNew(SVerticalBox)
+
+			// Toolbar
+			+SVerticalBox::Slot()
+			.AutoHeight()
+			.Padding(2.f)
+			[
+				SNew(SReplicationClientViewToolbar)
+				.GetChangeTrackerAttribute_Lambda([this](){ return &GetReplicationClientAttribute.Get()->GetDiffer(); })
+			]
+
+			// Editor
+			+SVerticalBox::Slot()
+			.FillHeight(1.f)
+			[
+				EditorView.ToSharedRef()
+			]
+		];
+		
+		// Refresh UI if streams change externally, e.g. a remote client changed what they sent
+		ReplicationClient->OnModelExternallyChanged_GameThread().AddSP(this, &SReplicationClientView::OnModelChanged);
+	}
+	
+	void SReplicationClientView::OnModelChanged() const
+	{
+		EditorView->Refresh();
+	}
+}
diff --git a/Engine/Plugins/Developer/Concert/ConcertApp/MultiUserClient/Source/MultiUserClient/Private/Widgets/ActiveSession/Replication/Joined/SReplicationClientView.h b/Engine/Plugins/Developer/Concert/ConcertApp/MultiUserClient/Source/MultiUserClient/Private/Widgets/ActiveSession/Replication/Joined/SReplicationClientView.h
new file mode 100644
--- /dev/null
+++ b/Engine/Plugins/Developer/Concert/ConcertApp/MultiUserClient/Source/MultiUserClient/Private/Widgets/ActiveSession/Replication/Joined/SReplicationClientView.h
@@ -0,0 +1,42 @@
+// Copyright Epic Games, Inc. All Rights Reserved.
+
+#pragma once
+
+#include "Widgets/DeclarativeSyntaxSupport.h"
+#include "Widgets/SCompoundWidget.h"
+
+namespace UE::ConcertClientSharedSlate
+{
+	class IReplicationEditorView;
+}
+
+namespace UE::MultiUserClient
+{
+	class FReplicationClient;
+	
+	/** Displays the contents of a client. */
+	class SReplicationClientView : public SCompoundWidget
+	{
+	public:
+
+		SLATE_BEGIN_ARGS(SReplicationClientView)
+		{}
+			/** The client to depict. Should always return true. If the client is destroyed, so should this widget be. */
+			SLATE_ATTRIBUTE(FReplicationClient*, GetReplicationClient)
+		SLATE_END_ARGS()
+
+		void Construct(const FArguments& InArgs);
+
+	private:
+		
+		/** The editor view of the replication content. */
+		TSharedPtr<ConcertClientSharedSlate::IReplicationEditorView> EditorView;
+		
+		/** The client to depict. Should always return true. If the client is destroyed, so should this widget be. */
+		TAttribute<FReplicationClient*> GetReplicationClientAttribute;
+		
+		/** Called when any of the streams change. */
+		void OnModelChanged() const;
+	};
+}
+
diff --git a/Engine/Plugins/Developer/Concert/ConcertApp/MultiUserClient/Source/MultiUserClient/Private/Widgets/ActiveSession/Replication/Joined/SReplicationClientViewToolbar.cpp b/Engine/Plugins/Developer/Concert/ConcertApp/MultiUserClient/Source/MultiUserClient/Private/Widgets/ActiveSession/Replication/Joined/SReplicationClientViewToolbar.cpp
new file mode 100644
--- /dev/null
+++ b/Engine/Plugins/Developer/Concert/ConcertApp/MultiUserClient/Source/MultiUserClient/Private/Widgets/ActiveSession/Replication/Joined/SReplicationClientViewToolbar.cpp
@@ -0,0 +1,207 @@
+// Copyright Epic Games, Inc. All Rights Reserved.
+
+#include "SReplicationClientViewToolbar.h"
+
+#include "ConcertLogGlobal.h"
+#include "Replication/IConcertClientReplicationManager.h"
+#include "Replication/MultiUserReplicationManager.h"
+#include "Replication/Stream/LocalStreamChangeTracker.h"
+
+#include "Framework/Notifications/NotificationManager.h"
+#include "SNegativeActionButton.h"
+#include "Styling/AppStyle.h"
+#include "Styling/StyleColors.h"
+#include "Widgets/Images/SImage.h"
+#include "Widgets/Input/SButton.h"
+#include "Widgets/Layout/SBox.h"
+#include "Widgets/Layout/SSpacer.h"
+#include "Widgets/Notifications/SNotificationList.h"
+#include "Widgets/SBoxPanel.h"
+#include "Widgets/Text/STextBlock.h"
+
+#define LOCTEXT_NAMESPACE "SReplicationJoinedWidget"
+
+namespace UE::MultiUserClient
+{
+	void SReplicationClientViewToolbar::Construct(const FArguments& InArgs)
+	{
+		GetChangeTrackerAttribute = InArgs._GetChangeTrackerAttribute;
+		check(GetChangeTrackerAttribute.IsBound() || GetChangeTrackerAttribute.IsSet());
+		
+		ChildSlot
+		[
+			SNew(SHorizontalBox)
+
+			+SHorizontalBox::Slot()
+			.FillWidth(1.f)
+			[
+				SNew(SSpacer)
+			]
+
+			// Upload button
+			+SHorizontalBox::Slot()
+			.AutoWidth()
+			[
+				BuildUploadButton()
+			]
+
+			// Revert changes
+			+SHorizontalBox::Slot()
+			.AutoWidth()
+			[
+				BuildRevertButton()
+			]
+		];
+	}
+
+	FLocalStreamChangeTracker& SReplicationClientViewToolbar::GetChangeTracker() const
+	{
+		FLocalStreamChangeTracker* ChangeTracker = GetChangeTrackerAttribute.Get();
+		checkf(ChangeTracker, TEXT("You were supposed to destroy this widget."));
+		return *ChangeTracker;
+	}
+
+	TSharedRef<SWidget> SReplicationClientViewToolbar::BuildUploadButton()
+	{
+		return SNew(SButton)
+			.ForegroundColor(FSlateColor::UseStyle())
+			.IsEnabled_Lambda([this](){ return GetChangeTracker().CanSubmitChanges(); })
+			.ToolTipText(this, &SReplicationClientViewToolbar::GetUploadButtonToolTipText)
+			.HAlign(HAlign_Center)
+			.VAlign(VAlign_Center)
+			.OnClicked(this, &SReplicationClientViewToolbar::OnUploadButtonClicked)
+			[
+				SNew(SHorizontalBox)
+				+ SHorizontalBox::Slot()
+				.HAlign(HAlign_Center)
+				.VAlign(VAlign_Center)
+				.AutoWidth()
+				[
+					SNew(SImage)
+					.Image_Lambda([this]()
+					{
+						const FLocalStreamChangeTracker& Differ = GetChangeTracker();
+						return Differ.HasAnyWarnings()
+							? FAppStyle::Get().GetBrush("Icons.Warning")
+							: FAppStyle::Get().GetBrush("Icons.Plus");
+					})
+					.ColorAndOpacity_Lambda([this]()
+					{
+						const FLocalStreamChangeTracker& Differ = GetChangeTracker();
+						return Differ.HasAnyWarnings()
+							? FStyleColors::AccentYellow
+							: FStyleColors::AccentGreen;
+					})
+				]
+				+ SHorizontalBox::Slot()
+				.Padding(FMargin(3, 0, 0, 0))
+				.VAlign(VAlign_Center)
+				.AutoWidth()
+				[
+					SNew(STextBlock)
+					.TextStyle(FAppStyle::Get(), "SmallButtonText")
+					.Text(LOCTEXT("UploadChanges.Label", "Upload"))
+				]
+			];
+	}
+
+	FText SReplicationClientViewToolbar::GetUploadButtonToolTipText() const
+	{
+		const FLocalStreamChangeTracker& Differ = GetChangeTracker();
+		const FText BaseToolTipText = [this, &Differ]()
+		{
+			if (Differ.HasSubmittableLocalChanges() && !Differ.CanMakeSubmitNetworkRequest())
+			{
+				return LOCTEXT("UploadChanges.ToolTip.Upload", "Upload your local changes to the server.");
+			}
+			if (Differ.CanMakeSubmitNetworkRequest())
+			{
+				return LOCTEXT("UploadChanges.ToolTip.InTransit", "Awaiting server response.");
+			}
+			return LOCTEXT("UploadChanges.ToolTip.NoChanges", "No submittable changes.");
+		}();
+		if (!Differ.HasAnyWarnings())
+		{
+			return BaseToolTipText;
+		}
+
+		return FText::Format(LOCTEXT("UploadChanges.ToolTip.BaseFmt", "{0}\n\n{1}"),
+			BaseToolTipText,
+			LOCTEXT("UploadChanges.ToolTip", "Warning: Objects with warnings will not be submitted. There are objects with warnings.")
+			);
+	}
+
+	FReply SReplicationClientViewToolbar::OnUploadButtonClicked() const
+	{
+		FLocalStreamChangeTracker& Differ = GetChangeTracker();
+		if (Differ.CanSubmitChanges())
+		{
+			FNotificationInfo NotificationInfo(LOCTEXT("Uploading.InProgress", "Uploading stream changes."));
+			NotificationInfo.bUseThrobber = true;
+			NotificationInfo.bUseSuccessFailIcons = true;
+			NotificationInfo.bFireAndForget = false;
+			const TSharedPtr<SNotificationItem> Notification = FSlateNotificationManager::Get().AddNotification(NotificationInfo);
+			Differ.SubmitChanges()
+				.Next([Notification](FSubmitChangesResult&& SubmitChangesResult)
+				{
+					TOptional<FCompletedChangeSubmission> SubmissionInfo = SubmitChangesResult.SubmissionInfo;
+					if (SubmissionInfo.IsSet() && SubmissionInfo->Response.IsSuccess())
+					{
+						Notification->SetCompletionState(SNotificationItem::CS_Success);
+						Notification->SetText(LOCTEXT("Uploading.Success", "Changes accepted."));
+									
+						const ConcertSyncClient::Replication::FChangeStreamRequest& Request = SubmissionInfo->Request;
+						// Multi-User adds only one stream so we only need to look at index 0
+						const int32 NumAdditions = Request.StreamsToAdd.IsEmpty() ? 0 : Request.StreamsToAdd[0].BaseDescription.ReplicationMap.ReplicatedObjects.Num();
+						const int32 NumPuts = Request.ObjectsToPut.Num() + NumAdditions;
+						Notification->SetSubText(FText::Format(LOCTEXT("Uploading.Success.SubText", "{0} puts, {1} removals"),
+							NumPuts,
+							Request.ObjectsToRemove.Num()
+							));
+					}
+					else
+					{
+						Notification->SetCompletionState(SNotificationItem::CS_Fail);
+						Notification->SetText(LOCTEXT("Uploading.Failure", "Changes rejected."));
+						Notification->SetSubText(LOCTEXT("Uploading.Failure.SubText", "See log for details."));
+
+						if (SubmissionInfo.IsSet())
+						{
+							FStringOutputDevice Errors;
+							SubmissionInfo->Response.LogErrors(Errors);
+							UE_LOG(LogConcert, Error, TEXT("Errors uploading stream changes:\n%s"), *Errors);
+						}
+						else
+						{
+							UE_LOG(LogConcert, Error, TEXT("Failed to submit changes. ESubmitChangesErrorCode: %d"), SubmitChangesResult.ErrorCode);
+						}
+					}
+								
+					Notification->ExpireAndFadeout();
+				});
+		}
+		return FReply::Handled();
+	}
+
+	TSharedRef<SWidget> SReplicationClientViewToolbar::BuildRevertButton()
+	{
+		return SNew(SNegativeActionButton)
+			.Text(LOCTEXT("RevertChanges.Label", "Revert"))
+			.ToolTipText_Lambda([this]()
+			{
+				const FLocalStreamChangeTracker& Differ = GetChangeTracker();
+				return Differ.HasRevertableLocalChanges()
+					? LOCTEXT("RevertChanges.ToolTip.Revert", "Reverts all local changes you've made to the state that's on the server.")
+					: LOCTEXT("RevertChanges.ToolTip.NothingToUpload", "Once you make changes to the stream, you can use this button to revert them.");
+			})
+			.IsEnabled_Lambda([this](){ return GetChangeTracker().HasRevertableLocalChanges(); })
+			.OnClicked_Lambda([this]()
+			{
+				FLocalStreamChangeTracker& Differ = GetChangeTracker();
+				Differ.RevertCachedChanges();
+				return FReply::Handled();
+			});
+	}
+}
+
+#undef LOCTEXT_NAMESPACE
\ No newline at end of file
diff --git a/Engine/Plugins/Developer/Concert/ConcertApp/MultiUserClient/Source/MultiUserClient/Private/Widgets/ActiveSession/Replication/Joined/SReplicationClientViewToolbar.h b/Engine/Plugins/Developer/Concert/ConcertApp/MultiUserClient/Source/MultiUserClient/Private/Widgets/ActiveSession/Replication/Joined/SReplicationClientViewToolbar.h
new file mode 100644
--- /dev/null
+++ b/Engine/Plugins/Developer/Concert/ConcertApp/MultiUserClient/Source/MultiUserClient/Private/Widgets/ActiveSession/Replication/Joined/SReplicationClientViewToolbar.h
@@ -0,0 +1,40 @@
+// Copyright Epic Games, Inc. All Rights Reserved.
+
+#pragma once
+
+#include "Widgets/DeclarativeSyntaxSupport.h"
+#include "Widgets/SCompoundWidget.h"
+
+namespace UE::MultiUserClient
+{
+	class FLocalStreamChangeTracker;
+
+	/** Contains a bunch of actions that can be performed while connected to replication. */
+	class SReplicationClientViewToolbar : public SCompoundWidget
+	{
+	public:
+		
+		SLATE_BEGIN_ARGS(SReplicationClientViewToolbar)
+		{}
+			/** Used to submit and revert locally made changes. */
+			SLATE_ATTRIBUTE(FLocalStreamChangeTracker*, GetChangeTrackerAttribute)
+		SLATE_END_ARGS()
+
+		void Construct(const FArguments& InArgs);
+
+	private:
+		
+		/** Used to submit and revert locally made changes. */
+		TAttribute<FLocalStreamChangeTracker*> GetChangeTrackerAttribute;
+
+		/** Check()s that the GetChangeTrackerAttribute is valid. */
+		FLocalStreamChangeTracker& GetChangeTracker() const;
+		
+		TSharedRef<SWidget> BuildUploadButton();
+		FText GetUploadButtonToolTipText() const;
+		FReply OnUploadButtonClicked() const;
+		
+		TSharedRef<SWidget> BuildRevertButton();
+	};
+}
+
diff --git a/Engine/Plugins/Developer/Concert/ConcertApp/MultiUserClient/Source/MultiUserClient/Private/Widgets/ActiveSession/Replication/Joined/SReplicationJoinedToolbar.cpp b/Engine/Plugins/Developer/Concert/ConcertApp/MultiUserClient/Source/MultiUserClient/Private/Widgets/ActiveSession/Replication/Joined/SReplicationJoinedToolbar.cpp
deleted file mode 100644
diff --git a/Engine/Plugins/Developer/Concert/ConcertApp/MultiUserClient/Source/MultiUserClient/Private/Widgets/ActiveSession/Replication/Joined/SReplicationJoinedToolbar.h b/Engine/Plugins/Developer/Concert/ConcertApp/MultiUserClient/Source/MultiUserClient/Private/Widgets/ActiveSession/Replication/Joined/SReplicationJoinedToolbar.h
deleted file mode 100644
diff --git a/Engine/Plugins/Developer/Concert/ConcertApp/MultiUserClient/Source/MultiUserClient/Private/Widgets/ActiveSession/Replication/Joined/SReplicationJoinedView.cpp b/Engine/Plugins/Developer/Concert/ConcertApp/MultiUserClient/Source/MultiUserClient/Private/Widgets/ActiveSession/Replication/Joined/SReplicationJoinedView.cpp
--- a/Engine/Plugins/Developer/Concert/ConcertApp/MultiUserClient/Source/MultiUserClient/Private/Widgets/ActiveSession/Replication/Joined/SReplicationJoinedView.cpp
+++ b/Engine/Plugins/Developer/Concert/ConcertApp/MultiUserClient/Source/MultiUserClient/Private/Widgets/ActiveSession/Replication/Joined/SReplicationJoinedView.cpp
@@ -3,13 +3,8 @@
 #include "SReplicationJoinedView.h"
 
 #include "ConcertLogGlobal.h"
-#include "Replication/Editor/Model/Object/EditorObjectSelectionSourceModel.h"
-#include "Replication/Editor/Model/Property/SelectPropertyFromUClassModel.h"
-#include "Replication/Editor/View/IReplicationEditorView.h"
+#include "SReplicationClientView.h"
 #include "Replication/MultiUserReplicationManager.h"
-#include "Replication/ReplicationWidgetFactories.h"
-#include "Replication/Stream/ClientStreamRepository.h"
-#include "SReplicationJoinedToolbar.h"
 
 #include "Framework/Notifications/NotificationManager.h"
 #include "Widgets/Layout/SBox.h"
@@ -23,53 +18,18 @@ namespace UE::MultiUserClient
 	void SReplicationJoinedView::Construct(const FArguments& InArgs, TSharedRef<FMultiUserReplicationManager> InReplicationManager)
 	{
 		ReplicationManager = InReplicationManager;
-		
-		using namespace ConcertClientSharedSlate;
-		const FCreateEditorParams ReplicationEditorCreationParams
-		{
-			ReplicationManager->GetStreamSynchronizer()->GetLocalClientEditModel(),
-			MakeShared<FEditorObjectSelectionSourceModel>(),
-			MakeShared<FSelectPropertyFromUClassModel>()
-		};
 
-		const TSharedRef<IReplicationEditorView> NewEditorView = CreateEditorForUnrealEditor(ReplicationEditorCreationParams);
-		WeakEditorView = NewEditorView;
 		ChildSlot
 		[
-			SNew(SVerticalBox)
-
-			// Toolbar
-			+SVerticalBox::Slot()
-			.AutoHeight()
-			.Padding(2.f)
-			[
-				SNew(SReplicationJoinedToolbar, InReplicationManager)
-			]
-
-			// Editor
-			+SVerticalBox::Slot()
-			.FillHeight(1.f)
-			[
-				NewEditorView
-			]
+			SNew(SReplicationClientView)
+			.GetReplicationClient_Lambda([this](){ return &ReplicationManager->GetClientManager()->GetLocalClient(); })
 		];
-		
-		// Refresh UI if streams change externally, e.g. a remote client changed what they sent
-		ReplicationManager->GetStreamSynchronizer()->OnModelChanged_GameThread().AddSP(this, &SReplicationJoinedView::OnModelChanged);
 
 		// Show notifications about changing authority
 		ReplicationManager->GetAuthorityPolicy()->OnAuthorityRequestSent_AnyThread().AddSP(this, &SReplicationJoinedView::OnAuthorityRequestSent_AnyThread);
 		ReplicationManager->GetAuthorityPolicy()->OnAuthorityResponseReceived_AnyThread().AddSP(this, &SReplicationJoinedView::OnAuthorityResponseReceived_AnyThread);
 	}
 
-	void SReplicationJoinedView::OnModelChanged() const
-	{
-		if (const TSharedPtr<ConcertClientSharedSlate::IReplicationEditorView> EditorView = WeakEditorView.Pin())
-		{
-			EditorView->Refresh();
-		}
-	}
-
 	void SReplicationJoinedView::OnAuthorityRequestSent_AnyThread(
 		const ConcertSyncClient::Replication::FAuthorityChangeRequest& Request
 		)
diff --git a/Engine/Plugins/Developer/Concert/ConcertApp/MultiUserClient/Source/MultiUserClient/Private/Widgets/ActiveSession/Replication/Joined/SReplicationJoinedView.h b/Engine/Plugins/Developer/Concert/ConcertApp/MultiUserClient/Source/MultiUserClient/Private/Widgets/ActiveSession/Replication/Joined/SReplicationJoinedView.h
--- a/Engine/Plugins/Developer/Concert/ConcertApp/MultiUserClient/Source/MultiUserClient/Private/Widgets/ActiveSession/Replication/Joined/SReplicationJoinedView.h
+++ b/Engine/Plugins/Developer/Concert/ConcertApp/MultiUserClient/Source/MultiUserClient/Private/Widgets/ActiveSession/Replication/Joined/SReplicationJoinedView.h
@@ -6,6 +6,11 @@
 #include "Widgets/DeclarativeSyntaxSupport.h"
 #include "Widgets/SCompoundWidget.h"
 
+namespace UE::MultiUserClient
+{
+	class SReplicationClientView;
+}
+
 class SNotificationItem;
 
 namespace UE::ConcertClientSharedSlate
@@ -34,17 +39,14 @@ namespace UE::MultiUserClient
 		TSharedPtr<FMultiUserReplicationManager> ReplicationManager;
 		
 		/**
-		 * The editor view if one is being displayed.
+		 * The local client's view if one is being displayed.
 		 * ChildSlot keeps the reference alive if we're in state EMultiUserReplicationConnectionState::Connected.
 		 * @see ShowWidget_Connected
 		 */
-		TWeakPtr<ConcertClientSharedSlate::IReplicationEditorView> WeakEditorView;
+		TWeakPtr<SReplicationClientView> WeakLocalClientView;
 
 		/** Notification about in progress authority change, if any. */
 		TSharedPtr<SNotificationItem> AuthorityChangeNotification;
-		
-		/** Called when any of the streams change. */
-		void OnModelChanged() const;
 
 		void OnAuthorityRequestSent_AnyThread(const ConcertSyncClient::Replication::FAuthorityChangeRequest& Request);
 		void OnAuthorityResponseReceived_AnyThread(
