Display other clients' registered streams

diff --git a/Engine/Plugins/Developer/Concert/ConcertApp/MultiUserClient/Source/MultiUserClient/Private/Replication/Authority/AuthorityPolicy.cpp b/Engine/Plugins/Developer/Concert/ConcertApp/MultiUserClient/Source/MultiUserClient/Private/Replication/Authority/AuthorityPolicy.cpp
--- a/Engine/Plugins/Developer/Concert/ConcertApp/MultiUserClient/Source/MultiUserClient/Private/Replication/Authority/AuthorityPolicy.cpp
+++ b/Engine/Plugins/Developer/Concert/ConcertApp/MultiUserClient/Source/MultiUserClient/Private/Replication/Authority/AuthorityPolicy.cpp
@@ -39,17 +39,17 @@ namespace UE::MultiUserClient
 		}
 	}
 	
-	FAuthorityPolicy::FAuthorityPolicy(TSharedRef<IConcertSyncClient> InClient, TSharedRef<IClientStreamSynchronizer> InStreamSynchronizer)
+	FAuthorityPolicy::FAuthorityPolicy(TSharedRef<IConcertSyncClient> InClient, IClientStreamSynchronizer&  InStreamSynchronizer)
 		: Client(MoveTemp(InClient))
-		, StreamSynchronizer(MoveTemp(InStreamSynchronizer))
+		, StreamSynchronizer(InStreamSynchronizer)
 	{
 		// After submitting new objects, immediately request authority over it.
-		StreamSynchronizer->OnChangesAccepted_AnyThread().AddRaw(this, &FAuthorityPolicy::OnChangesAccepted_AnyThread);
+		StreamSynchronizer.OnChangesAccepted().AddRaw(this, &FAuthorityPolicy::OnChangesAccepted_AnyThread);
 	}
 
 	FAuthorityPolicy::~FAuthorityPolicy()
 	{
-		StreamSynchronizer->OnChangesAccepted_AnyThread().RemoveAll(this);
+		StreamSynchronizer.OnChangesAccepted().RemoveAll(this);
 	}
 
 	void FAuthorityPolicy::OnChangesAccepted_AnyThread(
diff --git a/Engine/Plugins/Developer/Concert/ConcertApp/MultiUserClient/Source/MultiUserClient/Private/Replication/Authority/AuthorityPolicy.h b/Engine/Plugins/Developer/Concert/ConcertApp/MultiUserClient/Source/MultiUserClient/Private/Replication/Authority/AuthorityPolicy.h
--- a/Engine/Plugins/Developer/Concert/ConcertApp/MultiUserClient/Source/MultiUserClient/Private/Replication/Authority/AuthorityPolicy.h
+++ b/Engine/Plugins/Developer/Concert/ConcertApp/MultiUserClient/Source/MultiUserClient/Private/Replication/Authority/AuthorityPolicy.h
@@ -33,7 +33,7 @@ namespace UE::MultiUserClient
 	{
 	public:
 		
-		FAuthorityPolicy(TSharedRef<IConcertSyncClient> InClient, TSharedRef<IClientStreamSynchronizer> InStreamSynchronizer);
+		FAuthorityPolicy(TSharedRef<IConcertSyncClient> InClient, IClientStreamSynchronizer& InStreamSynchronizer);
 		~FAuthorityPolicy();
 
 		DECLARE_MULTICAST_DELEGATE_OneParam(FOnAuthorityRequestSent,
@@ -56,8 +56,11 @@ namespace UE::MultiUserClient
 
 		/** Used to send authority requests to the server. */
 		const TSharedRef<IConcertSyncClient> Client;
-		/** Informs us when changes have been successfully submitted to the server. */
-		TSharedRef<IClientStreamSynchronizer> StreamSynchronizer;
+		/**
+		 * Informs us when changes have been successfully submitted to the server.
+		 * Outlives this FAuthorityPolicy.
+		 */
+		IClientStreamSynchronizer& StreamSynchronizer;
 
 		FOnAuthorityRequestSent OnAuthorityRequestSentDelegate;
 		FOnAuthorityResponseReceived OnAuthorityResponseReceivedDelegate;
diff --git a/Engine/Plugins/Developer/Concert/ConcertApp/MultiUserClient/Source/MultiUserClient/Private/Replication/Client/LocalReplicationClient.cpp b/Engine/Plugins/Developer/Concert/ConcertApp/MultiUserClient/Source/MultiUserClient/Private/Replication/Client/LocalReplicationClient.cpp
new file mode 100644
--- /dev/null
+++ b/Engine/Plugins/Developer/Concert/ConcertApp/MultiUserClient/Source/MultiUserClient/Private/Replication/Client/LocalReplicationClient.cpp
@@ -0,0 +1,16 @@
+// Copyright Epic Games, Inc. All Rights Reserved.
+
+#include "LocalReplicationClient.h"
+
+namespace UE::MultiUserClient
+{
+	FLocalReplicationClient::FLocalReplicationClient(
+		UMultiUserReplicationClientPreset& InSessionContent,
+		TUniquePtr<IClientStreamSynchronizer> InStreamSynchronizer,
+		TSharedRef<IConcertSyncClient> InClient
+		)
+		: FReplicationClient(InSessionContent, MoveTemp(InStreamSynchronizer))
+		, AuthorityPolicy(MoveTemp(InClient), GetStreamSynchronizer())
+	{
+	}
+}
diff --git a/Engine/Plugins/Developer/Concert/ConcertApp/MultiUserClient/Source/MultiUserClient/Private/Replication/Client/LocalReplicationClient.h b/Engine/Plugins/Developer/Concert/ConcertApp/MultiUserClient/Source/MultiUserClient/Private/Replication/Client/LocalReplicationClient.h
new file mode 100644
--- /dev/null
+++ b/Engine/Plugins/Developer/Concert/ConcertApp/MultiUserClient/Source/MultiUserClient/Private/Replication/Client/LocalReplicationClient.h
@@ -0,0 +1,29 @@
+// Copyright Epic Games, Inc. All Rights Reserved.
+
+#pragma once
+
+#include "ReplicationClient.h"
+#include "Replication/Authority/AuthorityPolicy.h"
+
+namespace UE::MultiUserClient
+{
+	/** Holds extra information about a local replication client. */
+	class FLocalReplicationClient : public FReplicationClient
+	{
+	public:
+
+		FLocalReplicationClient(
+			UMultiUserReplicationClientPreset& InSessionContent,
+			TUniquePtr<IClientStreamSynchronizer> InStreamSynchronizer,
+			TSharedRef<IConcertSyncClient> InClient
+			);
+
+		FAuthorityPolicy& GetAuthorityPolicy() { return AuthorityPolicy; }
+
+	private:
+
+		/** Automatically requests authority for submitted, locally-owned objects. */
+		FAuthorityPolicy AuthorityPolicy;
+	};
+}
+
diff --git a/Engine/Plugins/Developer/Concert/ConcertApp/MultiUserClient/Source/MultiUserClient/Private/Replication/Client/RemoteReplicationClient.cpp b/Engine/Plugins/Developer/Concert/ConcertApp/MultiUserClient/Source/MultiUserClient/Private/Replication/Client/RemoteReplicationClient.cpp
new file mode 100644
--- /dev/null
+++ b/Engine/Plugins/Developer/Concert/ConcertApp/MultiUserClient/Source/MultiUserClient/Private/Replication/Client/RemoteReplicationClient.cpp
@@ -0,0 +1,25 @@
+// Copyright Epic Games, Inc. All Rights Reserved.
+
+#include "RemoteReplicationClient.h"
+
+#include "Assets/MultiUserReplicationClientPreset.h"
+#include "Replication/Stream/RemoteClientStreamSynchronizer.h"
+
+namespace UE::MultiUserClient
+{
+	FRemoteReplicationClient::FRemoteReplicationClient(
+		UMultiUserReplicationClientPreset& InSessionContent,
+		const FGuid& InConcertClientId,
+		FRegularQueryService& QueryService
+		)
+		: FReplicationClient(InSessionContent, MakeUnique<FRemoteClientStreamSynchronizer>(InConcertClientId, QueryService))
+		, RemoteEndpointId(InConcertClientId)
+	{
+		// When the remote client's state has changed, refresh the UI.
+		GetStreamSynchronizer().OnServerStateSynched().AddLambda([this]()
+		{
+			GetClientContent()->Stream->ReplicationMap = GetStreamSynchronizer().GetServerState();
+			OnModelExternallyChanged().Broadcast();
+		});
+	}
+}
diff --git a/Engine/Plugins/Developer/Concert/ConcertApp/MultiUserClient/Source/MultiUserClient/Private/Replication/Client/RemoteReplicationClient.h b/Engine/Plugins/Developer/Concert/ConcertApp/MultiUserClient/Source/MultiUserClient/Private/Replication/Client/RemoteReplicationClient.h
new file mode 100644
--- /dev/null
+++ b/Engine/Plugins/Developer/Concert/ConcertApp/MultiUserClient/Source/MultiUserClient/Private/Replication/Client/RemoteReplicationClient.h
@@ -0,0 +1,35 @@
+// Copyright Epic Games, Inc. All Rights Reserved.
+
+#pragma once
+
+#include "ReplicationClient.h"
+
+namespace UE::MultiUserClient
+{
+	class FRegularQueryService;
+	
+	/** Holds extra information about a remote replication client. */
+	class FRemoteReplicationClient : public FReplicationClient
+	{
+	public:
+
+		FRemoteReplicationClient(
+			UMultiUserReplicationClientPreset& InSessionContent,
+			const FGuid& InConcertClientId,
+			FRegularQueryService& QueryService
+			);
+		
+		const FGuid& GetRemoteEndpointId() const { return RemoteEndpointId; }
+		
+	private:
+
+		/**
+		 * Endpoint of the Concert client.
+		 *
+		 * Can be used to look up client info.
+		 * @see IConcertSession::FindSessionClient
+		 */
+		const FGuid RemoteEndpointId;
+	};
+}
+
diff --git a/Engine/Plugins/Developer/Concert/ConcertApp/MultiUserClient/Source/MultiUserClient/Private/Replication/Client/ReplicationClient.cpp b/Engine/Plugins/Developer/Concert/ConcertApp/MultiUserClient/Source/MultiUserClient/Private/Replication/Client/ReplicationClient.cpp
--- a/Engine/Plugins/Developer/Concert/ConcertApp/MultiUserClient/Source/MultiUserClient/Private/Replication/Client/ReplicationClient.cpp
+++ b/Engine/Plugins/Developer/Concert/ConcertApp/MultiUserClient/Source/MultiUserClient/Private/Replication/Client/ReplicationClient.cpp
@@ -11,14 +11,13 @@ namespace UE::MultiUserClient
 {
 	FReplicationClient::FReplicationClient(
 		UMultiUserReplicationClientPreset& InSessionContent,
-		TSharedRef<IConcertSyncClient> InClient,
-		TSharedRef<IClientStreamSynchronizer> InStreamSynchronizer
+		TUniquePtr<IClientStreamSynchronizer> InStreamSynchronizer
 		)
 		: ClientContentStorage(&InSessionContent)
 		, StreamSynchronizer(MoveTemp(InStreamSynchronizer))
 		, LocalClientEditModel(ConcertClientSharedSlate::CreatePropertySelectionModel(*ClientContentStorage->Stream, ClientContentStorage->Stream->MakeReplicationMapGetterAttribute()))
 		, LocalClientStreamDiffer(
-			StreamSynchronizer,
+			GetStreamSynchronizer(),
 			ClientContentStorage->Stream->MakeReplicationMapGetterAttribute(),
 			FLocalStreamChangeTracker::FOnModifyReplicationMap::CreateLambda([this](){ ClientContentStorage->Stream->Modify(); })
 			)
diff --git a/Engine/Plugins/Developer/Concert/ConcertApp/MultiUserClient/Source/MultiUserClient/Private/Replication/Client/ReplicationClient.h b/Engine/Plugins/Developer/Concert/ConcertApp/MultiUserClient/Source/MultiUserClient/Private/Replication/Client/ReplicationClient.h
--- a/Engine/Plugins/Developer/Concert/ConcertApp/MultiUserClient/Source/MultiUserClient/Private/Replication/Client/ReplicationClient.h
+++ b/Engine/Plugins/Developer/Concert/ConcertApp/MultiUserClient/Source/MultiUserClient/Private/Replication/Client/ReplicationClient.h
@@ -24,13 +24,12 @@ namespace UE::MultiUserClient
 
 		FReplicationClient(
 			UMultiUserReplicationClientPreset& InSessionContent,
-			TSharedRef<IConcertSyncClient> InClient,
-			TSharedRef<IClientStreamSynchronizer> InStreamSynchronizer
+			TUniquePtr<IClientStreamSynchronizer> InStreamSynchronizer
 			);
 
 		UMultiUserReplicationClientPreset* GetClientContent() const { return ClientContentStorage; }
 		TSharedRef<ConcertClientSharedSlate::IEditableObjectToPropertiesModel> GetClientEditModel() const { return LocalClientEditModel; }
-		TSharedRef<IClientStreamSynchronizer> GetStreamSynchronizer() const { return StreamSynchronizer; } 
+		IClientStreamSynchronizer& GetStreamSynchronizer() const { return *StreamSynchronizer.Get(); } 
 		
 		const FLocalStreamChangeTracker& GetDiffer() const { return LocalClientStreamDiffer; }
 		FLocalStreamChangeTracker& GetDiffer() { return LocalClientStreamDiffer; }
@@ -45,7 +44,7 @@ namespace UE::MultiUserClient
 		 * or the local client reverted a change that was rejected by the server.
 		 */
 		DECLARE_MULTICAST_DELEGATE(FOnModelExternallyChanged);
-		FOnModelExternallyChanged& OnModelExternallyChanged_GameThread() { return OnModelExternallyChangedDelegate; }
+		FOnModelExternallyChanged& OnModelExternallyChanged() { return OnModelExternallyChangedDelegate; }
 		
 	private:
 		
@@ -53,7 +52,7 @@ namespace UE::MultiUserClient
 		TObjectPtr<UMultiUserReplicationClientPreset> ClientContentStorage;
 		
 		/** Keeps the client's state on the server in sync. */
-		TSharedRef<IClientStreamSynchronizer> StreamSynchronizer;
+		TUniquePtr<IClientStreamSynchronizer> StreamSynchronizer;
 		
 		/**
 		 * Used to detect changes made to the client's config by the local editor.
diff --git a/Engine/Plugins/Developer/Concert/ConcertApp/MultiUserClient/Source/MultiUserClient/Private/Replication/Client/ReplicationClientManager.cpp b/Engine/Plugins/Developer/Concert/ConcertApp/MultiUserClient/Source/MultiUserClient/Private/Replication/Client/ReplicationClientManager.cpp
--- a/Engine/Plugins/Developer/Concert/ConcertApp/MultiUserClient/Source/MultiUserClient/Private/Replication/Client/ReplicationClientManager.cpp
+++ b/Engine/Plugins/Developer/Concert/ConcertApp/MultiUserClient/Source/MultiUserClient/Private/Replication/Client/ReplicationClientManager.cpp
@@ -2,6 +2,9 @@
 
 #include "ReplicationClientManager.h"
 
+#include "IConcertSyncClient.h"
+#include "LocalReplicationClient.h"
+#include "RemoteReplicationClient.h"
 #include "Assets/MultiUserReplicationSessionPreset.h"
 #include "Replication/Stream/LocalClientStreamSynchronizer.h"
 
@@ -10,17 +13,92 @@
 
 namespace UE::MultiUserClient
 {
-	FReplicationClientManager::FReplicationClientManager(TSharedRef<IConcertSyncClient> InClient)
-		: SessionContent(NewObject<UMultiUserReplicationSessionPreset>(GetTransientPackage(), NAME_None, RF_Transient)) 
-		, LocalClient([this, &InClient]()
+	FReplicationClientManager::FReplicationClientManager(
+		TSharedRef<IConcertSyncClient> InClient,
+		TSharedRef<IConcertClientSession> InSession
+		)
+		: SessionContent(NewObject<UMultiUserReplicationSessionPreset>(GetTransientPackage(), NAME_None, RF_Transient))
+		, Session(InSession)
+		, QueryService(InClient)
+		, LocalClient([this, InClient]()
 		{
 			UMultiUserReplicationClientPreset* ClientPreset = SessionContent->AddClient();
-			return FReplicationClient(*ClientPreset, InClient, MakeShared<FLocalClientStreamSynchronizer>(InClient, ClientPreset->Stream->StreamId));
+			return FLocalReplicationClient(*ClientPreset, MakeUnique<FLocalClientStreamSynchronizer>(InClient, ClientPreset->Stream->StreamId), InClient);
 		}())
-	{}
+	{
+		InSession->OnSessionClientChanged().AddRaw(this, &FReplicationClientManager::OnSessionClientChanged);
+
+		for (const FGuid& ClientEndpointId : InSession->GetSessionClientEndpointIds())
+		{
+			// Nobody is subscribed at this point, so no need to broadcast
+			constexpr bool bBroadcast = false;
+			CreateRemoteClient(ClientEndpointId, bBroadcast);
+		}
+	}
+
+	FReplicationClientManager::~FReplicationClientManager()
+	{
+		if (const TSharedPtr<IConcertClientSession> SessionPin = Session.Pin())
+		{
+			SessionPin->OnSessionClientChanged().RemoveAll(this);
+		}
+	}
+
+	const FRemoteReplicationClient* FReplicationClientManager::FindRemoteClient(const FGuid& EndpointId) const
+	{
+		return RemoteClients.FindByPredicate([&EndpointId](const FRemoteReplicationClient& Client)
+		{
+			return Client.GetRemoteEndpointId() == EndpointId;
+		});
+	}
 
 	void FReplicationClientManager::AddReferencedObjects(FReferenceCollector& Collector)
 	{
 		Collector.AddReferencedObject(SessionContent);
 	}
+
+	void FReplicationClientManager::OnSessionClientChanged(IConcertClientSession&, EConcertClientStatus NewStatus, const FConcertSessionClientInfo& ClientInfo)
+	{
+		const FGuid& ClientEndpointId = ClientInfo.ClientEndpointId;
+		switch (NewStatus)
+		{
+			
+		case EConcertClientStatus::Connected:
+			CreateRemoteClient(ClientEndpointId);
+			break;
+			
+		case EConcertClientStatus::Disconnected:
+			{
+				const int32 Index = RemoteClients.IndexOfByPredicate(
+					[&ClientEndpointId](const FRemoteReplicationClient& Client)
+					{
+						return Client.GetRemoteEndpointId() == ClientEndpointId;
+					});
+				if (!ensure(RemoteClients.IsValidIndex(Index)))
+				{
+					return;
+				}
+				
+				const FRemoteReplicationClient& Client = RemoteClients[Index];
+				SessionContent->RemoveClient(*Client.GetClientContent());
+				RemoteClients.RemoveAtSwap(Index);
+				OnRemoteClientsChangedDelegate.Broadcast();
+			}
+			break;
+			
+		case EConcertClientStatus::Updated:
+			break;
+		default: checkNoEntry();
+		}
+	}
+	
+	void FReplicationClientManager::CreateRemoteClient(const FGuid& ClientEndpointId, bool bBroadcastDelegate)
+	{
+		RemoteClients.Emplace(*SessionContent->AddClient(), ClientEndpointId, QueryService);
+
+		if (bBroadcastDelegate)
+		{
+			OnRemoteClientsChangedDelegate.Broadcast();
+		}
+	}
 }
diff --git a/Engine/Plugins/Developer/Concert/ConcertApp/MultiUserClient/Source/MultiUserClient/Private/Replication/Client/ReplicationClientManager.h b/Engine/Plugins/Developer/Concert/ConcertApp/MultiUserClient/Source/MultiUserClient/Private/Replication/Client/ReplicationClientManager.h
--- a/Engine/Plugins/Developer/Concert/ConcertApp/MultiUserClient/Source/MultiUserClient/Private/Replication/Client/ReplicationClientManager.h
+++ b/Engine/Plugins/Developer/Concert/ConcertApp/MultiUserClient/Source/MultiUserClient/Private/Replication/Client/ReplicationClientManager.h
@@ -2,14 +2,22 @@
 
 #pragma once
 
-#include "ReplicationClient.h"
 #include "Assets/MultiUserReplicationSessionPreset.h"
-#include "Replication/Editor/Model/IEditableObjectToPropertiesModel.h"
+#include "ConcertMessageData.h"
+#include "LocalReplicationClient.h"
+#include "ReplicationClient.h"
+#include "Replication/Util/RegularQueryService.h"
+
 #include "UObject/GCObject.h"
 #include "Templates/UnrealTemplate.h"
 
+enum class EConcertClientStatus : uint8;
+
+class IConcertClientSession;
 class IConcertSyncClient;
 
+struct FConcertSessionClientInfo;
+
 namespace UE::ConcertClientSharedSlate
 {
 	class IEditableObjectToPropertiesModel;
@@ -17,19 +25,42 @@ namespace UE::ConcertClientSharedSlate
 
 namespace UE::MultiUserClient
 {
+	class FRemoteReplicationClient;
 	class FReplicationClient;
 	class FLocalStreamChangeTracker;
 
-	/** Keeps track of connected clients synchronizing their stream data in a UMultiUserReplicationSessionPreset. */
+	/**
+	 * Keeps track of connected clients synchronizing their stream data in a UMultiUserReplicationSessionPreset.
+	 * This object only exists for as long as the local client is in a Concert session.
+	 */
 	class FReplicationClientManager
 		: public FGCObject
 		, public FNoncopyable
 	{
 	public:
+		
+		/**
+		 * @param InClient The local client. Outlives this objects.
+		 * @param InSession The session to observe. Outlives this objects.
+		 */
+		FReplicationClientManager(TSharedRef<IConcertSyncClient> InClient, TSharedRef<IConcertClientSession> InSession);
+		virtual ~FReplicationClientManager() override;
 
-		FReplicationClientManager(TSharedRef<IConcertSyncClient> InClient);
+		const FLocalReplicationClient& GetLocalClient() const { return LocalClient; }
+		FLocalReplicationClient& GetLocalClient() { return LocalClient; }
+		const TArray<FRemoteReplicationClient>& GetRemoteClients() const { return RemoteClients; }
 
-		FReplicationClient& GetLocalClient() { return LocalClient; }
+		/** Util for finding a remote client by its EndpointId. */
+		const FRemoteReplicationClient* FindRemoteClient(const FGuid& EndpointId) const;
+		FRemoteReplicationClient* FindRemoteClient(const FGuid& EndpointId)
+		{
+			const FReplicationClientManager* ConstThis = this;
+			return const_cast<FRemoteReplicationClient*>(ConstThis->FindRemoteClient(EndpointId));
+		}
+		
+		DECLARE_MULTICAST_DELEGATE(FRemoteClientsChanged);
+		/** Called when RemoteClients changes. */
+		FRemoteClientsChanged& OnRemoteClientsChanged() { return OnRemoteClientsChangedDelegate; }
 		
 		//~ Begin FGCObject Interface
 		virtual void AddReferencedObjects(FReferenceCollector& Collector) override;
@@ -41,8 +72,33 @@ namespace UE::MultiUserClient
 		/** The state of the server is synched up with this object and displayed in the UI. */
 		TObjectPtr<UMultiUserReplicationSessionPreset> SessionContent;
 
+		/**
+		 * The session the local client is in.
+		 * 
+		 * This FReplicationClientManager's owner is supposed to make sure this FReplicationClientManager is destroyed
+		 * when the session shuts down.
+		 */
+		const TWeakPtr<IConcertClientSession> Session;
+
+		/**
+		 * Sends FConcertReplication_QueryReplicationInfo_Request in regular intervals.
+		 * Shared by all remote clients so all requests are bundled reducing the number of network requests. 
+		 */
+		FRegularQueryService QueryService;
+		
 		/** Manages the local client */
-		FReplicationClient LocalClient;
+		FLocalReplicationClient LocalClient;
+		/** Manages remote clients. Updated when client connects or disconnects to the active session. */
+		TArray<FRemoteReplicationClient> RemoteClients;
+
+		/** Called when RemoteClients changes. */
+		FRemoteClientsChanged OnRemoteClientsChangedDelegate;
+
+		/** Updates RemoteClients depending on the change. */
+		void OnSessionClientChanged(IConcertClientSession&, EConcertClientStatus NewStatus, const FConcertSessionClientInfo& ClientInfo);
+		
+		/** Shared logic for creating a remote client. */
+		void CreateRemoteClient(const FGuid& ClientEndpointId, bool bBroadcastDelegate = true);
 	};
 }
 
diff --git a/Engine/Plugins/Developer/Concert/ConcertApp/MultiUserClient/Source/MultiUserClient/Private/Replication/MultiUserReplicationManager.cpp b/Engine/Plugins/Developer/Concert/ConcertApp/MultiUserClient/Source/MultiUserClient/Private/Replication/MultiUserReplicationManager.cpp
--- a/Engine/Plugins/Developer/Concert/ConcertApp/MultiUserClient/Source/MultiUserClient/Private/Replication/MultiUserReplicationManager.cpp
+++ b/Engine/Plugins/Developer/Concert/ConcertApp/MultiUserClient/Source/MultiUserClient/Private/Replication/MultiUserReplicationManager.cpp
@@ -50,6 +50,13 @@ namespace UE::MultiUserClient
 			});
 	}
 
+	FAuthorityPolicy* FMultiUserReplicationManager::GetAuthorityPolicy()
+	{
+		return ConnectedState
+			? &ConnectedState->ClientManager.GetLocalClient().GetAuthorityPolicy()
+			: nullptr;
+	}
+
 	void FMultiUserReplicationManager::OnSessionConnectionChanged(
 		IConcertClientSession& ConcertClientSession,
 		EConcertConnectionStatus ConcertConnectionStatus
@@ -71,12 +78,12 @@ namespace UE::MultiUserClient
 		}
 	}
 
-	void FMultiUserReplicationManager::OnLeaveSession(IConcertClientSession& ConcertClientSession)
+	void FMultiUserReplicationManager::OnLeaveSession(IConcertClientSession&)
 	{
-		// Keep in mind that FClientStreamRepository::GetLocalClientEditModel is referenced by the UI ...
-		ConnectedState.Reset();
-		// ... and after this broadcast the model should no longer be referenced by anyone
+		// This clears the UI. The clients' IEditableObjectToPropertiesModels should no longer be referenced by anyone.
 		SetConnectionStateAndBroadcast(EMultiUserReplicationConnectionState::Disconnected);
+		// Keep in mind that FClientStreamRepository::GetLocalClientEditModel is referenced by the UI so call this after clearing the UI.
+		ConnectedState.Reset();
 	}
 
 	void FMultiUserReplicationManager::HandleReplicationSessionJoined(const ConcertSyncClient::Replication::FJoinReplicatedSessionResult& JoinSessionResult)
@@ -100,7 +107,6 @@ namespace UE::MultiUserClient
 	}
 
 	FMultiUserReplicationManager::FConnectedState::FConnectedState(TSharedRef<IConcertSyncClient> InClient)
-		: ClientManager(InClient)
-		, AuthorityPolicy(InClient, ClientManager.GetLocalClient().GetStreamSynchronizer())
+		: ClientManager(InClient, InClient->GetConcertClient()->GetCurrentSession().ToSharedRef())
 	{}
 }
diff --git a/Engine/Plugins/Developer/Concert/ConcertApp/MultiUserClient/Source/MultiUserClient/Private/Replication/MultiUserReplicationManager.h b/Engine/Plugins/Developer/Concert/ConcertApp/MultiUserClient/Source/MultiUserClient/Private/Replication/MultiUserReplicationManager.h
--- a/Engine/Plugins/Developer/Concert/ConcertApp/MultiUserClient/Source/MultiUserClient/Private/Replication/MultiUserReplicationManager.h
+++ b/Engine/Plugins/Developer/Concert/ConcertApp/MultiUserClient/Source/MultiUserClient/Private/Replication/MultiUserReplicationManager.h
@@ -68,8 +68,7 @@ namespace UE::MultiUserClient
 		const FReplicationClientManager* GetClientManager() const { return ConnectedState ? &ConnectedState->ClientManager : nullptr; }
 
 		/** @note You're not supposed to keep any reference to AuthorityPolicy since it can become invalid depending on connection state. */
-		FAuthorityPolicy* GetAuthorityPolicy() { return ConnectedState ? & ConnectedState->AuthorityPolicy : nullptr; }
-		const FAuthorityPolicy* GetAuthorityPolicy() const { return ConnectedState ? & ConnectedState->AuthorityPolicy : nullptr; }
+		FAuthorityPolicy* GetAuthorityPolicy();
 
 		/** Called when the connection to the replication system changes. */
 		DECLARE_MULTICAST_DELEGATE_OneParam(FOnReplicationConnectionStateChanged, EMultiUserReplicationConnectionState /*NewState*/);
@@ -96,9 +95,6 @@ namespace UE::MultiUserClient
 			 * Only valid when ConnectionState == EMultiUserReplicationConnectionState::Connected.
 			 */
 			FReplicationClientManager ClientManager;
-
-			/** Manages the client's authority like automatically taking ownership over newly submitted objects. */
-			FAuthorityPolicy AuthorityPolicy;
 			
 			FConnectedState(TSharedRef<IConcertSyncClient> InClient);
 		};
@@ -111,7 +107,7 @@ namespace UE::MultiUserClient
 		/** Callback into Concert for when client connection has changed. */
 		void OnSessionConnectionChanged(IConcertClientSession& ConcertClientSession, EConcertConnectionStatus ConcertConnectionStatus);
 		/** Leaves the current replication session */
-		void OnLeaveSession(IConcertClientSession& ConcertClientSession);
+		void OnLeaveSession(IConcertClientSession&);
 
 		/** Handles server response for joining replication session */
 		void HandleReplicationSessionJoined(const ConcertSyncClient::Replication::FJoinReplicatedSessionResult& JoinSessionResult);
diff --git a/Engine/Plugins/Developer/Concert/ConcertApp/MultiUserClient/Source/MultiUserClient/Private/Replication/Stream/IClientStreamSynchronizer.h b/Engine/Plugins/Developer/Concert/ConcertApp/MultiUserClient/Source/MultiUserClient/Private/Replication/Stream/IClientStreamSynchronizer.h
--- a/Engine/Plugins/Developer/Concert/ConcertApp/MultiUserClient/Source/MultiUserClient/Private/Replication/Stream/IClientStreamSynchronizer.h
+++ b/Engine/Plugins/Developer/Concert/ConcertApp/MultiUserClient/Source/MultiUserClient/Private/Replication/Stream/IClientStreamSynchronizer.h
@@ -72,10 +72,10 @@ namespace UE::MultiUserClient
 
 		DECLARE_MULTICAST_DELEGATE_TwoParams(FOnChangesAccepted, const FObjectReplicationMap& OldServerState, const ConcertSyncClient::Replication::FChangeStreamRequest& AcceptedRequest);
 		/** @return Event executed when a change request that was in transit was accepted by the server. This is executed before OnServerStateSynched: the local confirmed state is not updated, yet. */
-		virtual FOnChangesAccepted& OnChangesAccepted_AnyThread() = 0;
+		virtual FOnChangesAccepted& OnChangesAccepted() = 0;
 		DECLARE_MULTICAST_DELEGATE(FOnServerStateChanged);
 		/** @return Event executed when the result of GetServerState has been synched. Called after OnChangesAccepted_AnyThread. */
-		virtual FOnServerStateChanged& OnServerStateSynched_AnyThread() = 0;
+		virtual FOnServerStateChanged& OnServerStateSynched() = 0;
 		
 		virtual ~IClientStreamSynchronizer() = default;
 	};
diff --git a/Engine/Plugins/Developer/Concert/ConcertApp/MultiUserClient/Source/MultiUserClient/Private/Replication/Stream/LocalClientStreamSynchronizer.h b/Engine/Plugins/Developer/Concert/ConcertApp/MultiUserClient/Source/MultiUserClient/Private/Replication/Stream/LocalClientStreamSynchronizer.h
--- a/Engine/Plugins/Developer/Concert/ConcertApp/MultiUserClient/Source/MultiUserClient/Private/Replication/Stream/LocalClientStreamSynchronizer.h
+++ b/Engine/Plugins/Developer/Concert/ConcertApp/MultiUserClient/Source/MultiUserClient/Private/Replication/Stream/LocalClientStreamSynchronizer.h
@@ -29,8 +29,8 @@ namespace UE::MultiUserClient
 		virtual bool CanMakeSubmitRequest() const override { return !ChangeInTransit.IsSet(); }
 		virtual FGuid GetStreamId() const override { return LocalClientStreamId; }
 		virtual const FObjectReplicationMap& GetServerState() const override { return ConfirmedServerState; }
-		virtual FOnChangesAccepted& OnChangesAccepted_AnyThread() override { return OnChangesAcceptedDelegate; }
-		virtual FOnServerStateChanged& OnServerStateSynched_AnyThread() override { return OnServerStateSynchedDelegate; }
+		virtual FOnChangesAccepted& OnChangesAccepted() override { return OnChangesAcceptedDelegate; }
+		virtual FOnServerStateChanged& OnServerStateSynched() override { return OnServerStateSynchedDelegate; }
 		//~ End IClientStreamSynchronizer Interface
 
 	private:
diff --git a/Engine/Plugins/Developer/Concert/ConcertApp/MultiUserClient/Source/MultiUserClient/Private/Replication/Stream/LocalStreamChangeTracker.cpp b/Engine/Plugins/Developer/Concert/ConcertApp/MultiUserClient/Source/MultiUserClient/Private/Replication/Stream/LocalStreamChangeTracker.cpp
--- a/Engine/Plugins/Developer/Concert/ConcertApp/MultiUserClient/Source/MultiUserClient/Private/Replication/Stream/LocalStreamChangeTracker.cpp
+++ b/Engine/Plugins/Developer/Concert/ConcertApp/MultiUserClient/Source/MultiUserClient/Private/Replication/Stream/LocalStreamChangeTracker.cpp
@@ -2,7 +2,6 @@
 
 #include "LocalStreamChangeTracker.h"
 
-#include "IConcertSyncClient.h"
 #include "Replication/ChangeStreamSharedUtils.h"
 #include "Replication/IConcertClientReplicationManager.h"
 
@@ -13,21 +12,21 @@
 namespace UE::MultiUserClient
 {
 	FLocalStreamChangeTracker::FLocalStreamChangeTracker(
-		TSharedRef<IClientStreamSynchronizer> InStreamSynchronizer,
+		IClientStreamSynchronizer& InStreamSynchronizer,
 		TAttribute<FObjectReplicationMap*> InStreamWithInProgressChangesAttribute,
 		FOnModifyReplicationMap InOnModifyReplicationMapDelegate
 		)
-		: StreamSynchronizer(MoveTemp(InStreamSynchronizer))
+		: StreamSynchronizer(InStreamSynchronizer)
 		, StreamWithInProgressChangesAttribute(MoveTemp(InStreamWithInProgressChangesAttribute))
 		, OnModifyReplicationMapDelegate(MoveTemp(InOnModifyReplicationMapDelegate))
 	{
 		check(InStreamWithInProgressChangesAttribute.IsBound() || InStreamWithInProgressChangesAttribute.IsSet());
-		StreamSynchronizer->OnServerStateSynched_AnyThread().AddRaw(this, &FLocalStreamChangeTracker::RefreshChangesCache);
+		StreamSynchronizer.OnServerStateSynched().AddRaw(this, &FLocalStreamChangeTracker::RefreshChangesCache);
 	}
 
 	FLocalStreamChangeTracker::~FLocalStreamChangeTracker()
 	{
-		StreamSynchronizer->OnServerStateSynched_AnyThread().RemoveAll(this);
+		StreamSynchronizer.OnServerStateSynched().RemoveAll(this);
 	}
 
 	void FLocalStreamChangeTracker::RefreshChangesCache()
@@ -42,7 +41,7 @@ namespace UE::MultiUserClient
 		{
 			RefreshChangesCache();
 		}
-		return StreamSynchronizer->SubmitChanges(CachedDeltaChange);
+		return StreamSynchronizer.SubmitChanges(CachedDeltaChange);
 	}
 
 	void FLocalStreamChangeTracker::RevertCachedChanges()
@@ -52,7 +51,7 @@ namespace UE::MultiUserClient
 			FScopedTransaction Transaction(LOCTEXT("RevertCachedChanges", "Revert replication changes"));
 			OnModifyReplicationMapDelegate.ExecuteIfBound();
 			
-			*StreamWithInProgressChangesAttribute.Get() = StreamSynchronizer->GetServerState();
+			*StreamWithInProgressChangesAttribute.Get() = StreamSynchronizer.GetServerState();
 			CachedDeltaChange = {};
 
 			check(IsInGameThread());
@@ -73,7 +72,7 @@ namespace UE::MultiUserClient
 
 	bool FLocalStreamChangeTracker::CanMakeSubmitNetworkRequest() const
 	{
-		return StreamSynchronizer->CanMakeSubmitRequest();
+		return StreamSynchronizer.CanMakeSubmitRequest();
 	}
 
 	EObjectWarningFlags FLocalStreamChangeTracker::GetObjectWarningFlags(const FSoftObjectPath& ObjectPath)
diff --git a/Engine/Plugins/Developer/Concert/ConcertApp/MultiUserClient/Source/MultiUserClient/Private/Replication/Stream/LocalStreamChangeTracker.h b/Engine/Plugins/Developer/Concert/ConcertApp/MultiUserClient/Source/MultiUserClient/Private/Replication/Stream/LocalStreamChangeTracker.h
--- a/Engine/Plugins/Developer/Concert/ConcertApp/MultiUserClient/Source/MultiUserClient/Private/Replication/Stream/LocalStreamChangeTracker.h
+++ b/Engine/Plugins/Developer/Concert/ConcertApp/MultiUserClient/Source/MultiUserClient/Private/Replication/Stream/LocalStreamChangeTracker.h
@@ -36,7 +36,7 @@ namespace UE::MultiUserClient
 
 		DECLARE_DELEGATE(FOnModifyReplicationMap);
 		FLocalStreamChangeTracker(
-			TSharedRef<IClientStreamSynchronizer> InStreamSynchronizer,
+			IClientStreamSynchronizer& InStreamSynchronizer,
 			TAttribute<FObjectReplicationMap*> InStreamWithInProgressChangesAttribute,
 			FOnModifyReplicationMap InOnModifyReplicationMapDelegate
 			);
@@ -101,8 +101,8 @@ namespace UE::MultiUserClient
 		
 	private:
 
-		/** Keeps track of the edited stream's server state. */
-		TSharedRef<IClientStreamSynchronizer> StreamSynchronizer;
+		/** Keeps track of the edited stream's server state. Outlives the FLocalStreamChangeTracker instance. */
+		IClientStreamSynchronizer& StreamSynchronizer;
 
 		/**
 		 * Represents ConfirmedServerState with changes made to it. These changes have not been sent to the server, yet.
@@ -127,7 +127,7 @@ namespace UE::MultiUserClient
 		FStreamChangelist DiffChanges() const
 		{
 			const FObjectReplicationMap* Map = StreamWithInProgressChangesAttribute.Get();
-			return Map ? DiffChanges(StreamSynchronizer->GetStreamId(), StreamSynchronizer->GetServerState(), *Map) : FStreamChangelist{};
+			return Map ? DiffChanges(StreamSynchronizer.GetStreamId(), StreamSynchronizer.GetServerState(), *Map) : FStreamChangelist{};
 		}
 		/** Builds the changelist to get from Base to Changed. */
 		static FStreamChangelist DiffChanges(const FGuid& StreamId, const FObjectReplicationMap& Base, const FObjectReplicationMap& Changed);
diff --git a/Engine/Plugins/Developer/Concert/ConcertApp/MultiUserClient/Source/MultiUserClient/Private/Replication/Stream/RemoteClientStreamSynchronizer.cpp b/Engine/Plugins/Developer/Concert/ConcertApp/MultiUserClient/Source/MultiUserClient/Private/Replication/Stream/RemoteClientStreamSynchronizer.cpp
new file mode 100644
--- /dev/null
+++ b/Engine/Plugins/Developer/Concert/ConcertApp/MultiUserClient/Source/MultiUserClient/Private/Replication/Stream/RemoteClientStreamSynchronizer.cpp
@@ -0,0 +1,68 @@
+// Copyright Epic Games, Inc. All Rights Reserved.
+
+#include "RemoteClientStreamSynchronizer.h"
+
+#include "ConcertLogGlobal.h"
+#include "Assets/MultiUserReplicationClientPreset.h"
+#include "Replication/Util/RegularQueryService.h"
+
+namespace UE::MultiUserClient
+{
+	constexpr float QueryTimeInterval = 1.f;
+	
+	FRemoteClientStreamSynchronizer::FRemoteClientStreamSynchronizer(const FGuid& RemoteEndpointId, FRegularQueryService& InQueryService)
+		: QueryService(InQueryService)
+		, QueryStreamHandle(
+			QueryService.RegisterStreamQuery(
+				RemoteEndpointId,
+				FStreamQueryDelegate::CreateRaw(this, &FRemoteClientStreamSynchronizer::HandleStreamQuery)
+				)
+			)
+	{}
+
+	FRemoteClientStreamSynchronizer::~FRemoteClientStreamSynchronizer()
+	{
+		QueryService.UnregisterStreamQuery(QueryStreamHandle);
+	}
+
+	TFuture<FSubmitChangesResult> FRemoteClientStreamSynchronizer::SubmitChanges(const FStreamChangelist& Changelist)
+	{
+		
+		// UE-180657: Changing remote client's stream is not implemented for now
+		return MakeFulfilledPromise<FSubmitChangesResult>(FSubmitChangesResult{ ESubmitChangesErrorCode::CannotSendRequest }).GetFuture();
+	}
+
+	bool FRemoteClientStreamSynchronizer::CanMakeSubmitRequest() const
+	{
+		// UE-180657: Changing remote client's stream is not implemented for now
+		return false;
+	}
+
+	FGuid FRemoteClientStreamSynchronizer::GetStreamId() const
+	{
+		return UMultiUserReplicationClientPreset::MultiUserStreamID;
+	}
+
+	void FRemoteClientStreamSynchronizer::HandleStreamQuery(const TArray<FSharedReplicationStreamDescription>& Streams)
+	{
+		// MU clients use UMultiUserReplicationClientPreset::MultiUserStreamID for streams.
+		// That handles the (unlikely) case in which some external logic had added streams to the same client which we must differentiate
+		const FSharedReplicationStreamDescription* MultiUserDescription = Streams.FindByPredicate([](const FSharedReplicationStreamDescription& Description)
+		{
+			return Description.Identifier == UMultiUserReplicationClientPreset::MultiUserStreamID;
+		});
+		if (MultiUserDescription
+			// Avoid unnecessary Broadcast()s
+			&& LastKnownServerState != MultiUserDescription->ReplicationMap)
+		{
+			LastKnownServerState = MultiUserDescription->ReplicationMap;
+			OnServerStateSynchedDelegate.Broadcast();
+		}
+		// If the remote client removed the last object, the stream is implicitly deleted.
+		else if (!MultiUserDescription && !LastKnownServerState.ReplicatedObjects.IsEmpty())
+		{
+			LastKnownServerState = {};
+			OnServerStateSynchedDelegate.Broadcast();
+		}
+	}
+}
diff --git a/Engine/Plugins/Developer/Concert/ConcertApp/MultiUserClient/Source/MultiUserClient/Private/Replication/Stream/RemoteClientStreamSynchronizer.h b/Engine/Plugins/Developer/Concert/ConcertApp/MultiUserClient/Source/MultiUserClient/Private/Replication/Stream/RemoteClientStreamSynchronizer.h
new file mode 100644
--- /dev/null
+++ b/Engine/Plugins/Developer/Concert/ConcertApp/MultiUserClient/Source/MultiUserClient/Private/Replication/Stream/RemoteClientStreamSynchronizer.h
@@ -0,0 +1,57 @@
+// Copyright Epic Games, Inc. All Rights Reserved.
+
+#pragma once
+
+#include "IClientStreamSynchronizer.h"
+#include "Replication/Data/ObjectReplicationMap.h"
+
+#include "Async/Future.h"
+#include "Misc/Attribute.h"
+#include "Templates/UnrealTemplate.h"
+
+class IConcertSyncClient;
+
+namespace UE::MultiUserClient
+{
+	class FRegularQueryService;
+	
+	/**
+	 * Tracks the state of a remote client.
+	 * Queries the client's state in regular intervals.
+	 */
+	class FRemoteClientStreamSynchronizer : public IClientStreamSynchronizer, public FNoncopyable
+	{
+	public:
+		
+		FRemoteClientStreamSynchronizer(const FGuid& RemoteEndpointId, FRegularQueryService& InQueryService);
+		virtual ~FRemoteClientStreamSynchronizer() override;
+
+		//~ Begin IClientStreamSynchronizer Interface
+		virtual TFuture<FSubmitChangesResult> SubmitChanges(const FStreamChangelist& Changelist) override;
+		virtual bool CanMakeSubmitRequest() const override;
+		virtual FGuid GetStreamId() const override;
+		virtual const FObjectReplicationMap& GetServerState() const override { return LastKnownServerState; }
+		virtual FOnChangesAccepted& OnChangesAccepted() override { return OnChangesAcceptedDelegate; }
+		virtual FOnServerStateChanged& OnServerStateSynched() override { return OnServerStateSynchedDelegate; }
+		//~ End IClientStreamSynchronizer Interface
+
+	private:
+
+		/** Queries the server in regular intervals. This services outlives our object. */
+		FRegularQueryService& QueryService;
+		/** Used to unregister upon destruction. */
+		const FDelegateHandle QueryStreamHandle; 
+
+		/** Represents what the local client thinks the replication map on the server currently looks like. */
+		FObjectReplicationMap LastKnownServerState;
+		
+		/** Called when a change request that was in transit was accepted by the server. */
+		FOnChangesAccepted OnChangesAcceptedDelegate;
+		/** Event executed when the result of GetServerState has been synched. */
+		FOnServerStateChanged OnServerStateSynchedDelegate;
+
+		/** Called in regular intervals with the contents of the remote client. */
+		void HandleStreamQuery(const TArray<FSharedReplicationStreamDescription>& Streams);
+	};
+}
+
diff --git a/Engine/Plugins/Developer/Concert/ConcertApp/MultiUserClient/Source/MultiUserClient/Private/Replication/Util/RegularQueryService.cpp b/Engine/Plugins/Developer/Concert/ConcertApp/MultiUserClient/Source/MultiUserClient/Private/Replication/Util/RegularQueryService.cpp
new file mode 100644
--- /dev/null
+++ b/Engine/Plugins/Developer/Concert/ConcertApp/MultiUserClient/Source/MultiUserClient/Private/Replication/Util/RegularQueryService.cpp
@@ -0,0 +1,202 @@
+// Copyright Epic Games, Inc. All Rights Reserved.
+
+#include "RegularQueryService.h"
+
+#include "IConcertSyncClient.h"
+#include "Replication/IConcertClientReplicationManager.h"
+
+#include "Templates/UnrealTemplate.h"
+
+namespace UE::MultiUserClient
+{
+	FRegularQueryService::FRegularQueryService(
+		TSharedRef<IConcertSyncClient> InOwningClient,
+		float InInterval
+		)
+		: OwningClient(MoveTemp(InOwningClient))
+		, TickerDelegateHandle(FTSTicker::GetCoreTicker().AddTicker(
+			TEXT("Multi-User Replication Query"),
+			InInterval,
+			[this](float)
+			{
+				SendQueryEvent();
+				return true;
+			}))
+	{}
+
+	FRegularQueryService::~FRegularQueryService()
+	{
+		FTSTicker::GetCoreTicker().RemoveTicker(TickerDelegateHandle);
+	}
+
+	FDelegateHandle FRegularQueryService::RegisterStreamQuery(const FGuid& EndpointId, FStreamQueryDelegate Delegate)
+	{
+		FStreamQueryInfo& Info = StreamQueryInfos.FindOrAdd(EndpointId);
+		const FDelegateHandle DelegateHandle = Info.Delegate.AddLambda([Delegate = MoveTemp(Delegate)](const TArray<FSharedReplicationStreamDescription>& Descriptions)
+		{
+			Delegate.Execute(Descriptions);
+		});
+		Info.Handles.Add(DelegateHandle);
+		return DelegateHandle;
+	}
+
+	FDelegateHandle FRegularQueryService::RegisterAuthorityQuery(const FGuid& EndpointId, FAuthorityQueryDelegate Delegate)
+	{
+		FAuthorityQueryInfo& Info = AuthorityQueryInfos.FindOrAdd(EndpointId);
+		const FDelegateHandle DelegateHandle = Info.Delegate.AddLambda([Delegate = MoveTemp(Delegate)](const TArray<FReplicationAuthorityInfo>& Infos)
+		{
+			Delegate.Execute(Infos);
+		});
+		Info.Handles.Add(DelegateHandle);
+		return DelegateHandle;
+	}
+
+	void FRegularQueryService::UnregisterStreamQuery(const FDelegateHandle& Handle)
+	{
+		for (auto StreamIt = StreamQueryInfos.CreateIterator(); StreamIt; ++StreamIt)
+		{
+			FStreamQueryInfo& Info = StreamIt->Value;
+			if (Info.Handles.Contains(Handle))
+			{
+				Info.Handles.Remove(Handle);
+				Info.Delegate.Remove(Handle);
+				
+				// If bIsHandlingQueryResponse, then HandleQueryResponse is executing delegates and a delegate is unregistering.
+				// Removing now would cause a crash since Broadcast will read some bookkeeping memory after Broadcast.
+				// Memory is cleaned up later by HandleQueryResponse.
+				if (!bIsHandlingQueryResponse && Info.Handles.IsEmpty())
+				{
+					StreamIt.RemoveCurrent();
+				}
+				
+				return;
+			}
+		}
+	}
+
+	void FRegularQueryService::UnregisterAuthorityQuery(const FDelegateHandle& Handle)
+	{
+		for (auto AuthorityIt = AuthorityQueryInfos.CreateIterator(); AuthorityIt; ++AuthorityIt)
+		{
+			FAuthorityQueryInfo& Info = AuthorityIt->Value;
+			if (Info.Handles.Contains(Handle))
+			{
+				Info.Handles.Remove(Handle);
+				Info.Delegate.Remove(Handle);
+				
+				// If bIsHandlingQueryResponse, then HandleQueryResponse is executing delegates and a delegate is unregistering.
+				// Removing now would cause a crash since Broadcast will read some bookkeeping memory after Broadcast.
+				// Memory is cleaned up later by HandleQueryResponse.
+				if (!bIsHandlingQueryResponse && Info.Handles.IsEmpty())
+				{
+					AuthorityIt.RemoveCurrent();
+				}
+				
+				return;
+			}
+		}
+	}
+
+	void FRegularQueryService::SendQueryEvent()
+	{
+		const TSharedPtr<IConcertSyncClient> SessionPin = OwningClient.Pin();
+		IConcertClientReplicationManager* ReplicationManager = SessionPin ? SessionPin->GetReplicationManager() : nullptr;
+		if (!ensure(ReplicationManager))
+		{
+			return;
+		}
+		
+		FConcertReplication_QueryReplicationInfo_Request Request;
+		BuildStreamRequest(Request);
+		BuildAuthorityRequest(Request);
+
+		if (!Request.ClientEndpointIds.IsEmpty())
+		{
+			ReplicationManager->QueryClientInfo({ Request })
+				.Next([this, WeakToken = Token.ToWeakPtr()](ConcertSyncClient::Replication::FClientQueryResponse&& Response)
+				{
+					const bool bCanProcessRequest = WeakToken.Pin().IsValid();
+					if (bCanProcessRequest)
+					{
+						HandleQueryResponse(Response);
+					}
+				});
+		}
+	}
+
+	void FRegularQueryService::BuildStreamRequest(FConcertReplication_QueryReplicationInfo_Request& Request) const
+	{
+		if (StreamQueryInfos.IsEmpty())
+		{
+			Request.QueryFlags |= EConcertQueryClientStreamFlags::SkipStreamInfo;
+		}
+		else
+		{
+			for (const TPair<FGuid, FStreamQueryInfo>& Query : StreamQueryInfos)
+			{
+				Request.ClientEndpointIds.Add(Query.Key);
+			}
+		}
+	}
+
+	void FRegularQueryService::BuildAuthorityRequest(FConcertReplication_QueryReplicationInfo_Request& Request) const
+	{
+		if (AuthorityQueryInfos.IsEmpty())
+		{
+			Request.QueryFlags |= EConcertQueryClientStreamFlags::SkipAuthority;
+		}
+		else
+		{
+			for (const TPair<FGuid, FAuthorityQueryInfo>& Query : AuthorityQueryInfos)
+			{
+				Request.ClientEndpointIds.Add(Query.Key);
+			}
+		}
+	}
+
+	void FRegularQueryService::HandleQueryResponse(
+		const FConcertReplication_QueryReplicationInfo_Response& Response
+		)
+	{
+		TGuardValue<bool> Guard(bIsHandlingQueryResponse, true);
+		
+		for (const TPair<FGuid, FReplicationClientQueriedInfo>& StreamQueryPair : Response.ClientInfo)
+		{
+			const FGuid StreamId = StreamQueryPair.Key;
+			
+			// Note that the delegates may have been unsubscribed since request was sent, which is why we need to Find() them
+			if (const FStreamQueryInfo* QueryInfo = StreamQueryInfos.Find(StreamId))
+			{
+				QueryInfo->Delegate.Broadcast(StreamQueryPair.Value.Streams);
+			}
+
+			// Note that the delegates may have been unsubscribed since request was sent, which is why we need to Find() them
+			if (const FAuthorityQueryInfo* QueryInfo = AuthorityQueryInfos.Find(StreamId))
+			{
+				QueryInfo->Delegate.Broadcast(StreamQueryPair.Value.Authority);
+			}
+		}
+
+		// Delegates may have removed themselves up above.
+		CompactDelegates();
+	}
+
+	void FRegularQueryService::CompactDelegates()
+	{
+		for (auto StreamIt = StreamQueryInfos.CreateIterator(); StreamIt; ++StreamIt)
+		{
+			if (StreamIt->Value.Handles.IsEmpty())
+			{
+				StreamIt.RemoveCurrent();
+			}
+		}
+
+		for (auto Authority = AuthorityQueryInfos.CreateIterator(); Authority; ++Authority)
+		{
+			if (Authority->Value.Handles.IsEmpty())
+			{
+				Authority.RemoveCurrent();
+			}
+		}
+	}
+}
diff --git a/Engine/Plugins/Developer/Concert/ConcertApp/MultiUserClient/Source/MultiUserClient/Private/Replication/Util/RegularQueryService.h b/Engine/Plugins/Developer/Concert/ConcertApp/MultiUserClient/Source/MultiUserClient/Private/Replication/Util/RegularQueryService.h
new file mode 100644
--- /dev/null
+++ b/Engine/Plugins/Developer/Concert/ConcertApp/MultiUserClient/Source/MultiUserClient/Private/Replication/Util/RegularQueryService.h
@@ -0,0 +1,82 @@
+// Copyright Epic Games, Inc. All Rights Reserved.
+
+#pragma once
+
+#include "IConcertSessionHandler.h"
+#include "Replication/Messages/ClientQuery.h"
+#include "Containers/Ticker.h"
+#include "Replication/IToken.h"
+#include "Templates/SharedPointer.h"
+#include "Templates/UnrealTemplate.h"
+
+class IConcertSyncClient;
+
+namespace UE::MultiUserClient
+{
+	DECLARE_DELEGATE_OneParam(FStreamQueryDelegate, const TArray<FSharedReplicationStreamDescription>&);
+	DECLARE_DELEGATE_OneParam(FAuthorityQueryDelegate, const TArray<FReplicationAuthorityInfo>&);
+	
+	/** Sends regular FConcertReplication_QueryReplicationInfo_Request to endpoints and publishes the results. */
+	class FRegularQueryService : public FNoncopyable
+	{
+	public:
+
+		FRegularQueryService(
+			TSharedRef<IConcertSyncClient> InOwningClient,
+			float InInterval = 1.f
+			);
+		~FRegularQueryService();
+
+		/** Registers a delegate to invoke for querying an endpoint about its registered streams. */
+		FDelegateHandle RegisterStreamQuery(const FGuid& EndpointId, FStreamQueryDelegate Delegate);
+		/** Registers a delegate to invoke for querying an endpoint about its authority. */
+		FDelegateHandle RegisterAuthorityQuery(const FGuid& EndpointId, FAuthorityQueryDelegate Delegate);
+
+		void UnregisterStreamQuery(const FDelegateHandle& Handle);
+		void UnregisterAuthorityQuery(const FDelegateHandle& Handle);
+		
+	private:
+
+		/** Used to check whether we were destroyed after a response is received. */
+		const TSharedRef<FToken> Token = FToken::Make();
+		
+		/**
+		 * Used to send queries.
+		 * 
+		 * This FRegularQueryService's owner is supposed to make sure this FRegularQueryService is destroyed
+		 * when the client shuts down.
+		 */
+		const TWeakPtr<IConcertSyncClient> OwningClient;
+
+		/** Used to unregister the ticker. */
+		const FTSTicker::FDelegateHandle TickerDelegateHandle;
+
+		// We use a multicast delegate here because it handles unsubscribing while being Broadcast
+		DECLARE_MULTICAST_DELEGATE_OneParam(FMulticastStreamQueryDelegate, const TArray<FSharedReplicationStreamDescription>&);
+		DECLARE_MULTICAST_DELEGATE_OneParam(FMulticastAuthorityQueryDelegate, const TArray<FReplicationAuthorityInfo>&);
+		struct FStreamQueryInfo
+		{
+			TSet<FDelegateHandle> Handles;
+			FMulticastStreamQueryDelegate Delegate;
+		};
+		struct FAuthorityQueryInfo
+		{
+			TSet<FDelegateHandle> Handles;
+			FMulticastAuthorityQueryDelegate Delegate;
+		};
+		TMap<FGuid, FStreamQueryInfo> StreamQueryInfos;
+		TMap<FGuid, FAuthorityQueryInfo> AuthorityQueryInfos;
+
+		bool bIsHandlingQueryResponse = false;
+
+		/** Queries the server for the client's current state. */
+		void SendQueryEvent();
+		void BuildStreamRequest(FConcertReplication_QueryReplicationInfo_Request& Request) const;
+		void BuildAuthorityRequest(FConcertReplication_QueryReplicationInfo_Request& Request) const;
+		
+		/** Handles the server's response. */
+		void HandleQueryResponse(const FConcertReplication_QueryReplicationInfo_Response& Response);
+		void CompactDelegates();
+	};
+}
+
diff --git a/Engine/Plugins/Developer/Concert/ConcertApp/MultiUserClient/Source/MultiUserClient/Private/Widgets/ActiveSession/Replication/Joined/SReplicationClientView.cpp b/Engine/Plugins/Developer/Concert/ConcertApp/MultiUserClient/Source/MultiUserClient/Private/Widgets/ActiveSession/Replication/Joined/SReplicationClientView.cpp
--- a/Engine/Plugins/Developer/Concert/ConcertApp/MultiUserClient/Source/MultiUserClient/Private/Widgets/ActiveSession/Replication/Joined/SReplicationClientView.cpp
+++ b/Engine/Plugins/Developer/Concert/ConcertApp/MultiUserClient/Source/MultiUserClient/Private/Widgets/ActiveSession/Replication/Joined/SReplicationClientView.cpp
@@ -39,6 +39,7 @@ namespace UE::MultiUserClient
 			[
 				SNew(SReplicationClientViewToolbar)
 				.GetChangeTrackerAttribute_Lambda([this](){ return &GetReplicationClientAttribute.Get()->GetDiffer(); })
+				.AdditionalToolbarWidgets() [ InArgs._AdditionalToolbarWidgets.Widget ]
 			]
 
 			// Editor
@@ -50,7 +51,7 @@ namespace UE::MultiUserClient
 		];
 		
 		// Refresh UI if streams change externally, e.g. a remote client changed what they sent
-		ReplicationClient->OnModelExternallyChanged_GameThread().AddSP(this, &SReplicationClientView::OnModelChanged);
+		ReplicationClient->OnModelExternallyChanged().AddSP(this, &SReplicationClientView::OnModelChanged);
 	}
 	
 	void SReplicationClientView::OnModelChanged() const
diff --git a/Engine/Plugins/Developer/Concert/ConcertApp/MultiUserClient/Source/MultiUserClient/Private/Widgets/ActiveSession/Replication/Joined/SReplicationClientView.h b/Engine/Plugins/Developer/Concert/ConcertApp/MultiUserClient/Source/MultiUserClient/Private/Widgets/ActiveSession/Replication/Joined/SReplicationClientView.h
--- a/Engine/Plugins/Developer/Concert/ConcertApp/MultiUserClient/Source/MultiUserClient/Private/Widgets/ActiveSession/Replication/Joined/SReplicationClientView.h
+++ b/Engine/Plugins/Developer/Concert/ConcertApp/MultiUserClient/Source/MultiUserClient/Private/Widgets/ActiveSession/Replication/Joined/SReplicationClientView.h
@@ -23,6 +23,8 @@ namespace UE::MultiUserClient
 		{}
 			/** The client to depict. Should always return true. If the client is destroyed, so should this widget be. */
 			SLATE_ATTRIBUTE(FReplicationClient*, GetReplicationClient)
+			/** Extra widgets to add to the toolbar. */
+			SLATE_NAMED_SLOT(FArguments, AdditionalToolbarWidgets)
 		SLATE_END_ARGS()
 
 		void Construct(const FArguments& InArgs);
diff --git a/Engine/Plugins/Developer/Concert/ConcertApp/MultiUserClient/Source/MultiUserClient/Private/Widgets/ActiveSession/Replication/Joined/SReplicationClientViewToolbar.cpp b/Engine/Plugins/Developer/Concert/ConcertApp/MultiUserClient/Source/MultiUserClient/Private/Widgets/ActiveSession/Replication/Joined/SReplicationClientViewToolbar.cpp
--- a/Engine/Plugins/Developer/Concert/ConcertApp/MultiUserClient/Source/MultiUserClient/Private/Widgets/ActiveSession/Replication/Joined/SReplicationClientViewToolbar.cpp
+++ b/Engine/Plugins/Developer/Concert/ConcertApp/MultiUserClient/Source/MultiUserClient/Private/Widgets/ActiveSession/Replication/Joined/SReplicationClientViewToolbar.cpp
@@ -31,13 +31,20 @@ namespace UE::MultiUserClient
 		ChildSlot
 		[
 			SNew(SHorizontalBox)
-
+			
 			+SHorizontalBox::Slot()
 			.FillWidth(1.f)
 			[
 				SNew(SSpacer)
 			]
 
+			+SHorizontalBox::Slot()
+			.AutoWidth()
+			.Padding(30.f, 0.f)
+			[
+				InArgs._AdditionalToolbarWidgets.Widget
+			]
+
 			// Upload button
 			+SHorizontalBox::Slot()
 			.AutoWidth()
diff --git a/Engine/Plugins/Developer/Concert/ConcertApp/MultiUserClient/Source/MultiUserClient/Private/Widgets/ActiveSession/Replication/Joined/SReplicationClientViewToolbar.h b/Engine/Plugins/Developer/Concert/ConcertApp/MultiUserClient/Source/MultiUserClient/Private/Widgets/ActiveSession/Replication/Joined/SReplicationClientViewToolbar.h
--- a/Engine/Plugins/Developer/Concert/ConcertApp/MultiUserClient/Source/MultiUserClient/Private/Widgets/ActiveSession/Replication/Joined/SReplicationClientViewToolbar.h
+++ b/Engine/Plugins/Developer/Concert/ConcertApp/MultiUserClient/Source/MultiUserClient/Private/Widgets/ActiveSession/Replication/Joined/SReplicationClientViewToolbar.h
@@ -9,7 +9,12 @@ namespace UE::MultiUserClient
 {
 	class FLocalStreamChangeTracker;
 
-	/** Contains a bunch of actions that can be performed while connected to replication. */
+	/**
+	 * Contains a bunch of actions that can be performed on client view.
+	 *
+	 * Layout:
+	 * |                  AdditionalToolbarWidgets | Submit | Revert |
+	 */
 	class SReplicationClientViewToolbar : public SCompoundWidget
 	{
 	public:
@@ -18,6 +23,9 @@ namespace UE::MultiUserClient
 		{}
 			/** Used to submit and revert locally made changes. */
 			SLATE_ATTRIBUTE(FLocalStreamChangeTracker*, GetChangeTrackerAttribute)
+
+			/** Placed towards the left of the buttons, which are on the right side of the toolbar. */
+			SLATE_NAMED_SLOT(FArguments, AdditionalToolbarWidgets)
 		SLATE_END_ARGS()
 
 		void Construct(const FArguments& InArgs);
diff --git a/Engine/Plugins/Developer/Concert/ConcertApp/MultiUserClient/Source/MultiUserClient/Private/Widgets/ActiveSession/Replication/Joined/SReplicationJoinedView.cpp b/Engine/Plugins/Developer/Concert/ConcertApp/MultiUserClient/Source/MultiUserClient/Private/Widgets/ActiveSession/Replication/Joined/SReplicationJoinedView.cpp
--- a/Engine/Plugins/Developer/Concert/ConcertApp/MultiUserClient/Source/MultiUserClient/Private/Widgets/ActiveSession/Replication/Joined/SReplicationJoinedView.cpp
+++ b/Engine/Plugins/Developer/Concert/ConcertApp/MultiUserClient/Source/MultiUserClient/Private/Widgets/ActiveSession/Replication/Joined/SReplicationJoinedView.cpp
@@ -3,33 +3,214 @@
 #include "SReplicationJoinedView.h"
 
 #include "ConcertLogGlobal.h"
+#include "IConcertSyncClient.h"
 #include "SReplicationClientView.h"
 #include "Replication/MultiUserReplicationManager.h"
+#include "SSelectClientViewComboButton.h"
 
 #include "Framework/Notifications/NotificationManager.h"
+#include "Replication/Client/RemoteReplicationClient.h"
 #include "Widgets/Layout/SBox.h"
-#include "Widgets/SBoxPanel.h"
+#include "Widgets/Layout/SWidgetSwitcher.h"
 #include "Widgets/Notifications/SNotificationList.h"
+#include "Widgets/SBoxPanel.h"
+#include "Widgets/Text/STextBlock.h"
 
 #define LOCTEXT_NAMESPACE "SReplicationJoinedWidget"
 
 namespace UE::MultiUserClient
 {
-	void SReplicationJoinedView::Construct(const FArguments& InArgs, TSharedRef<FMultiUserReplicationManager> InReplicationManager)
+	void SReplicationJoinedView::Construct(
+		const FArguments& InArgs,
+		TSharedRef<FMultiUserReplicationManager> InReplicationManager,
+		TSharedRef<IConcertSyncClient> InClient
+		)
 	{
-		ReplicationManager = InReplicationManager;
+		Client = MoveTemp(InClient);
+		ReplicationManager = MoveTemp(InReplicationManager);
 
 		ChildSlot
 		[
-			SNew(SReplicationClientView)
-			.GetReplicationClient_Lambda([this](){ return &ReplicationManager->GetClientManager()->GetLocalClient(); })
+			SAssignNew(ClientViewSwitcher, SWidgetSwitcher)
+			+SWidgetSwitcher::Slot()
+			[
+				SNew(SReplicationClientView)
+				.GetReplicationClient_Lambda([this](){ return &ReplicationManager->GetClientManager()->GetLocalClient(); })
+				.AdditionalToolbarWidgets()
+				[
+					MakeClientSelectionArea()
+				]
+			]
 		];
-
+		
+		RefreshClientViewSwitcher();
+		ReplicationManager->GetClientManager()->OnRemoteClientsChanged().AddSP(this, &SReplicationJoinedView::RefreshClientViewSwitcher);
+		
 		// Show notifications about changing authority
 		ReplicationManager->GetAuthorityPolicy()->OnAuthorityRequestSent_AnyThread().AddSP(this, &SReplicationJoinedView::OnAuthorityRequestSent_AnyThread);
 		ReplicationManager->GetAuthorityPolicy()->OnAuthorityResponseReceived_AnyThread().AddSP(this, &SReplicationJoinedView::OnAuthorityResponseReceived_AnyThread);
 	}
 
+	void SReplicationJoinedView::RefreshClientViewSwitcher()
+	{
+		const int32 ActiveViewIndex = ClientViewSwitcher->GetActiveWidgetIndex();
+		const FGuid DisplayedClientId = ActiveViewIndex == 0
+			? Client->GetConcertClient()->GetCurrentSession()->GetSessionClientEndpointId()
+			: [this, &ActiveViewIndex]() 
+			{
+				for (const TPair<FGuid, int32>& Pair : RemoteClientToWidgetSwitcherIndex)
+				{
+					if (Pair.Value == ActiveViewIndex)
+					{
+						return Pair.Key;
+					}
+				}
+				return FGuid{};
+			}();
+		
+		// Retain the old SWidgets so their transient state, like highlights, are retained when the remote clients change.
+		TArray<TSharedRef<SWidget>> OldClientWidgets;
+		OldClientWidgets.Reserve(ClientViewSwitcher->GetNumWidgets());
+		while (ClientViewSwitcher->GetNumWidgets() > 0)
+		{
+			const TSharedRef<SWidget> Widget = ClientViewSwitcher->GetWidget(0).ToSharedRef();
+			OldClientWidgets.Add(Widget);
+			ClientViewSwitcher->RemoveSlot(Widget);
+		}
+		checkf(OldClientWidgets.Num() >= 1, TEXT("Was supposed to contain local client widget"));
+
+		ClientViewSwitcher->AddSlot()
+			[
+				OldClientWidgets[0]
+			];
+		RebuildClientViewSwitcherChildren(OldClientWidgets);
+
+		// Need to change view if the displayed client was removed
+		const bool bDisplayedClientWasRemoved = RemoteClientToWidgetSwitcherIndex.Contains(DisplayedClientId); 
+		if (bDisplayedClientWasRemoved)
+		{
+			// Will show the local client
+			ClientViewSwitcher->SetActiveWidgetIndex(0);
+		}
+		
+		// ~OldClientWidgets now disposes of the widgets for which no client exists anymore
+	}
+
+	void SReplicationJoinedView::RebuildClientViewSwitcherChildren(const TArray<TSharedRef<SWidget>> OldClientWidgets)
+	{
+		TMap<FGuid, int32> OldRemoteClientToWidgetSwitcherIndex = MoveTemp(RemoteClientToWidgetSwitcherIndex);
+        for (const FRemoteReplicationClient& RemoteClient : ReplicationManager->GetClientManager()->GetRemoteClients())
+        {
+        	const FGuid& EndpointId = RemoteClient.GetRemoteEndpointId();
+        	
+        	const int32* ExistingClientWidgetIndex = OldRemoteClientToWidgetSwitcherIndex.Find(EndpointId);
+        	if (ExistingClientWidgetIndex)
+        	{
+        		ClientViewSwitcher->AddSlot()
+        		[
+        			OldClientWidgets[*ExistingClientWidgetIndex]
+        		];
+        	}
+        	else
+        	{
+        		ClientViewSwitcher->AddSlot()
+        		[
+        			SNew(SReplicationClientView)
+        			.GetReplicationClient_Lambda([this, EndpointId = RemoteClient.GetRemoteEndpointId()]()
+        			{
+        				// It is unsafe to simply capture RemoteClient because the containing TArray may reallocate its location
+        				return ReplicationManager->GetClientManager()->FindRemoteClient(EndpointId);
+        			})
+        			.AdditionalToolbarWidgets()
+        			[
+        				MakeClientSelectionArea()
+        			]
+        		];
+        	}
+
+        	const int32 NewIndex = ClientViewSwitcher->GetNumWidgets() - 1;
+        	RemoteClientToWidgetSwitcherIndex.Add(EndpointId, NewIndex);
+        }
+	}
+
+	TSharedRef<SWidget> SReplicationJoinedView::MakeClientSelectionArea()
+	{
+		return SNew(SHorizontalBox)
+			.ToolTipText(LOCTEXT("ClientView.ToolTip", "Select the client(s) of which you want to see the registered objects & properties."))
+
+			+SHorizontalBox::Slot()
+			.AutoWidth()
+			.VAlign(VAlign_Center)
+			[
+				SNew(STextBlock)
+				.Text(LOCTEXT("ClientView.Label", "Client View"))
+			]
+
+			+SHorizontalBox::Slot()
+			.AutoWidth()
+			.Padding(5.f, 0.f)
+			[
+				SNew(SBox)
+				.MinDesiredWidth(200.f)
+				[
+					MakeClientSelectionComboBox()
+				]
+			];
+	}
+
+	TSharedRef<SWidget> SReplicationJoinedView::MakeClientSelectionComboBox()
+	{
+		return SNew(SSelectClientViewComboButton)
+		
+			// Returning client info
+			.Client(Client->GetConcertClient())
+			.SelectableClients_Lambda([this]()
+			{
+				const TArray<FRemoteReplicationClient>& RemoteClients = ReplicationManager->GetClientManager()->GetRemoteClients();
+				TArray<FGuid> Result;
+				Algo::Transform(RemoteClients, Result, [](const FRemoteReplicationClient& InClient){ return InClient.GetRemoteEndpointId(); });
+
+				const TSharedPtr<IConcertClientSession> CurrentSession = Client->GetConcertClient()->GetCurrentSession();
+				Result.Sort([&CurrentSession](const FGuid& Left, const FGuid& Right)
+				{
+					FConcertSessionClientInfo LeftInfo;
+					FConcertSessionClientInfo RightInfo;
+					CurrentSession->FindSessionClient(Left, LeftInfo);
+					CurrentSession->FindSessionClient(Right, RightInfo);
+					return LeftInfo.ClientInfo.DisplayName <= RightInfo.ClientInfo.DisplayName;
+				});
+				Result.Insert(CurrentSession->GetSessionClientEndpointId(), 0);
+				return Result;
+			})
+			.CurrentSelection_Lambda([this]()
+			{
+				const TOptional<FGuid> RemoteClient = GetRemoteClientBySwitcherIndex(ClientViewSwitcher->GetActiveWidgetIndex());
+				return RemoteClient
+					? *RemoteClient
+					: Client->GetConcertClient()->GetCurrentSession()->GetSessionClientEndpointId();
+			})
+			.OnSelectClient_Lambda([this](const FGuid& InClientId)
+			{
+				const int32* RemoteClientIndex = RemoteClientToWidgetSwitcherIndex.Find(InClientId);
+				// If it's not a remote client, it's the local client
+				const int32 ActiveWidgetIndex = RemoteClientIndex ? *RemoteClientIndex : 0;
+				ClientViewSwitcher->SetActiveWidgetIndex(ActiveWidgetIndex);
+			})
+		;
+	}
+
+	TOptional<FGuid> SReplicationJoinedView::GetRemoteClientBySwitcherIndex(int32 WidgetSwitcherIndex) const
+	{
+		for (const TPair<FGuid, int32>& Pair : RemoteClientToWidgetSwitcherIndex)
+		{
+			if (Pair.Value == WidgetSwitcherIndex)
+			{
+				return { Pair.Key };
+			}
+		}
+		return {};
+	}
+
 	void SReplicationJoinedView::OnAuthorityRequestSent_AnyThread(
 		const ConcertSyncClient::Replication::FAuthorityChangeRequest& Request
 		)
diff --git a/Engine/Plugins/Developer/Concert/ConcertApp/MultiUserClient/Source/MultiUserClient/Private/Widgets/ActiveSession/Replication/Joined/SReplicationJoinedView.h b/Engine/Plugins/Developer/Concert/ConcertApp/MultiUserClient/Source/MultiUserClient/Private/Widgets/ActiveSession/Replication/Joined/SReplicationJoinedView.h
--- a/Engine/Plugins/Developer/Concert/ConcertApp/MultiUserClient/Source/MultiUserClient/Private/Widgets/ActiveSession/Replication/Joined/SReplicationJoinedView.h
+++ b/Engine/Plugins/Developer/Concert/ConcertApp/MultiUserClient/Source/MultiUserClient/Private/Widgets/ActiveSession/Replication/Joined/SReplicationJoinedView.h
@@ -6,13 +6,15 @@
 #include "Widgets/DeclarativeSyntaxSupport.h"
 #include "Widgets/SCompoundWidget.h"
 
+class IConcertSyncClient;
+class SNotificationItem;
+class SWidgetSwitcher;
+
 namespace UE::MultiUserClient
 {
 	class SReplicationClientView;
 }
 
-class SNotificationItem;
-
 namespace UE::ConcertClientSharedSlate
 {
 	class IReplicationEditorView;
@@ -31,22 +33,38 @@ namespace UE::MultiUserClient
 		{}
 		SLATE_END_ARGS()
 
-		void Construct(const FArguments& InArgs, TSharedRef<FMultiUserReplicationManager> InReplicationManager);
+		void Construct(
+			const FArguments& InArgs,
+			TSharedRef<FMultiUserReplicationManager> InReplicationManager,
+			TSharedRef<IConcertSyncClient> InClient
+			);
 
 	private:
 
+		/** The local client this widget is created for. */
+		TSharedPtr<IConcertSyncClient> Client;
 		/** Acts as the model of this view */
 		TSharedPtr<FMultiUserReplicationManager> ReplicationManager;
 		
-		/**
-		 * The local client's view if one is being displayed.
-		 * ChildSlot keeps the reference alive if we're in state EMultiUserReplicationConnectionState::Connected.
-		 * @see ShowWidget_Connected
-		 */
-		TWeakPtr<SReplicationClientView> WeakLocalClientView;
-
+		/** Selects which client is being view. */
+		TSharedPtr<SWidgetSwitcher> ClientViewSwitcher;
+		/** Maps a remote client to index in ClientViewSwitcher. */
+		TMap<FGuid, int32> RemoteClientToWidgetSwitcherIndex;
+		
 		/** Notification about in progress authority change, if any. */
 		TSharedPtr<SNotificationItem> AuthorityChangeNotification;
+		
+		// Building ClientViewSwitcher
+		/** Gets the remote clients and makes sure ClientViewSwitcher has a widget for each. */
+		void RefreshClientViewSwitcher();
+		/** Util for adding back old client widgets after ClientViewSwitcher was cleared. */
+		void RebuildClientViewSwitcherChildren(const TArray<TSharedRef<SWidget>> OldClientWidgets);
+
+		/** Warps the combo box with a text */
+		TSharedRef<SWidget> MakeClientSelectionArea();
+		/** Creates a combobox with which the content of ClientViewSwitcher can be changed. */
+		TSharedRef<SWidget> MakeClientSelectionComboBox();
+		TOptional<FGuid> GetRemoteClientBySwitcherIndex(int32 WidgetSwitcherIndex) const;
 
 		void OnAuthorityRequestSent_AnyThread(const ConcertSyncClient::Replication::FAuthorityChangeRequest& Request);
 		void OnAuthorityResponseReceived_AnyThread(
diff --git a/Engine/Plugins/Developer/Concert/ConcertApp/MultiUserClient/Source/MultiUserClient/Private/Widgets/ActiveSession/Replication/Joined/SSelectClientViewComboButton.cpp b/Engine/Plugins/Developer/Concert/ConcertApp/MultiUserClient/Source/MultiUserClient/Private/Widgets/ActiveSession/Replication/Joined/SSelectClientViewComboButton.cpp
new file mode 100644
--- /dev/null
+++ b/Engine/Plugins/Developer/Concert/ConcertApp/MultiUserClient/Source/MultiUserClient/Private/Widgets/ActiveSession/Replication/Joined/SSelectClientViewComboButton.cpp
@@ -0,0 +1,81 @@
+// Copyright Epic Games, Inc. All Rights Reserved.
+
+#include "SSelectClientViewComboButton.h"
+
+#include "IConcertClient.h"
+#include "Widgets/ClientName/SLocalClientName.h"
+#include "Widgets/ClientName/SRemoteClientName.h"
+
+#include "Framework/Commands/UIAction.h"
+#include "Framework/MultiBox/MultiBoxBuilder.h"
+#include "Widgets/Input/SComboButton.h"
+#include "Widgets/Layout/SWidgetSwitcher.h"
+
+namespace UE::MultiUserClient
+{
+	void SSelectClientViewComboButton::Construct(const FArguments& InArgs)
+	{
+		Client = InArgs._Client;
+		ClientsAttribute = InArgs._SelectableClients;
+		CurrentSelection = InArgs._CurrentSelection;
+		OnSelectClientDelegate = InArgs._OnSelectClient;
+		
+		ChildSlot
+		[
+			SNew(SComboButton)
+			.OnGetMenuContent(this, &SSelectClientViewComboButton::MakeMenuContent)
+			.ButtonContent()
+			[
+				SAssignNew(ButtonContent, SWidgetSwitcher)
+				.WidgetIndex(this, &SSelectClientViewComboButton::GetActiveWidgetIndex)
+
+				+SWidgetSwitcher::Slot()
+				[
+					SNew(ConcertClientSharedSlate::SLocalClientName, Client.ToSharedRef())
+				]
+				+SWidgetSwitcher::Slot()
+				[
+					SNew(ConcertClientSharedSlate::SRemoteClientName, Client.ToSharedRef())
+					.ClientEndpointId(this, &SSelectClientViewComboButton::GetSelectedClientEndpointId)
+				]
+			]
+		];
+	}
+
+	TSharedRef<SWidget> SSelectClientViewComboButton::MakeMenuContent()
+	{
+		FMenuBuilder MenuBuilder(true, nullptr);
+
+		for (const FGuid& ClientId : ClientsAttribute.Get())
+		{
+			const bool bIsLocalClient = ClientId == Client->GetCurrentSession()->GetSessionClientEndpointId();
+			FUIAction UIAction(
+				FExecuteAction::CreateLambda([this, ClientId](){ OnSelectClientDelegate.Execute(ClientId); }),
+				FCanExecuteAction::CreateLambda([this, ClientId](){ return CurrentSelection.Get() != ClientId; })
+				);
+			
+			if (bIsLocalClient)
+			{
+				MenuBuilder.AddMenuEntry(UIAction, SNew(ConcertClientSharedSlate::SLocalClientName, Client.ToSharedRef()));
+			}
+			else
+			{
+				MenuBuilder.AddMenuEntry(UIAction, SNew(ConcertClientSharedSlate::SRemoteClientName, Client.ToSharedRef()).ClientEndpointId(ClientId));
+			}
+		}
+		
+		return MenuBuilder.MakeWidget();
+	}
+
+	int32 SSelectClientViewComboButton::GetActiveWidgetIndex() const
+	{
+		return CurrentSelection.Get() == Client->GetCurrentSession()->GetSessionClientEndpointId()
+			? static_cast<int32>(EButtonContent::LocalClient)
+			: static_cast<uint32>(EButtonContent::RemoteClient);
+	}
+
+	FGuid SSelectClientViewComboButton::GetSelectedClientEndpointId() const
+	{
+		return CurrentSelection.Get();
+	}
+}
diff --git a/Engine/Plugins/Developer/Concert/ConcertApp/MultiUserClient/Source/MultiUserClient/Private/Widgets/ActiveSession/Replication/Joined/SSelectClientViewComboButton.h b/Engine/Plugins/Developer/Concert/ConcertApp/MultiUserClient/Source/MultiUserClient/Private/Widgets/ActiveSession/Replication/Joined/SSelectClientViewComboButton.h
new file mode 100644
--- /dev/null
+++ b/Engine/Plugins/Developer/Concert/ConcertApp/MultiUserClient/Source/MultiUserClient/Private/Widgets/ActiveSession/Replication/Joined/SSelectClientViewComboButton.h
@@ -0,0 +1,67 @@
+// Copyright Epic Games, Inc. All Rights Reserved.
+
+#pragma once
+
+#include "Widgets/DeclarativeSyntaxSupport.h"
+#include "Widgets/SCompoundWidget.h"
+
+class IConcertClient;
+class SWidgetSwitcher;
+
+namespace UE::MultiUserClient
+{
+	class FRemoteReplicationClient;
+	
+	/**
+	 * Wraps a SComboButton that is used to select which SReplicationClientView should be displayed in a SWidgetSwitcher.
+	 */
+	class SSelectClientViewComboButton : public SCompoundWidget
+	{
+	public:
+
+		DECLARE_DELEGATE_OneParam(FSelectClient, const FGuid&);
+		
+		SLATE_BEGIN_ARGS(SSelectClientViewComboButton)
+		{}
+			/** Uses to query local and remote client display info. */
+			SLATE_ARGUMENT(TSharedPtr<IConcertClient>, Client)
+		
+			/** Remote clients that can be selected from in the order that they should be displayed. */
+			SLATE_ATTRIBUTE(TArray<FGuid>, SelectableClients)
+			/** The client that is currently selected */
+			SLATE_ATTRIBUTE(FGuid, CurrentSelection)
+
+			/** Called when a client is selected */
+			SLATE_EVENT(FSelectClient, OnSelectClient)
+		SLATE_END_ARGS()
+
+		void Construct(const FArguments& InArgs);
+
+	private:
+
+		/** Uses to query local and remote client display info. */
+		TSharedPtr<IConcertClient> Client;
+		
+		/** Remote clients that can be selected from */
+		TAttribute<TArray<FGuid>> ClientsAttribute;
+		/** The current client selection */
+		TAttribute<FGuid> CurrentSelection;
+		
+		/** Called when a client is selected */
+		FSelectClient OnSelectClientDelegate;
+
+		enum class EButtonContent 
+		{
+			LocalClient = 0,
+			RemoteClient = 1
+		};
+		/** Content for the button */
+		TSharedPtr<SWidgetSwitcher> ButtonContent;
+		
+		TSharedRef<SWidget> MakeMenuContent();
+		
+		int32 GetActiveWidgetIndex() const;
+		FGuid GetSelectedClientEndpointId() const;
+	};
+}
+
diff --git a/Engine/Plugins/Developer/Concert/ConcertApp/MultiUserClient/Source/MultiUserClient/Private/Widgets/ActiveSession/Replication/SReplicationRootWidget.cpp b/Engine/Plugins/Developer/Concert/ConcertApp/MultiUserClient/Source/MultiUserClient/Private/Widgets/ActiveSession/Replication/SReplicationRootWidget.cpp
--- a/Engine/Plugins/Developer/Concert/ConcertApp/MultiUserClient/Source/MultiUserClient/Private/Widgets/ActiveSession/Replication/SReplicationRootWidget.cpp
+++ b/Engine/Plugins/Developer/Concert/ConcertApp/MultiUserClient/Source/MultiUserClient/Private/Widgets/ActiveSession/Replication/SReplicationRootWidget.cpp
@@ -20,8 +20,13 @@
 
 namespace UE::MultiUserClient
 {
-	void SReplicationRootWidget::Construct(const FArguments& InArgs, TSharedRef<FMultiUserReplicationManager> InReplicationManager)
+	void SReplicationRootWidget::Construct(
+		const FArguments& InArgs,
+		TSharedRef<FMultiUserReplicationManager> InReplicationManager,
+		TSharedRef<IConcertSyncClient> InClient
+		)
 	{
+		Client = MoveTemp(InClient);
 		ReplicationManager = MoveTemp(InReplicationManager);
 		
 		// Show different widget based on whether we've joined replication.
@@ -80,7 +85,7 @@ namespace UE::MultiUserClient
 	{
 		ChildSlot
 		[
-			SNew(SReplicationJoinedView, ReplicationManager.ToSharedRef())
+			SNew(SReplicationJoinedView, ReplicationManager.ToSharedRef(), Client.ToSharedRef())
 		];
 	}
 
diff --git a/Engine/Plugins/Developer/Concert/ConcertApp/MultiUserClient/Source/MultiUserClient/Private/Widgets/ActiveSession/Replication/SReplicationRootWidget.h b/Engine/Plugins/Developer/Concert/ConcertApp/MultiUserClient/Source/MultiUserClient/Private/Widgets/ActiveSession/Replication/SReplicationRootWidget.h
--- a/Engine/Plugins/Developer/Concert/ConcertApp/MultiUserClient/Source/MultiUserClient/Private/Widgets/ActiveSession/Replication/SReplicationRootWidget.h
+++ b/Engine/Plugins/Developer/Concert/ConcertApp/MultiUserClient/Source/MultiUserClient/Private/Widgets/ActiveSession/Replication/SReplicationRootWidget.h
@@ -6,6 +6,8 @@
 #include "Widgets/SCompoundWidget.h"
 #include "Widgets/Layout/SSplitter.h"
 
+class IConcertSyncClient;
+
 namespace UE::MultiUserClient
 {
 	class FMultiUserReplicationManager;
@@ -21,10 +23,12 @@ namespace UE::MultiUserClient
 		{}
 		SLATE_END_ARGS()
 
-		void Construct(const FArguments& InArgs, TSharedRef<FMultiUserReplicationManager> InReplicationManager);
+		void Construct(const FArguments& InArgs, TSharedRef<FMultiUserReplicationManager> InReplicationManager, TSharedRef<IConcertSyncClient> InClient);
 
 	private:
 
+		/** The client this widget was created for. */
+		TSharedPtr<IConcertSyncClient> Client;
 		/** Manages the business logic which we represent. */
 		TSharedPtr<FMultiUserReplicationManager> ReplicationManager;
 		
diff --git a/Engine/Plugins/Developer/Concert/ConcertApp/MultiUserClient/Source/MultiUserClient/Private/Widgets/ActiveSession/Replication/SReplicationTabWithWarningOverlay.cpp b/Engine/Plugins/Developer/Concert/ConcertApp/MultiUserClient/Source/MultiUserClient/Private/Widgets/ActiveSession/Replication/SReplicationTabWithWarningOverlay.cpp
--- a/Engine/Plugins/Developer/Concert/ConcertApp/MultiUserClient/Source/MultiUserClient/Private/Widgets/ActiveSession/Replication/SReplicationTabWithWarningOverlay.cpp
+++ b/Engine/Plugins/Developer/Concert/ConcertApp/MultiUserClient/Source/MultiUserClient/Private/Widgets/ActiveSession/Replication/SReplicationTabWithWarningOverlay.cpp
@@ -14,7 +14,10 @@
 
 namespace UE::MultiUserClient
 {
-	void SReplicationTabWithWarningOverlay::Construct(const FArguments& InArgs, TSharedRef<FMultiUserReplicationManager> InReplicationManager)
+	void SReplicationTabWithWarningOverlay::Construct(
+		const FArguments& InArgs,
+		TSharedRef<FMultiUserReplicationManager> InReplicationManager,
+		TSharedRef<IConcertSyncClient> InClient)
 	{
 		ChildSlot
 		[
@@ -23,7 +26,7 @@ namespace UE::MultiUserClient
 			// The real content
 			+SOverlay::Slot()
 			[
-				SNew(SReplicationRootWidget, InReplicationManager)
+				SNew(SReplicationRootWidget, InReplicationManager, InClient)
 			]
 
 			//  
diff --git a/Engine/Plugins/Developer/Concert/ConcertApp/MultiUserClient/Source/MultiUserClient/Private/Widgets/ActiveSession/Replication/SReplicationTabWithWarningOverlay.h b/Engine/Plugins/Developer/Concert/ConcertApp/MultiUserClient/Source/MultiUserClient/Private/Widgets/ActiveSession/Replication/SReplicationTabWithWarningOverlay.h
--- a/Engine/Plugins/Developer/Concert/ConcertApp/MultiUserClient/Source/MultiUserClient/Private/Widgets/ActiveSession/Replication/SReplicationTabWithWarningOverlay.h
+++ b/Engine/Plugins/Developer/Concert/ConcertApp/MultiUserClient/Source/MultiUserClient/Private/Widgets/ActiveSession/Replication/SReplicationTabWithWarningOverlay.h
@@ -5,6 +5,7 @@
 #include "Widgets/DeclarativeSyntaxSupport.h"
 #include "Widgets/SCompoundWidget.h"
 
+class IConcertSyncClient;
 class SOverlay;
 
 namespace UE::MultiUserClient
@@ -20,7 +21,7 @@ namespace UE::MultiUserClient
 		{}
 		SLATE_END_ARGS()
 
-		void Construct(const FArguments& InArgs, TSharedRef<FMultiUserReplicationManager> InReplicationManager);
+		void Construct(const FArguments& InArgs, TSharedRef<FMultiUserReplicationManager> InReplicationManager, TSharedRef<IConcertSyncClient> InClient);
 		
 	private:
 
diff --git a/Engine/Plugins/Developer/Concert/ConcertApp/MultiUserClient/Source/MultiUserClient/Private/Widgets/ActiveSession/SActiveSessionRoot.cpp b/Engine/Plugins/Developer/Concert/ConcertApp/MultiUserClient/Source/MultiUserClient/Private/Widgets/ActiveSession/SActiveSessionRoot.cpp
--- a/Engine/Plugins/Developer/Concert/ConcertApp/MultiUserClient/Source/MultiUserClient/Private/Widgets/ActiveSession/SActiveSessionRoot.cpp
+++ b/Engine/Plugins/Developer/Concert/ConcertApp/MultiUserClient/Source/MultiUserClient/Private/Widgets/ActiveSession/SActiveSessionRoot.cpp
@@ -11,7 +11,6 @@
 #include "Styling/SlateTypes.h"
 #include "Widgets/Docking/SDockTab.h"
 #include "Widgets/SBoxPanel.h"
-#include "Widgets/SNullWidget.h"
 
 #define LOCTEXT_NAMESPACE "SActiveSessionRoot"
 
@@ -94,14 +93,11 @@ namespace UE::MultiUserClient
 
 	TSharedRef<SDockTab> SActiveSessionRoot::SpawnTab_Overview(const FSpawnTabArgs& Args)
 	{
-		TSharedPtr<IConcertSyncClient> ConcertSyncClientPin = ConcertSyncClient.Pin();
 		return SNew(SDockTab)
 			.Label(LOCTEXT("OverviewTab.Label", "Overview"))
 			.ToolTipText(LOCTEXT("OverviewTab.Tooltip", "Displays active session clients and activity."))
 			[
-				ConcertSyncClientPin
-					? SNew(SActiveSessionOverviewTab, MoveTemp(ConcertSyncClientPin))
-					: SNullWidget::NullWidget
+				SNew(SActiveSessionOverviewTab, ConcertSyncClient)
 			];
 	}
 
@@ -111,7 +107,7 @@ namespace UE::MultiUserClient
 			.Label(LOCTEXT("ReplicationTab.Label", "Replication"))
 			.ToolTipText(LOCTEXT("ReplicationTab.Tooltip", "Manage real-time object replication"))
 			[
-				SNew(SReplicationTabWithWarningOverlay, InReplicationManager)
+				SNew(SReplicationTabWithWarningOverlay, InReplicationManager, ConcertSyncClient.ToSharedRef())
 			];
 	}
 }
diff --git a/Engine/Plugins/Developer/Concert/ConcertApp/MultiUserClient/Source/MultiUserClient/Private/Widgets/ActiveSession/SActiveSessionRoot.h b/Engine/Plugins/Developer/Concert/ConcertApp/MultiUserClient/Source/MultiUserClient/Private/Widgets/ActiveSession/SActiveSessionRoot.h
--- a/Engine/Plugins/Developer/Concert/ConcertApp/MultiUserClient/Source/MultiUserClient/Private/Widgets/ActiveSession/SActiveSessionRoot.h
+++ b/Engine/Plugins/Developer/Concert/ConcertApp/MultiUserClient/Source/MultiUserClient/Private/Widgets/ActiveSession/SActiveSessionRoot.h
@@ -35,7 +35,7 @@ namespace UE::MultiUserClient
 
 	private:
 
-		TWeakPtr<IConcertSyncClient> ConcertSyncClient;
+		TSharedPtr<IConcertSyncClient> ConcertSyncClient;
 		/** Holds the child content: SActiveSessionOverviewTab and SReplicationControlsTab. */
 		TSharedPtr<FTabManager> TabManager;
 		
diff --git a/Engine/Plugins/Developer/Concert/ConcertApp/MultiUserClient/Source/MultiUserReplicationEditor/Internal/Assets/MultiUserReplicationClientPreset.h b/Engine/Plugins/Developer/Concert/ConcertApp/MultiUserClient/Source/MultiUserReplicationEditor/Internal/Assets/MultiUserReplicationClientPreset.h
--- a/Engine/Plugins/Developer/Concert/ConcertApp/MultiUserClient/Source/MultiUserReplicationEditor/Internal/Assets/MultiUserReplicationClientPreset.h
+++ b/Engine/Plugins/Developer/Concert/ConcertApp/MultiUserClient/Source/MultiUserReplicationEditor/Internal/Assets/MultiUserReplicationClientPreset.h
@@ -16,6 +16,9 @@ class MULTIUSERREPLICATIONEDITOR_API UMultiUserReplicationClientPreset : public
 	GENERATED_BODY()
 public:
 
+	/** The stream ID that is used by all Multi-User streams. */
+	static constexpr FGuid MultiUserStreamID { 0xAAAAAAAA, 0xBBBBBBBB, 0xCCCCCCCC, 0xDDDDDDDD };
+
 	/** The stream this client is managing */
 	UPROPERTY(Instanced)
 	TObjectPtr<UMultiUserReplicationStream> Stream;
diff --git a/Engine/Plugins/Developer/Concert/ConcertApp/MultiUserClient/Source/MultiUserReplicationEditor/Internal/Assets/MultiUserReplicationSessionPreset.h b/Engine/Plugins/Developer/Concert/ConcertApp/MultiUserClient/Source/MultiUserReplicationEditor/Internal/Assets/MultiUserReplicationSessionPreset.h
--- a/Engine/Plugins/Developer/Concert/ConcertApp/MultiUserClient/Source/MultiUserReplicationEditor/Internal/Assets/MultiUserReplicationSessionPreset.h
+++ b/Engine/Plugins/Developer/Concert/ConcertApp/MultiUserClient/Source/MultiUserReplicationEditor/Internal/Assets/MultiUserReplicationSessionPreset.h
@@ -16,6 +16,8 @@ class MULTIUSERREPLICATIONEDITOR_API UMultiUserReplicationSessionPreset : public
 public:
 
 	UMultiUserReplicationClientPreset* AddClient();
+	
+	void RemoveClient(UMultiUserReplicationClientPreset& Client);
 	void ClearClients();
 
 	const TArray<TObjectPtr<UMultiUserReplicationClientPreset>>& GetClientPresets() const { return ClientPresets; }
diff --git a/Engine/Plugins/Developer/Concert/ConcertApp/MultiUserClient/Source/MultiUserReplicationEditor/Private/Assets/MultiUserReplicationClientPreset.cpp b/Engine/Plugins/Developer/Concert/ConcertApp/MultiUserClient/Source/MultiUserReplicationEditor/Private/Assets/MultiUserReplicationClientPreset.cpp
--- a/Engine/Plugins/Developer/Concert/ConcertApp/MultiUserClient/Source/MultiUserReplicationEditor/Private/Assets/MultiUserReplicationClientPreset.cpp
+++ b/Engine/Plugins/Developer/Concert/ConcertApp/MultiUserClient/Source/MultiUserReplicationEditor/Private/Assets/MultiUserReplicationClientPreset.cpp
@@ -8,6 +8,7 @@
 UMultiUserReplicationClientPreset::UMultiUserReplicationClientPreset()
 {
 	Stream = CreateDefaultSubobject<UMultiUserReplicationStream>(TEXT("ReplicationList"));
+	Stream->StreamId = MultiUserStreamID;
 	Stream->SetFlags(RF_Transactional);
 }
 
diff --git a/Engine/Plugins/Developer/Concert/ConcertApp/MultiUserClient/Source/MultiUserReplicationEditor/Private/Assets/MultiUserReplicationSessionPreset.cpp b/Engine/Plugins/Developer/Concert/ConcertApp/MultiUserClient/Source/MultiUserReplicationEditor/Private/Assets/MultiUserReplicationSessionPreset.cpp
--- a/Engine/Plugins/Developer/Concert/ConcertApp/MultiUserClient/Source/MultiUserReplicationEditor/Private/Assets/MultiUserReplicationSessionPreset.cpp
+++ b/Engine/Plugins/Developer/Concert/ConcertApp/MultiUserClient/Source/MultiUserReplicationEditor/Private/Assets/MultiUserReplicationSessionPreset.cpp
@@ -14,6 +14,11 @@ UMultiUserReplicationClientPreset* UMultiUserReplicationSessionPreset::AddClient
 	return Result;
 }
 
+void UMultiUserReplicationSessionPreset::RemoveClient(UMultiUserReplicationClientPreset& Client)
+{
+	ClientPresets.RemoveSingle(&Client);
+}
+
 void UMultiUserReplicationSessionPreset::ClearClients()
 {
 	ClientPresets.Empty();
diff --git a/Engine/Plugins/Developer/Concert/ConcertUI/ConcertClientSharedSlate/Source/ConcertClientSharedSlate/Private/Replication/Editor/View/ObjectEditor/FakeObjectToPropertiesEditorModel.cpp b/Engine/Plugins/Developer/Concert/ConcertUI/ConcertClientSharedSlate/Source/ConcertClientSharedSlate/Private/Replication/Editor/View/ObjectEditor/FakeObjectToPropertiesEditorModel.cpp
--- a/Engine/Plugins/Developer/Concert/ConcertUI/ConcertClientSharedSlate/Source/ConcertClientSharedSlate/Private/Replication/Editor/View/ObjectEditor/FakeObjectToPropertiesEditorModel.cpp
+++ b/Engine/Plugins/Developer/Concert/ConcertUI/ConcertClientSharedSlate/Source/ConcertClientSharedSlate/Private/Replication/Editor/View/ObjectEditor/FakeObjectToPropertiesEditorModel.cpp
@@ -10,6 +10,38 @@
 
 namespace UE::ConcertClientSharedSlate
 {
+	namespace Private
+	{
+		static TOptional<FSoftObjectPath> GetTopLevelObjectOf(const FSoftObjectPath& SoftObjectPath)
+		{
+			// Example of an actor called floor
+			// SoftObjectPath = { AssetPath = {PackageName = "/Game/Maps/SyncBoxLevel", AssetName = "SyncBoxLevel"}, SubPathString = "PersistentLevel.Floor" } }
+			const FString& SubPathString = SoftObjectPath.GetSubPathString();
+
+			constexpr int32 PersistentLevelStringLength = 16; // "PersistentLevel." has 16 characters
+			const bool bIsWorldObject = SubPathString.Contains(TEXT("PersistentLevel."), ESearchCase::CaseSensitive);
+			if (!bIsWorldObject)
+			{
+				// Not a path to a world object
+				return {};
+			}
+
+			// Start search after the . behind PersistentLevel
+			const int32 StartSearch = PersistentLevelStringLength + 1;
+			const int32 IndexOfDotAfterActorName = SubPathString.Find(TEXT("."), ESearchCase::CaseSensitive, ESearchDir::FromStart, StartSearch);
+			if (IndexOfDotAfterActorName == INDEX_NONE)
+			{
+				// SoftObjectPath points to an actor
+				return {};
+			}
+
+			const int32 NumToChopOffRight = SubPathString.Len() - IndexOfDotAfterActorName;
+			const FString NewSubstring = SubPathString.LeftChop(NumToChopOffRight);
+			const FSoftObjectPath PathToOwningActor(SoftObjectPath.GetAssetPath(), NewSubstring);
+			return PathToOwningActor;
+		}
+	}
+	
 	FSoftClassPath FFakeObjectToPropertiesEditorModel::GetObjectClass(const FSoftObjectPath& Object) const
 	{
 		const FSoftClassPath ResolvedClass = RealModel->GetObjectClass(Object);
@@ -26,15 +58,75 @@ namespace UE::ConcertClientSharedSlate
 	bool FFakeObjectToPropertiesEditorModel::ForEachReplicatedObject(TFunctionRef<EBreakBehavior(const FSoftObjectPath& Object)> Delegate) const
 	{
 		// This makes SObjectToPropertyViewer only show actors in the outliner
-		return RealModel->ForEachReplicatedObject([this, &Delegate](const FSoftObjectPath& Object)
+		return IterateTopLevelObjects(Delegate);
+	}
+
+	bool FFakeObjectToPropertiesEditorModel::ForEachProperty(const FSoftObjectPath& ObjectPath, TFunctionRef<EBreakBehavior(const FConcertPropertyChain& Property)> Delegate) const
+	{
+		return IterateDisplayedProperties(ObjectPath, Delegate);
+	}
+
+	bool FFakeObjectToPropertiesEditorModel::IterateTopLevelObjects(TFunctionRef<EBreakBehavior(const FSoftObjectPath& Object)> Delegate) const
+	{
+		// Case: RealModel contains only components but not the owning actor.
+		// In that case, we want the UI to still show the owning actor.
+		// We'll track this with these containers:
+		TSet<FSoftObjectPath> NonTopLevelObjects;
+		TSet<FSoftObjectPath> FoundTopLevelObjects;
+
+		// Here we only exposes objects that are top level so we skip components, etc.
+		EBreakBehavior BreakBehavior = EBreakBehavior::Continue;
+		const bool bResult = RealModel->ForEachReplicatedObject([this, &Delegate, &NonTopLevelObjects, &FoundTopLevelObjects, &BreakBehavior](const FSoftObjectPath& Object)
 		{
-			return IsRootObject(Object)
-				? Delegate(Object)
-				: EBreakBehavior::Continue;
+			if (IsTopLevelObject(Object))
+			{
+				FoundTopLevelObjects.Add(Object);
+				BreakBehavior = Delegate(Object); 
+				return BreakBehavior;
+			}
+			
+			NonTopLevelObjects.Add(Object);
+			return EBreakBehavior::Continue;
 		});
+
+		// No additional objects to show
+		if (BreakBehavior == EBreakBehavior::Break)
+		{
+			return bResult;
+		}
+
+		// Now determine the top level objects that are not in RealModel but that need to be shown for subobjects that are in RealModel
+		for (const FSoftObjectPath& NonTopLevelObject : NonTopLevelObjects)
+		{
+			const TOptional<FSoftObjectPath> TopLevelObject = Private::GetTopLevelObjectOf(NonTopLevelObject);
+			if (!TopLevelObject)
+			{
+				continue;
+			}
+			
+			if (!FoundTopLevelObjects.Contains(*TopLevelObject))
+			{
+				if (Delegate(*TopLevelObject) == EBreakBehavior::Break)
+				{
+					return true;
+				}
+
+				// Avoid calling Delegate again for TopLevelObject
+				FoundTopLevelObjects.Add(*TopLevelObject);
+			}
+		}
+
+		return bResult;
 	}
 
-	bool FFakeObjectToPropertiesEditorModel::ForEachProperty(const FSoftObjectPath& ObjectPath, TFunctionRef<EBreakBehavior(const FConcertPropertyChain& Property)> Delegate) const
+	bool FFakeObjectToPropertiesEditorModel::IsTopLevelObject(const FSoftObjectPath& ObjectPath) const
+	{
+		const FSoftClassPath ClassPath = GetObjectClass(ObjectPath);
+		const UClass* LoadedClass = ClassPath.TryLoadClass<UObject>();
+		return LoadedClass && LoadedClass->IsChildOf(AActor::StaticClass());
+	}
+
+	bool FFakeObjectToPropertiesEditorModel::IterateDisplayedProperties(const FSoftObjectPath& ObjectPath, TFunctionRef<EBreakBehavior(const FConcertPropertyChain& Property)> Delegate) const
 	{
 		// GetObjectClass might return null because ObjectPath might not be contained in the model
 		const FSoftClassPath ClassPath = GetObjectClass(ObjectPath);
@@ -51,17 +143,10 @@ namespace UE::ConcertClientSharedSlate
 		
 		// This makes SObjectToPropertyViewer list all available properties
 		PropertySelectionSource->GetPropertySource(LoadedClass)
-			->EnumerateSelectableItems([&Delegate](const FSelectablePropertyInfo& PropertyInfo)
-			{
-				return Delegate(PropertyInfo.Property);
-			});
+		   ->EnumerateSelectableItems([&Delegate](const FSelectablePropertyInfo& PropertyInfo)
+		   {
+			   return Delegate(PropertyInfo.Property);
+		   });
 		return true;
 	}
-
-	bool FFakeObjectToPropertiesEditorModel::IsRootObject(const FSoftObjectPath& ObjectPath) const
-	{
-		const FSoftClassPath ClassPath = GetObjectClass(ObjectPath);
-		UClass* LoadedClass = ClassPath.TryLoadClass<UObject>();
-		return LoadedClass && LoadedClass->IsChildOf(AActor::StaticClass());
-	}
 }
diff --git a/Engine/Plugins/Developer/Concert/ConcertUI/ConcertClientSharedSlate/Source/ConcertClientSharedSlate/Private/Replication/Editor/View/ObjectEditor/FakeObjectToPropertiesEditorModel.h b/Engine/Plugins/Developer/Concert/ConcertUI/ConcertClientSharedSlate/Source/ConcertClientSharedSlate/Private/Replication/Editor/View/ObjectEditor/FakeObjectToPropertiesEditorModel.h
--- a/Engine/Plugins/Developer/Concert/ConcertUI/ConcertClientSharedSlate/Source/ConcertClientSharedSlate/Private/Replication/Editor/View/ObjectEditor/FakeObjectToPropertiesEditorModel.h
+++ b/Engine/Plugins/Developer/Concert/ConcertUI/ConcertClientSharedSlate/Source/ConcertClientSharedSlate/Private/Replication/Editor/View/ObjectEditor/FakeObjectToPropertiesEditorModel.h
@@ -51,7 +51,13 @@ namespace UE::ConcertClientSharedSlate
 		/** Determines the properties that can be selected. */
 		TSharedRef<IPropertySelectionSourceModel> PropertySelectionSource;
 
-		bool IsRootObject(const FSoftObjectPath& ObjectPath) const;
+		/** Returns all objects that must be listed in the top-outliner section of the replication view. */
+		bool IterateTopLevelObjects(TFunctionRef<EBreakBehavior(const FSoftObjectPath& Object)> Delegate) const;
+		/** Whether this object should be displayed in the outliner. */
+		bool IsTopLevelObject(const FSoftObjectPath& ObjectPath) const;
+
+		/** Iterates all properties on the given object's path. */
+		bool IterateDisplayedProperties(const FSoftObjectPath& ObjectPath, TFunctionRef<EBreakBehavior(const FConcertPropertyChain& Property)> Delegate) const;
 	};
 }
 
diff --git a/Engine/Plugins/Developer/Concert/ConcertUI/ConcertClientSharedSlate/Source/ConcertClientSharedSlate/Private/Widgets/ClientName/SClientName.cpp b/Engine/Plugins/Developer/Concert/ConcertUI/ConcertClientSharedSlate/Source/ConcertClientSharedSlate/Private/Widgets/ClientName/SClientName.cpp
new file mode 100644
--- /dev/null
+++ b/Engine/Plugins/Developer/Concert/ConcertUI/ConcertClientSharedSlate/Source/ConcertClientSharedSlate/Private/Widgets/ClientName/SClientName.cpp
@@ -0,0 +1,45 @@
+// Copyright Epic Games, Inc. All Rights Reserved.
+
+#include "Widgets/ClientName/SClientName.h"
+
+#include "ConcertFrontendUtils.h"
+
+#include "Widgets/Text/STextBlock.h"
+
+#define LOCTEXT_NAMESPACE "SClientName"
+
+namespace UE::ConcertClientSharedSlate
+{
+	void SClientName::Construct(const FArguments& InArgs)
+	{
+		ClientInfoAttribute = InArgs._ClientInfo;
+		DisplayAsLocalClientAttribute = InArgs._DisplayAsLocalClient;
+		check(ClientInfoAttribute.IsSet() || ClientInfoAttribute.IsBound());
+		check(DisplayAsLocalClientAttribute.IsSet() || DisplayAsLocalClientAttribute.IsBound());
+		
+		ChildSlot
+		[
+			ConcertFrontendUtils::CreateDisplayName(
+				TAttribute<FText>::CreateSP(this, &SClientName::GetClientDisplayName)
+				)
+		];
+	}
+
+	FText SClientName::GetClientDisplayName() const
+	{
+		const FConcertClientInfo* ClientInfo = ClientInfoAttribute.Get();
+		check(ClientInfo);
+		
+		if (DisplayAsLocalClientAttribute.Get())
+		{
+			return FText::Format(
+				LOCTEXT("ClientDisplayNameFmt", "{0} (me)"),
+				FText::FromString(ClientInfo->DisplayName)
+				);
+		}
+		
+		return FText::FromString(ClientInfo->DisplayName);
+	}
+}
+
+#undef LOCTEXT_NAMESPACE
\ No newline at end of file
diff --git a/Engine/Plugins/Developer/Concert/ConcertUI/ConcertClientSharedSlate/Source/ConcertClientSharedSlate/Private/Widgets/ClientName/SLocalClientName.cpp b/Engine/Plugins/Developer/Concert/ConcertUI/ConcertClientSharedSlate/Source/ConcertClientSharedSlate/Private/Widgets/ClientName/SLocalClientName.cpp
new file mode 100644
--- /dev/null
+++ b/Engine/Plugins/Developer/Concert/ConcertUI/ConcertClientSharedSlate/Source/ConcertClientSharedSlate/Private/Widgets/ClientName/SLocalClientName.cpp
@@ -0,0 +1,25 @@
+// Copyright Epic Games, Inc. All Rights Reserved.
+
+#include "Widgets/ClientName/SLocalClientName.h"
+
+#include "IConcertClient.h"
+
+#include "Widgets/ClientName/SClientName.h"
+#include "Widgets/Text/STextBlock.h"
+
+#define LOCTEXT_NAMESPACE "SLocalClientName"
+
+namespace UE::ConcertClientSharedSlate
+{
+	void SLocalClientName::Construct(const FArguments& InArgs, TSharedRef<IConcertClient> InClient)
+	{
+		ChildSlot
+		[
+			SNew(SClientName)
+			.ClientInfo_Lambda([InClient](){ return &InClient->GetClientInfo(); })
+			.DisplayAsLocalClient(true)
+		];
+	}
+}
+
+#undef LOCTEXT_NAMESPACE
\ No newline at end of file
diff --git a/Engine/Plugins/Developer/Concert/ConcertUI/ConcertClientSharedSlate/Source/ConcertClientSharedSlate/Private/Widgets/ClientName/SRemoteClientName.cpp b/Engine/Plugins/Developer/Concert/ConcertUI/ConcertClientSharedSlate/Source/ConcertClientSharedSlate/Private/Widgets/ClientName/SRemoteClientName.cpp
new file mode 100644
--- /dev/null
+++ b/Engine/Plugins/Developer/Concert/ConcertUI/ConcertClientSharedSlate/Source/ConcertClientSharedSlate/Private/Widgets/ClientName/SRemoteClientName.cpp
@@ -0,0 +1,37 @@
+// Copyright Epic Games, Inc. All Rights Reserved.
+
+#include "Widgets/ClientName/SRemoteClientName.h"
+
+#include "IConcertClient.h"
+#include "Widgets/ClientName/SClientName.h"
+#include "Widgets/Text/STextBlock.h"
+
+#define LOCTEXT_NAMESPACE "SLocalClientName"
+
+namespace UE::ConcertClientSharedSlate
+{
+	void SRemoteClientName::Construct(const FArguments& InArgs, TSharedRef<IConcertClient> InClient)
+	{
+		Client = MoveTemp(InClient);
+		ClientEndpointIdAttribute = InArgs._ClientEndpointId;
+		check(ClientEndpointIdAttribute.IsSet() || ClientEndpointIdAttribute.IsBound());
+		
+		ChildSlot
+		[
+			SNew(SClientName)
+			.ClientInfo(this, &SRemoteClientName::GetClientInfo)
+		];
+	}
+
+	const FConcertClientInfo* SRemoteClientName::GetClientInfo() const
+	{
+		if (const TSharedPtr<IConcertClientSession> ConnectedSession = Client->GetCurrentSession())
+		{
+			ConnectedSession->FindSessionClient(ClientEndpointIdAttribute.Get(), LastKnownClientInfo);
+		}
+
+		return &LastKnownClientInfo.ClientInfo;
+	}
+}
+
+#undef LOCTEXT_NAMESPACE
\ No newline at end of file
diff --git a/Engine/Plugins/Developer/Concert/ConcertUI/ConcertClientSharedSlate/Source/ConcertClientSharedSlate/Public/Widgets/ClientName/SClientName.h b/Engine/Plugins/Developer/Concert/ConcertUI/ConcertClientSharedSlate/Source/ConcertClientSharedSlate/Public/Widgets/ClientName/SClientName.h
new file mode 100644
--- /dev/null
+++ b/Engine/Plugins/Developer/Concert/ConcertUI/ConcertClientSharedSlate/Source/ConcertClientSharedSlate/Public/Widgets/ClientName/SClientName.h
@@ -0,0 +1,39 @@
+// Copyright Epic Games, Inc. All Rights Reserved.
+
+#pragma once
+
+#include "ConcertMessageData.h"
+#include "Widgets/DeclarativeSyntaxSupport.h"
+#include "Widgets/SCompoundWidget.h"
+
+class IConcertClient;
+
+namespace UE::ConcertClientSharedSlate
+{
+	/** Knows how to display FConcertClientInfo. */
+	class CONCERTCLIENTSHAREDSLATE_API SClientName : public SCompoundWidget
+	{
+	public:
+
+		SLATE_BEGIN_ARGS(SClientName)
+			: _DisplayAsLocalClient(false)
+		{}
+			/** The client info to display. */
+			SLATE_ATTRIBUTE(const FConcertClientInfo*, ClientInfo)
+			/** Whether visually indicate that this is a local client (appends "(me)" if true). */
+			SLATE_ATTRIBUTE(bool, DisplayAsLocalClient)
+		SLATE_END_ARGS()
+
+		void Construct(const FArguments& InArgs);
+
+	private:
+
+		/** The client info to display. */
+		TAttribute<const FConcertClientInfo*> ClientInfoAttribute;
+		/** Whether visually indicate that this is a local client (appends "(me)" if true). */
+		TAttribute<bool> DisplayAsLocalClientAttribute;
+
+		/** Gets the display name. */
+		FText GetClientDisplayName() const;
+	};
+}
diff --git a/Engine/Plugins/Developer/Concert/ConcertUI/ConcertClientSharedSlate/Source/ConcertClientSharedSlate/Public/Widgets/ClientName/SLocalClientName.h b/Engine/Plugins/Developer/Concert/ConcertUI/ConcertClientSharedSlate/Source/ConcertClientSharedSlate/Public/Widgets/ClientName/SLocalClientName.h
new file mode 100644
--- /dev/null
+++ b/Engine/Plugins/Developer/Concert/ConcertUI/ConcertClientSharedSlate/Source/ConcertClientSharedSlate/Public/Widgets/ClientName/SLocalClientName.h
@@ -0,0 +1,29 @@
+// Copyright Epic Games, Inc. All Rights Reserved.
+
+#pragma once
+
+#include "ConcertMessageData.h"
+#include "Widgets/DeclarativeSyntaxSupport.h"
+#include "Widgets/SCompoundWidget.h"
+
+class IConcertClient;
+
+namespace UE::ConcertClientSharedSlate
+{
+	/**
+	 * Displays the name of a client.
+	 * 
+	 * The name will look like "Client Name (me)".
+	 * @see SRemoteClientName
+	 */
+	class CONCERTCLIENTSHAREDSLATE_API SLocalClientName : public SCompoundWidget
+	{
+	public:
+
+		SLATE_BEGIN_ARGS(SLocalClientName)
+		{}
+		SLATE_END_ARGS()
+
+		void Construct(const FArguments& InArgs, TSharedRef<IConcertClient> InClient);
+	};
+}
diff --git a/Engine/Plugins/Developer/Concert/ConcertUI/ConcertClientSharedSlate/Source/ConcertClientSharedSlate/Public/Widgets/ClientName/SRemoteClientName.h b/Engine/Plugins/Developer/Concert/ConcertUI/ConcertClientSharedSlate/Source/ConcertClientSharedSlate/Public/Widgets/ClientName/SRemoteClientName.h
new file mode 100644
--- /dev/null
+++ b/Engine/Plugins/Developer/Concert/ConcertUI/ConcertClientSharedSlate/Source/ConcertClientSharedSlate/Public/Widgets/ClientName/SRemoteClientName.h
@@ -0,0 +1,49 @@
+// Copyright Epic Games, Inc. All Rights Reserved.
+
+#pragma once
+
+#include "ConcertMessageData.h"
+#include "Widgets/DeclarativeSyntaxSupport.h"
+#include "Widgets/SCompoundWidget.h"
+
+class IConcertClient;
+
+namespace UE::ConcertClientSharedSlate
+{
+	/**
+	 * Displays the name of a client.
+	 * 
+	 * The name will look like "Client Name".
+	 * @see SLocalClientName
+	 *
+	 * If the client disconnects, the last known info is used.
+	 * If the client info is unknown, the widget will display an empty FConcertClientInfo;
+	 */
+	class CONCERTCLIENTSHAREDSLATE_API SRemoteClientName : public SCompoundWidget
+	{
+	public:
+
+		SLATE_BEGIN_ARGS(SRemoteClientName)
+		{}
+			SLATE_ATTRIBUTE(FGuid, ClientEndpointId)
+		SLATE_END_ARGS()
+
+		void Construct(const FArguments& InArgs, TSharedRef<IConcertClient> InClient);
+
+	private:
+
+		/** The displayed local client. */
+		TSharedPtr<IConcertClient> Client;
+		/** The endpoint ID of the client to display. */
+		TAttribute<FGuid> ClientEndpointIdAttribute;
+
+		/**
+		 * Cached so that the info remains known when the client disconnects.
+		 * Must be mutable because TAttribute::CreateSP requires GetClientInfo to be const.
+		 */
+		mutable FConcertSessionClientInfo LastKnownClientInfo;
+
+		/** Gets the display info. */
+		const FConcertClientInfo* GetClientInfo() const;
+	};
+}
