Extended IConcertClientReplicationManager:
- added way to query current server state of authority
- added events fired when local known server state changes

diff --git a/Engine/Plugins/Developer/Concert/ConcertApp/MultiUserClient/Source/MultiUserClient/Private/Replication/Client/RemoteReplicationClient.cpp b/Engine/Plugins/Developer/Concert/ConcertApp/MultiUserClient/Source/MultiUserClient/Private/Replication/Client/RemoteReplicationClient.cpp
--- a/Engine/Plugins/Developer/Concert/ConcertApp/MultiUserClient/Source/MultiUserClient/Private/Replication/Client/RemoteReplicationClient.cpp
+++ b/Engine/Plugins/Developer/Concert/ConcertApp/MultiUserClient/Source/MultiUserClient/Private/Replication/Client/RemoteReplicationClient.cpp
@@ -16,7 +16,7 @@ namespace UE::MultiUserClient
 		, RemoteEndpointId(InConcertClientId)
 	{
 		// When the remote client's state has changed, refresh the UI.
-		GetStreamSynchronizer().OnServerStateSynched().AddLambda([this]()
+		GetStreamSynchronizer().OnServerStateChanged().AddLambda([this]()
 		{
 			GetClientContent()->Stream->ReplicationMap = GetStreamSynchronizer().GetServerState();
 			OnModelExternallyChanged().Broadcast();
diff --git a/Engine/Plugins/Developer/Concert/ConcertApp/MultiUserClient/Source/MultiUserClient/Private/Replication/Stream/IClientStreamSynchronizer.h b/Engine/Plugins/Developer/Concert/ConcertApp/MultiUserClient/Source/MultiUserClient/Private/Replication/Stream/IClientStreamSynchronizer.h
--- a/Engine/Plugins/Developer/Concert/ConcertApp/MultiUserClient/Source/MultiUserClient/Private/Replication/Stream/IClientStreamSynchronizer.h
+++ b/Engine/Plugins/Developer/Concert/ConcertApp/MultiUserClient/Source/MultiUserClient/Private/Replication/Stream/IClientStreamSynchronizer.h
@@ -75,7 +75,7 @@ namespace UE::MultiUserClient
 		virtual FOnChangesAccepted& OnChangesAccepted() = 0;
 		DECLARE_MULTICAST_DELEGATE(FOnServerStateChanged);
 		/** @return Event executed when the result of GetServerState has been synched. Called after OnChangesAccepted_AnyThread. */
-		virtual FOnServerStateChanged& OnServerStateSynched() = 0;
+		virtual FOnServerStateChanged& OnServerStateChanged() = 0;
 		
 		virtual ~IClientStreamSynchronizer() = default;
 	};
diff --git a/Engine/Plugins/Developer/Concert/ConcertApp/MultiUserClient/Source/MultiUserClient/Private/Replication/Stream/LocalClientStreamSynchronizer.cpp b/Engine/Plugins/Developer/Concert/ConcertApp/MultiUserClient/Source/MultiUserClient/Private/Replication/Stream/LocalClientStreamSynchronizer.cpp
--- a/Engine/Plugins/Developer/Concert/ConcertApp/MultiUserClient/Source/MultiUserClient/Private/Replication/Stream/LocalClientStreamSynchronizer.cpp
+++ b/Engine/Plugins/Developer/Concert/ConcertApp/MultiUserClient/Source/MultiUserClient/Private/Replication/Stream/LocalClientStreamSynchronizer.cpp
@@ -3,10 +3,8 @@
 #include "LocalClientStreamSynchronizer.h"
 
 #include "IConcertSyncClient.h"
-#include "Replication/ChangeStreamSharedUtils.h"
 #include "Replication/IConcertClientReplicationManager.h"
-
-#include "ScopedTransaction.h"
+#include "Replication/Messages/ChangeStream.h"
 
 #define LOCTEXT_NAMESPACE "FLocalClientStreamDiffer"
 
@@ -18,7 +16,24 @@ namespace UE::MultiUserClient
 		)
 		: LocalClient(MoveTemp(InLocalClient))
 		, LocalClientStreamId(InLocalClientStreamId)
-	{}
+	{
+		IConcertClientReplicationManager* ReplicationManager = LocalClient->GetReplicationManager();
+		if (ensure(ReplicationManager))
+		{
+			ReplicationManager->OnPreStreamsChanged().AddRaw(this, &FLocalClientStreamSynchronizer::OnPreStreamsChanged);
+			ReplicationManager->OnPostStreamsChanged().AddRaw(this, &FLocalClientStreamSynchronizer::OnPostStreamsChanged);
+		}
+	}
+
+	FLocalClientStreamSynchronizer::~FLocalClientStreamSynchronizer()
+	{
+		IConcertClientReplicationManager* ReplicationManager = LocalClient->GetReplicationManager();
+		if (ReplicationManager)
+		{
+			ReplicationManager->OnPreStreamsChanged().RemoveAll(this);
+			ReplicationManager->OnPostStreamsChanged().RemoveAll(this);
+		}
+	}
 
 	TFuture<FSubmitChangesResult> FLocalClientStreamSynchronizer::SubmitChanges(const FStreamChangelist& Changelist)
 	{
@@ -30,8 +45,8 @@ namespace UE::MultiUserClient
 			return MakeFulfilledPromise<FSubmitChangesResult>(FSubmitChangesResult{ ESubmitChangesErrorCode::CannotSendRequest }).GetFuture();
 		}
 		
-		ChangeInTransit = Changelist;
 		FChangeStreamRequest Request = BuildChangeRequest(Changelist);
+		ChangeInTransit = Request;
 		
 		return ReplicationManager->ChangeStream(Request)
 			.Next([this, DestructionDetection = LifetimeToken->AsWeak(), Request](FChangeStreamResponse&& Response)
@@ -39,25 +54,39 @@ namespace UE::MultiUserClient
 				const FSubmitChangesResult SubmissionResult { ESubmitChangesErrorCode::Success, { FCompletedChangeSubmission{Request, Response } } };
 				// The request might execute after we're destroyed, e.g. by leaving session while request is on the way.
 				// In that case, the Concert session triggers the OnSessionConnectionChanged which destroys us. Only after that, all the requests are timed out.
-				if (!DestructionDetection.IsValid())
+				if (DestructionDetection.IsValid())
 				{
-					return SubmissionResult;
+					// Already called by OnPreStreamsChanged but needed here, too, in case request fails.
+					ChangeInTransit.Reset();
 				}
 
-				const FStreamChangelist ChangeThatWasInTransit = MoveTemp(ChangeInTransit.GetValue());
-				ChangeInTransit.Reset();
-				
-				if (Response.IsSuccess())
-				{
-					// Tell everybody who cares old and new implicit state. They can e.g. decide to request authority for the new objects.
-					OnChangesAcceptedDelegate.Broadcast(ConfirmedServerState, Request);
-					UpdateConfirmedServerState(Request, ChangeThatWasInTransit);
-				}
-				
 				return SubmissionResult;
 			});
 	}
 
+	const FObjectReplicationMap& FLocalClientStreamSynchronizer::GetServerState() const
+	{
+		IConcertClientReplicationManager* ReplicationManager = LocalClient->GetReplicationManager();
+		if (!ensure(ReplicationManager))
+		{
+			return EmptyState;
+		}
+		
+		const FObjectReplicationMap* Result = nullptr;
+		ReplicationManager->ForEachRegisteredStream([this, &Result](const FReplicationStreamDescription& Stream)
+		{
+			if (Stream.BaseDescription.Identifier == LocalClientStreamId)
+			{
+				Result = &Stream.BaseDescription.ReplicationMap;
+				return EBreakBehavior::Break;
+			}
+			return EBreakBehavior::Continue;
+		});
+		return Result
+			? *Result
+			: EmptyState;
+	}
+
 	ConcertSyncClient::Replication::FChangeStreamRequest FLocalClientStreamSynchronizer::BuildChangeRequest(const FStreamChangelist& Changelist) const
 	{
 		switch (ComputeNextRequestType())
@@ -103,45 +132,30 @@ namespace UE::MultiUserClient
 	
 	FLocalClientStreamSynchronizer::EChangeRequestType FLocalClientStreamSynchronizer::ComputeNextRequestType() const
 	{
-		return ConfirmedServerState.ReplicatedObjects.IsEmpty()
+		return GetServerState().ReplicatedObjects.IsEmpty()
 			? EChangeRequestType::CreateNewStream
 			: EChangeRequestType::UpdateExistingStream;
 	}
 
-	void FLocalClientStreamSynchronizer::UpdateConfirmedServerState(
-		const ConcertSyncClient::Replication::FChangeStreamRequest& Request,
-		const FStreamChangelist& ChangeThatWasInTransit
+	void FLocalClientStreamSynchronizer::OnPreStreamsChanged(
+		const ConcertSyncClient::Replication::FChangeStreamRequest& ChangeStreamRequest,
+		const ConcertSyncClient::Replication::FChangeStreamResponse& ChangeStreamResponse
 		)
 	{
-		const bool bCreatedNewStream = !Request.StreamsToAdd.IsEmpty(); 
-		if (bCreatedNewStream)
-		{
-			ConfirmedServerState = Request.StreamsToAdd[0].BaseDescription.ReplicationMap;
-		}
-		else
+		// Was the change made by us? External code can also call ChangeStream.
+		const bool bRequestEqualToInTransit =  ChangeInTransit.IsSet() && *ChangeInTransit == ChangeStreamRequest;
+		if (bRequestEqualToInTransit)
 		{
-			// The goal here is to leverage ApplyValidatedRequest for updating ConfirmedServerState.
-			TArray<FReplicationStreamDescription> FakeDescriptions;
-			FakeDescriptions.Emplace();
-			FReplicationStreamDescription& FakeStreamDescription = FakeDescriptions[0];
-			FakeStreamDescription.BaseDescription.Identifier = LocalClientStreamId;
-			// Cheaply move ConfirmedServerState into the FReplicationStreamDescription so it satisfies the ApplyValidatedRequest API...
-			FakeStreamDescription.BaseDescription.ReplicationMap = MoveTemp(ConfirmedServerState);
-			ConcertSyncCore::Replication::ChangeStreamUtils::ApplyValidatedRequest(
-				BuildChangeRequest_UpdateExistingStream(ChangeThatWasInTransit),
-				FakeDescriptions
-				);
-						
-			// Careful: ApplyValidatedRequest may have emptied the array!
-			FObjectReplicationMap NewServerState = FakeDescriptions.IsEmpty()
-				? FObjectReplicationMap{}
-			: FakeStreamDescription.BaseDescription.ReplicationMap;
-						
-			// ... and cheaply move the updated map back into our ConfirmedServerState
-			ConfirmedServerState = MoveTemp(NewServerState);
+			// Allows making another change in response to Broadcast
+			ChangeInTransit.Reset();
+			OnChangesAcceptedDelegate.Broadcast(GetServerState(), ChangeStreamRequest);
 		}
-					
-		OnServerStateSynchedDelegate.Broadcast();
+	}
+
+	void FLocalClientStreamSynchronizer::OnPostStreamsChanged()
+	{
+		// Note that this fires even for non MU-related changes, e.g. due to some other system's Concert replication API usage.
+		OnServerStateChangedDelegate.Broadcast();
 	}
 }
 
diff --git a/Engine/Plugins/Developer/Concert/ConcertApp/MultiUserClient/Source/MultiUserClient/Private/Replication/Stream/LocalClientStreamSynchronizer.h b/Engine/Plugins/Developer/Concert/ConcertApp/MultiUserClient/Source/MultiUserClient/Private/Replication/Stream/LocalClientStreamSynchronizer.h
--- a/Engine/Plugins/Developer/Concert/ConcertApp/MultiUserClient/Source/MultiUserClient/Private/Replication/Stream/LocalClientStreamSynchronizer.h
+++ b/Engine/Plugins/Developer/Concert/ConcertApp/MultiUserClient/Source/MultiUserClient/Private/Replication/Stream/LocalClientStreamSynchronizer.h
@@ -23,43 +23,40 @@ namespace UE::MultiUserClient
 	public:
 
 		FLocalClientStreamSynchronizer(TSharedRef<IConcertSyncClient> InLocalClient, const FGuid& InLocalClientStreamId);
+		virtual ~FLocalClientStreamSynchronizer() override;
 
 		//~ Begin IClientStreamSynchronizer Interface
 		virtual TFuture<FSubmitChangesResult> SubmitChanges(const FStreamChangelist& Changelist) override;
 		virtual bool CanMakeSubmitRequest() const override { return !ChangeInTransit.IsSet(); }
 		virtual FGuid GetStreamId() const override { return LocalClientStreamId; }
-		virtual const FObjectReplicationMap& GetServerState() const override { return ConfirmedServerState; }
+		virtual const FObjectReplicationMap& GetServerState() const override;
 		virtual FOnChangesAccepted& OnChangesAccepted() override { return OnChangesAcceptedDelegate; }
-		virtual FOnServerStateChanged& OnServerStateSynched() override { return OnServerStateSynchedDelegate; }
+		virtual FOnServerStateChanged& OnServerStateChanged() override { return OnServerStateChangedDelegate; }
 		//~ End IClientStreamSynchronizer Interface
 
 	private:
 		
 		/** Referenced by the change requests to detect destruction of FLocalClientStreamSynchronizer */
 		const TSharedRef<FToken> LifetimeToken = FToken::Make();
+		/** Returned by GetServerState when there is no registered stream. */
+		const FObjectReplicationMap EmptyState;
 		
 		/** Owning client. Used to send change requests to the server. */
 		const TSharedRef<IConcertSyncClient> LocalClient;
 		/** The ID of the local client's stream this FLocalClientStreamDiffer is managing. */
 		const FGuid LocalClientStreamId;
-
-		/**
-		 * Represents what the local client thinks the replication map on the server currently looks like.
-		 * This is updated every time the server confirms a change.
-		 */
-		FObjectReplicationMap ConfirmedServerState;
 		
 		/**
 		 * Represents a change to ConfirmedServerState that is currently in transit to the server.
 		 * When a change is submitted to the server, ChangeInTransit becomes the difference between ConfirmedServerState and StreamWithInProgressChangesAttribute.
 		 * When the server confirms these changes, they are applied to ConfirmedServerState.
 		 */
-		TOptional<FStreamChangelist> ChangeInTransit;
+		TOptional<ConcertSyncClient::Replication::FChangeStreamRequest> ChangeInTransit;
 		
 		/** Called when a change request that was in transit was accepted by the server. */
 		FOnChangesAccepted OnChangesAcceptedDelegate;
 		/** Event executed when the result of GetServerState has been synched. */
-		FOnServerStateChanged OnServerStateSynchedDelegate;
+		FOnServerStateChanged OnServerStateChangedDelegate;
 		
 		/** Builds a change request based on what's registered on the server. */
 		ConcertSyncClient::Replication::FChangeStreamRequest BuildChangeRequest(const FStreamChangelist& Changelist) const;
@@ -79,11 +76,11 @@ namespace UE::MultiUserClient
 		 */
 		EChangeRequestType ComputeNextRequestType() const;
 
-		/** Updates the server state after the server has answered a request. */
-		void UpdateConfirmedServerState(
-			const ConcertSyncClient::Replication::FChangeStreamRequest& Request,
-			const FStreamChangelist& ChangeThatWasInTransit
+		void OnPreStreamsChanged(
+			const ConcertSyncClient::Replication::FChangeStreamRequest& ChangeStreamRequest,
+			const ConcertSyncClient::Replication::FChangeStreamResponse& ChangeStreamResponse
 			);
+		void OnPostStreamsChanged();
 	};
 }
 
diff --git a/Engine/Plugins/Developer/Concert/ConcertApp/MultiUserClient/Source/MultiUserClient/Private/Replication/Stream/LocalStreamChangeTracker.cpp b/Engine/Plugins/Developer/Concert/ConcertApp/MultiUserClient/Source/MultiUserClient/Private/Replication/Stream/LocalStreamChangeTracker.cpp
--- a/Engine/Plugins/Developer/Concert/ConcertApp/MultiUserClient/Source/MultiUserClient/Private/Replication/Stream/LocalStreamChangeTracker.cpp
+++ b/Engine/Plugins/Developer/Concert/ConcertApp/MultiUserClient/Source/MultiUserClient/Private/Replication/Stream/LocalStreamChangeTracker.cpp
@@ -21,12 +21,12 @@ namespace UE::MultiUserClient
 		, OnModifyReplicationMapDelegate(MoveTemp(InOnModifyReplicationMapDelegate))
 	{
 		check(InStreamWithInProgressChangesAttribute.IsBound() || InStreamWithInProgressChangesAttribute.IsSet());
-		StreamSynchronizer.OnServerStateSynched().AddRaw(this, &FLocalStreamChangeTracker::RefreshChangesCache);
+		StreamSynchronizer.OnServerStateChanged().AddRaw(this, &FLocalStreamChangeTracker::RefreshChangesCache);
 	}
 
 	FLocalStreamChangeTracker::~FLocalStreamChangeTracker()
 	{
-		StreamSynchronizer.OnServerStateSynched().RemoveAll(this);
+		StreamSynchronizer.OnServerStateChanged().RemoveAll(this);
 	}
 
 	void FLocalStreamChangeTracker::RefreshChangesCache()
diff --git a/Engine/Plugins/Developer/Concert/ConcertApp/MultiUserClient/Source/MultiUserClient/Private/Replication/Stream/RemoteClientStreamSynchronizer.cpp b/Engine/Plugins/Developer/Concert/ConcertApp/MultiUserClient/Source/MultiUserClient/Private/Replication/Stream/RemoteClientStreamSynchronizer.cpp
--- a/Engine/Plugins/Developer/Concert/ConcertApp/MultiUserClient/Source/MultiUserClient/Private/Replication/Stream/RemoteClientStreamSynchronizer.cpp
+++ b/Engine/Plugins/Developer/Concert/ConcertApp/MultiUserClient/Source/MultiUserClient/Private/Replication/Stream/RemoteClientStreamSynchronizer.cpp
@@ -56,13 +56,13 @@ namespace UE::MultiUserClient
 			&& LastKnownServerState != MultiUserDescription->ReplicationMap)
 		{
 			LastKnownServerState = MultiUserDescription->ReplicationMap;
-			OnServerStateSynchedDelegate.Broadcast();
+			OnServerStateChangedDelegate.Broadcast();
 		}
 		// If the remote client removed the last object, the stream is implicitly deleted.
 		else if (!MultiUserDescription && !LastKnownServerState.ReplicatedObjects.IsEmpty())
 		{
 			LastKnownServerState = {};
-			OnServerStateSynchedDelegate.Broadcast();
+			OnServerStateChangedDelegate.Broadcast();
 		}
 	}
 }
diff --git a/Engine/Plugins/Developer/Concert/ConcertApp/MultiUserClient/Source/MultiUserClient/Private/Replication/Stream/RemoteClientStreamSynchronizer.h b/Engine/Plugins/Developer/Concert/ConcertApp/MultiUserClient/Source/MultiUserClient/Private/Replication/Stream/RemoteClientStreamSynchronizer.h
--- a/Engine/Plugins/Developer/Concert/ConcertApp/MultiUserClient/Source/MultiUserClient/Private/Replication/Stream/RemoteClientStreamSynchronizer.h
+++ b/Engine/Plugins/Developer/Concert/ConcertApp/MultiUserClient/Source/MultiUserClient/Private/Replication/Stream/RemoteClientStreamSynchronizer.h
@@ -32,7 +32,7 @@ namespace UE::MultiUserClient
 		virtual FGuid GetStreamId() const override;
 		virtual const FObjectReplicationMap& GetServerState() const override { return LastKnownServerState; }
 		virtual FOnChangesAccepted& OnChangesAccepted() override { return OnChangesAcceptedDelegate; }
-		virtual FOnServerStateChanged& OnServerStateSynched() override { return OnServerStateSynchedDelegate; }
+		virtual FOnServerStateChanged& OnServerStateChanged() override { return OnServerStateChangedDelegate; }
 		//~ End IClientStreamSynchronizer Interface
 
 	private:
@@ -48,7 +48,7 @@ namespace UE::MultiUserClient
 		/** Called when a change request that was in transit was accepted by the server. */
 		FOnChangesAccepted OnChangesAcceptedDelegate;
 		/** Event executed when the result of GetServerState has been synched. */
-		FOnServerStateChanged OnServerStateSynchedDelegate;
+		FOnServerStateChanged OnServerStateChangedDelegate;
 
 		/** Called in regular intervals with the contents of the remote client. */
 		void HandleStreamQuery(const TArray<FSharedReplicationStreamDescription>& Streams);
diff --git a/Engine/Plugins/Developer/Concert/ConcertSync/ConcertSyncClient/Source/ConcertSyncClient/Private/Replication/IConcertClientReplicationManager.cpp b/Engine/Plugins/Developer/Concert/ConcertSync/ConcertSyncClient/Source/ConcertSyncClient/Private/Replication/IConcertClientReplicationManager.cpp
--- a/Engine/Plugins/Developer/Concert/ConcertSync/ConcertSyncClient/Source/ConcertSyncClient/Private/Replication/IConcertClientReplicationManager.cpp
+++ b/Engine/Plugins/Developer/Concert/ConcertSync/ConcertSyncClient/Source/ConcertSyncClient/Private/Replication/IConcertClientReplicationManager.cpp
@@ -102,4 +102,15 @@ TFuture<UE::ConcertSyncClient::Replication::FAuthorityChangeResponse> IConcertCl
 	}
 	
 	return RequestAuthorityChange(MoveTemp(Request));
-}
\ No newline at end of file
+}
+
+TMap<FSoftObjectPath, TSet<FGuid>> IConcertClientReplicationManager::GetClientOwnedObjects() const
+{
+	TMap<FSoftObjectPath, TSet<FGuid>> Result;
+	ForEachClientOwnedObject([&Result](const FSoftObjectPath& Object, TSet<FGuid>&& OwningStreams)
+	{
+		Result.Emplace(Object, MoveTemp(OwningStreams));
+		return EBreakBehavior::Continue;
+	});
+	return Result;
+}
diff --git a/Engine/Plugins/Developer/Concert/ConcertSync/ConcertSyncClient/Source/ConcertSyncClient/Private/Replication/Manager/ReplicationManager.cpp b/Engine/Plugins/Developer/Concert/ConcertSync/ConcertSyncClient/Source/ConcertSyncClient/Private/Replication/Manager/ReplicationManager.cpp
--- a/Engine/Plugins/Developer/Concert/ConcertSync/ConcertSyncClient/Source/ConcertSyncClient/Private/Replication/Manager/ReplicationManager.cpp
+++ b/Engine/Plugins/Developer/Concert/ConcertSync/ConcertSyncClient/Source/ConcertSyncClient/Private/Replication/Manager/ReplicationManager.cpp
@@ -84,6 +84,49 @@ namespace UE::ConcertSyncClient::Replication
 			: MakeFulfilledPromise<FChangeStreamResponse>().GetFuture(); 
 	}
 
+	IConcertClientReplicationManager::EAuthorityEnumerationResult FReplicationManager::ForEachClientOwnedObject(
+		TFunctionRef<EBreakBehavior(const FSoftObjectPath& Object, TSet<FGuid>&& OwningStreams)> Callback) const
+	{
+		return ensureMsgf(CurrentState, TEXT("StartAcceptingJoinRequests should have been called at this point."))
+			? CurrentState->ForEachClientOwnedObject(Callback)
+			: EAuthorityEnumerationResult::NoAuthorityAvailable;
+	}
+
+	TSet<FGuid> FReplicationManager::GetClientOwnedStreamsForObject(const FSoftObjectPath& ObjectPath) const
+	{
+		return ensureMsgf(CurrentState, TEXT("StartAcceptingJoinRequests should have been called at this point."))
+			? CurrentState->GetClientOwnedStreamsForObject(ObjectPath)
+			: TSet<FGuid>{}; 
+	}
+
+	IConcertClientReplicationManager::FOnPreStreamsChanged& FReplicationManager::OnPreStreamsChanged()
+	{
+		// Check() to avoid returning some dummy static variable
+		checkf(CurrentState, TEXT("StartAcceptingJoinRequests should have been called at this point."));
+		return CurrentState->OnPreStreamsChanged();
+	}
+
+	IConcertClientReplicationManager::FOnPostStreamsChanged& FReplicationManager::OnPostStreamsChanged()
+	{
+		// Check() to avoid returning some dummy static variable
+		checkf(CurrentState, TEXT("StartAcceptingJoinRequests should have been called at this point."));
+		return CurrentState->OnPostStreamsChanged();
+	}
+
+	IConcertClientReplicationManager::FOnPreAuthorityChanged& FReplicationManager::OnPreAuthorityChanged()
+	{
+		// Check() to avoid returning some dummy static variable
+		checkf(CurrentState, TEXT("StartAcceptingJoinRequests should have been called at this point."));
+		return CurrentState->OnPreAuthorityChanged();
+	}
+
+	IConcertClientReplicationManager::FOnPostAuthorityChanged& FReplicationManager::OnPostAuthorityChanged()
+	{
+		// Check() to avoid returning some dummy static variable
+		checkf(CurrentState, TEXT("StartAcceptingJoinRequests should have been called at this point."));
+		return CurrentState->OnPostAuthorityChanged();
+	}
+
 	void FReplicationManager::OnChangeState(TSharedRef<FReplicationManagerState> NewState)
 	{
 		CurrentState = MoveTemp(NewState);
diff --git a/Engine/Plugins/Developer/Concert/ConcertSync/ConcertSyncClient/Source/ConcertSyncClient/Private/Replication/Manager/ReplicationManager.h b/Engine/Plugins/Developer/Concert/ConcertSync/ConcertSyncClient/Source/ConcertSyncClient/Private/Replication/Manager/ReplicationManager.h
--- a/Engine/Plugins/Developer/Concert/ConcertSync/ConcertSyncClient/Source/ConcertSyncClient/Private/Replication/Manager/ReplicationManager.h
+++ b/Engine/Plugins/Developer/Concert/ConcertSync/ConcertSyncClient/Source/ConcertSyncClient/Private/Replication/Manager/ReplicationManager.h
@@ -33,6 +33,12 @@ namespace UE::ConcertSyncClient::Replication
 		virtual TFuture<FAuthorityChangeResponse> RequestAuthorityChange(FAuthorityChangeRequest Args) override;
 		virtual TFuture<FClientQueryResponse> QueryClientInfo(FClientQueryRequest Args) override;
 		virtual TFuture<FChangeStreamResponse> ChangeStream(FChangeStreamRequest Args) override;
+		virtual EAuthorityEnumerationResult ForEachClientOwnedObject(TFunctionRef<EBreakBehavior(const FSoftObjectPath& Object, TSet<FGuid>&& OwningStreams)>) const override;
+		virtual TSet<FGuid> GetClientOwnedStreamsForObject(const FSoftObjectPath& ObjectPath) const override;
+		virtual FOnPreStreamsChanged& OnPreStreamsChanged() override;
+		virtual FOnPostStreamsChanged& OnPostStreamsChanged() override;
+		virtual FOnPreAuthorityChanged& OnPreAuthorityChanged() override;
+		virtual FOnPostAuthorityChanged& OnPostAuthorityChanged() override;
 		//~ End IConcertClientReplicationManager Interface
 
 	private:
diff --git a/Engine/Plugins/Developer/Concert/ConcertSync/ConcertSyncClient/Source/ConcertSyncClient/Private/Replication/Manager/ReplicationManagerState.h b/Engine/Plugins/Developer/Concert/ConcertSync/ConcertSyncClient/Source/ConcertSyncClient/Private/Replication/Manager/ReplicationManagerState.h
--- a/Engine/Plugins/Developer/Concert/ConcertSync/ConcertSyncClient/Source/ConcertSyncClient/Private/Replication/Manager/ReplicationManagerState.h
+++ b/Engine/Plugins/Developer/Concert/ConcertSync/ConcertSyncClient/Source/ConcertSyncClient/Private/Replication/Manager/ReplicationManagerState.h
@@ -2,6 +2,7 @@
 
 #pragma once
 
+#include "ReplicationManagerUtils.h"
 #include "Replication/IConcertClientReplicationManager.h"
 
 namespace UE::ConcertSyncClient::Replication
@@ -18,7 +19,27 @@ namespace UE::ConcertSyncClient::Replication
 
 		FReplicationManagerState(FReplicationManager& Owner);
 
+		// Default implementations for subclasses in which the operation is not valid
+		
+		//~ Begin IConcertClientReplicationManager Interface
+		virtual EStreamEnumerationResult ForEachRegisteredStream(TFunctionRef<EBreakBehavior(const FReplicationStreamDescription& Stream)> Callback) const override { return EStreamEnumerationResult::NoRegisteredStreams; }
+		virtual TFuture<FAuthorityChangeResponse> RequestAuthorityChange(FAuthorityChangeRequest Args) override { return RejectAll(MoveTemp(Args)); }
+		virtual TFuture<FClientQueryResponse> QueryClientInfo(FClientQueryRequest Args) override { return MakeFulfilledPromise<FClientQueryResponse>().GetFuture(); }
+		virtual TFuture<FChangeStreamResponse> ChangeStream(FChangeStreamRequest Args) override { return MakeFulfilledPromise<FChangeStreamResponse>().GetFuture(); }
+		virtual EAuthorityEnumerationResult ForEachClientOwnedObject(TFunctionRef<EBreakBehavior(const FSoftObjectPath& Object, TSet<FGuid>&& OwningStreams)> Callback) const override { return EAuthorityEnumerationResult::NoAuthorityAvailable; }
+		virtual TSet<FGuid> GetClientOwnedStreamsForObject(const FSoftObjectPath& ObjectPath) const override { return {}; }
+		virtual FOnPreStreamsChanged& OnPreStreamsChanged() override { return OnPreStreamsChangedDelegate; } 
+		virtual FOnPostStreamsChanged& OnPostStreamsChanged() override { return OnPostStreamsChangedDelegate; }
+		virtual FOnPreAuthorityChanged& OnPreAuthorityChanged() override { return OnPreAuthorityChangedDelegate; }
+		virtual FOnPostAuthorityChanged& OnPostAuthorityChanged() override { return OnPostAuthorityChangedDelegate; }
+		//~ End IConcertClientReplicationManager Interface
+
 	protected:
+		
+		FOnPreStreamsChanged OnPreStreamsChangedDelegate;
+		FOnPostStreamsChanged OnPostStreamsChangedDelegate;
+		FOnPreAuthorityChanged OnPreAuthorityChangedDelegate;
+		FOnPostAuthorityChanged OnPostAuthorityChangedDelegate;
 
 		/**
 		 * Subclasses can change the state with this function.
diff --git a/Engine/Plugins/Developer/Concert/ConcertSync/ConcertSyncClient/Source/ConcertSyncClient/Private/Replication/Manager/ReplicationManagerState_Connected.cpp b/Engine/Plugins/Developer/Concert/ConcertSync/ConcertSyncClient/Source/ConcertSyncClient/Private/Replication/Manager/ReplicationManagerState_Connected.cpp
--- a/Engine/Plugins/Developer/Concert/ConcertSync/ConcertSyncClient/Source/ConcertSyncClient/Private/Replication/Manager/ReplicationManagerState_Connected.cpp
+++ b/Engine/Plugins/Developer/Concert/ConcertSync/ConcertSyncClient/Source/ConcertSyncClient/Private/Replication/Manager/ReplicationManagerState_Connected.cpp
@@ -15,6 +15,7 @@
 #include "Replication/Processing/ObjectReplicationSender.h"
 
 #include "Algo/RemoveIf.h"
+#include "Misc/ScopeExit.h"
 #include "UObject/Package.h"
 #include "UObject/UObjectGlobals.h"
 
@@ -124,13 +125,40 @@ namespace UE::ConcertSyncClient::Replication
 				const TSharedPtr<FReplicationManagerState_Connected> ThisPin = WeakThis.Pin();
 				if (ThisPin && Response.IsSuccess())
 				{
-					ThisPin->UpdateReplicatedObjectsAfterStreamChange(Args);
+					ThisPin->UpdateReplicatedObjectsAfterStreamChange(Args, Response);
 				}
 				
 				return FChangeStreamResponse { MoveTemp(Response) };
 			});
 	}
 
+	IConcertClientReplicationManager::EAuthorityEnumerationResult FReplicationManagerState_Connected::ForEachClientOwnedObject(
+		TFunctionRef<EBreakBehavior(const FSoftObjectPath& Object, TSet<FGuid>&& OwningStreams)> Callback
+		) const
+	{
+		TSet<FGuid> Result;
+		int32 ExpectedNumStreams = 0;
+		ReplicationDataSource->ForEachOwnedObject([this, &Callback, &Result, &ExpectedNumStreams](const FSoftObjectPath& ObjectPath)
+		{
+			// Reuse TSet (if possible) for a slightly better memory footprint
+			ExpectedNumStreams = FMath::Max(ExpectedNumStreams, Result.Num());
+			Result.Empty(Result.Num());
+			
+			ReplicationDataSource->AppendOwningStreamsForObject(ObjectPath, Result);
+			return Callback(ObjectPath, MoveTemp(Result));
+		});
+		return EAuthorityEnumerationResult::Iterated;
+	}
+
+	TSet<FGuid> FReplicationManagerState_Connected::GetClientOwnedStreamsForObject(
+		const FSoftObjectPath& ObjectPath
+		) const
+	{
+		TSet<FGuid> Result;
+		ReplicationDataSource->AppendOwningStreamsForObject(ObjectPath, Result);
+		return Result;
+	}
+
 	void FReplicationManagerState_Connected::OnEnterState()
 	{
 		LiveSession->OnTick().AddSP(this, &FReplicationManagerState_Connected::Tick);
@@ -168,8 +196,11 @@ namespace UE::ConcertSyncClient::Replication
 		ReplicationApplier->ProcessObjects(TimeBudget);
 	}
 
-	void FReplicationManagerState_Connected::UpdateReplicatedObjectsAfterStreamChange(const FChangeStreamRequest& Request)
+	void FReplicationManagerState_Connected::UpdateReplicatedObjectsAfterStreamChange(const FChangeStreamRequest& Request, const FConcertReplication_ChangeStream_Response& Response)
 	{
+		OnPreStreamsChangedDelegate.Broadcast(Request, { Response });
+		ON_SCOPE_EXIT{ OnPostStreamsChangedDelegate.Broadcast(); };
+		
 		// Build RegisteredStreams while RegisteredStreams has the old, unupdated state
 		TMap<FSoftObjectPath, TArray<FGuid>> BundledModifiedObjects;
 		for (const TPair<FObjectInStreamID, FConcertReplication_ChangeStream_PutObject>& PutObject : Request.ObjectsToPut)
@@ -216,6 +247,9 @@ namespace UE::ConcertSyncClient::Replication
 
 	void FReplicationManagerState_Connected::UpdateReplicatedObjectsAfterAuthorityChange(FAuthorityChangeRequest&& Request, const FConcertReplication_ChangeAuthority_Response& Response) const
 	{
+		OnPreAuthorityChangedDelegate.Broadcast(Request, { Response });
+		ON_SCOPE_EXIT{ OnPostAuthorityChangedDelegate.Broadcast(); };
+		
 		for (TPair<FSoftObjectPath, FConcertStreamArray>& TakeAuthority : Request.TakeAuthority)
 		{
 			const FSoftObjectPath& ReplicatedObject = TakeAuthority.Key;
diff --git a/Engine/Plugins/Developer/Concert/ConcertSync/ConcertSyncClient/Source/ConcertSyncClient/Private/Replication/Manager/ReplicationManagerState_Connected.h b/Engine/Plugins/Developer/Concert/ConcertSync/ConcertSyncClient/Source/ConcertSyncClient/Private/Replication/Manager/ReplicationManagerState_Connected.h
--- a/Engine/Plugins/Developer/Concert/ConcertSync/ConcertSyncClient/Source/ConcertSyncClient/Private/Replication/Manager/ReplicationManagerState_Connected.h
+++ b/Engine/Plugins/Developer/Concert/ConcertSync/ConcertSyncClient/Source/ConcertSyncClient/Private/Replication/Manager/ReplicationManagerState_Connected.h
@@ -51,6 +51,8 @@ namespace UE::ConcertSyncClient::Replication
 		virtual TFuture<FAuthorityChangeResponse> RequestAuthorityChange(FAuthorityChangeRequest Args) override;
 		virtual TFuture<FClientQueryResponse> QueryClientInfo(FClientQueryRequest Args) override;
 		virtual TFuture<FChangeStreamResponse> ChangeStream(FChangeStreamRequest Args) override;
+		virtual EAuthorityEnumerationResult ForEachClientOwnedObject(TFunctionRef<EBreakBehavior(const FSoftObjectPath& Object, TSet<FGuid>&& OwningStreams)> Callback) const override;
+		virtual TSet<FGuid> GetClientOwnedStreamsForObject(const FSoftObjectPath& ObjectPath) const override;
 		//~ End IConcertClientReplicationManager Interface
 
 	private:
@@ -117,7 +119,7 @@ namespace UE::ConcertSyncClient::Replication
 		void TickReceiver(float TimeBudget);
 		
 		/** Updates replicated objects affected by the change request. */
-		void UpdateReplicatedObjectsAfterStreamChange(const FChangeStreamRequest& Request);
+		void UpdateReplicatedObjectsAfterStreamChange(const FChangeStreamRequest& Request, const FConcertReplication_ChangeStream_Response& Response);
 		void HandleRemovingReplicatedObjects(const FChangeStreamRequest& Request) const;
 
 		/**
diff --git a/Engine/Plugins/Developer/Concert/ConcertSync/ConcertSyncClient/Source/ConcertSyncClient/Private/Replication/Manager/ReplicationManagerState_Disconnected.h b/Engine/Plugins/Developer/Concert/ConcertSync/ConcertSyncClient/Source/ConcertSyncClient/Private/Replication/Manager/ReplicationManagerState_Disconnected.h
--- a/Engine/Plugins/Developer/Concert/ConcertSync/ConcertSyncClient/Source/ConcertSyncClient/Private/Replication/Manager/ReplicationManagerState_Disconnected.h
+++ b/Engine/Plugins/Developer/Concert/ConcertSync/ConcertSyncClient/Source/ConcertSyncClient/Private/Replication/Manager/ReplicationManagerState_Disconnected.h
@@ -26,10 +26,6 @@ namespace UE::ConcertSyncClient::Replication
 		virtual void LeaveReplicationSession() override;
 		virtual bool CanJoin() override { return true; }
 		virtual bool IsConnectedToReplicationSession() override { return false; }
-		virtual EStreamEnumerationResult ForEachRegisteredStream(TFunctionRef<EBreakBehavior(const FReplicationStreamDescription& Stream)> Callback) const override { return EStreamEnumerationResult::NoRegisteredStreams; }
-		virtual TFuture<FAuthorityChangeResponse> RequestAuthorityChange(FAuthorityChangeRequest Args) override { return RejectAll(MoveTemp(Args)); }
-		virtual TFuture<FClientQueryResponse> QueryClientInfo(FClientQueryRequest Args) override { return MakeFulfilledPromise<FClientQueryResponse>().GetFuture(); }
-		virtual TFuture<FChangeStreamResponse> ChangeStream(FChangeStreamRequest Args) override { return MakeFulfilledPromise<FChangeStreamResponse>().GetFuture(); }
 		//~ End IConcertClientReplicationManager Interface
 
 	private:
diff --git a/Engine/Plugins/Developer/Concert/ConcertSync/ConcertSyncClient/Source/ConcertSyncClient/Private/Replication/Manager/ReplicationManagerState_Handshaking.h b/Engine/Plugins/Developer/Concert/ConcertSync/ConcertSyncClient/Source/ConcertSyncClient/Private/Replication/Manager/ReplicationManagerState_Handshaking.h
--- a/Engine/Plugins/Developer/Concert/ConcertSync/ConcertSyncClient/Source/ConcertSyncClient/Private/Replication/Manager/ReplicationManagerState_Handshaking.h
+++ b/Engine/Plugins/Developer/Concert/ConcertSync/ConcertSyncClient/Source/ConcertSyncClient/Private/Replication/Manager/ReplicationManagerState_Handshaking.h
@@ -31,10 +31,6 @@ namespace UE::ConcertSyncClient::Replication
 		virtual void LeaveReplicationSession() override;
 		virtual bool CanJoin() override { return false; }
 		virtual bool IsConnectedToReplicationSession() override { return false; }
-		virtual EStreamEnumerationResult ForEachRegisteredStream(TFunctionRef<EBreakBehavior(const FReplicationStreamDescription& Stream)> Callback) const override { return EStreamEnumerationResult::NoRegisteredStreams; }
-		virtual TFuture<FAuthorityChangeResponse> RequestAuthorityChange(FAuthorityChangeRequest Args) override { return RejectAll(MoveTemp(Args)); }
-		virtual TFuture<FClientQueryResponse> QueryClientInfo(FClientQueryRequest Args) override { return MakeFulfilledPromise<FClientQueryResponse>().GetFuture(); }
-		virtual TFuture<FChangeStreamResponse> ChangeStream(FChangeStreamRequest Args) override { return MakeFulfilledPromise<FChangeStreamResponse>().GetFuture(); }
 		//~ End IConcertClientReplicationManager Interface
 
 	private:
diff --git a/Engine/Plugins/Developer/Concert/ConcertSync/ConcertSyncClient/Source/ConcertSyncClient/Private/Replication/Processing/ClientReplicationDataCollector.cpp b/Engine/Plugins/Developer/Concert/ConcertSync/ConcertSyncClient/Source/ConcertSyncClient/Private/Replication/Processing/ClientReplicationDataCollector.cpp
--- a/Engine/Plugins/Developer/Concert/ConcertSync/ConcertSyncClient/Source/ConcertSyncClient/Private/Replication/Processing/ClientReplicationDataCollector.cpp
+++ b/Engine/Plugins/Developer/Concert/ConcertSync/ConcertSyncClient/Source/ConcertSyncClient/Private/Replication/Processing/ClientReplicationDataCollector.cpp
@@ -2,13 +2,15 @@
 
 #include "ClientReplicationDataCollector.h"
 
-#include "Algo/RemoveIf.h"
 #include "Replication/Data/ReplicationStreamDescription.h"
 #include "Replication/Formats/IObjectReplicationFormat.h"
 #include "Replication/IConcertClientReplicationBridge.h"
 #include "Replication/ReplicationPropertyFilter.h"
 #include "Replication/Data/ObjectIds.h"
 
+#include "Algo/RemoveIf.h"
+#include "Misc/EBreakBehavior.h"
+
 namespace UE::ConcertSyncClient::Replication
 {
 	FClientReplicationDataCollector::FClientReplicationDataCollector(
@@ -124,6 +126,33 @@ namespace UE::ConcertSyncClient::Replication
 		}
 	}
 
+	void FClientReplicationDataCollector::ForEachOwnedObject(
+		TFunctionRef<EBreakBehavior(const FSoftObjectPath&)> Callback
+		) const
+	{
+		for (const TPair<FSoftObjectPath, TArray<FObjectInfo>>& Pair : ObjectsToReplicate)
+		{
+			if (Callback(Pair.Key) == EBreakBehavior::Break)
+			{
+				break;
+			}
+		}
+	}
+
+	void FClientReplicationDataCollector::AppendOwningStreamsForObject(
+		const FSoftObjectPath& ObjectPath,
+		TSet<FGuid>& Paths
+		) const
+	{
+		if (const TArray<FObjectInfo>* ObjectInfo = ObjectsToReplicate.Find(ObjectPath))
+		{
+			Algo::Transform(*ObjectInfo, Paths, [](const FObjectInfo& ObjectInfo)
+			{
+				return ObjectInfo.StreamId;
+			});
+		}
+	}
+
 	void FClientReplicationDataCollector::ForEachPendingObject(TFunctionRef<void(const FObjectInStreamID&)> ProcessItemFunc) const
 	{
 		for (const TPair<FSoftObjectPath, TArray<FObjectInfo>>& Pair : ObjectsToReplicate)
diff --git a/Engine/Plugins/Developer/Concert/ConcertSync/ConcertSyncClient/Source/ConcertSyncClient/Private/Replication/Processing/ClientReplicationDataCollector.h b/Engine/Plugins/Developer/Concert/ConcertSync/ConcertSyncClient/Source/ConcertSyncClient/Private/Replication/Processing/ClientReplicationDataCollector.h
--- a/Engine/Plugins/Developer/Concert/ConcertSync/ConcertSyncClient/Source/ConcertSyncClient/Private/Replication/Processing/ClientReplicationDataCollector.h
+++ b/Engine/Plugins/Developer/Concert/ConcertSync/ConcertSyncClient/Source/ConcertSyncClient/Private/Replication/Processing/ClientReplicationDataCollector.h
@@ -14,6 +14,7 @@
 class FConcertSyncClientLiveSession;
 class IConcertClientReplicationBridge;
 class UObject;
+enum class EBreakBehavior : uint8;
 struct FConcertPropertySelection;
 struct FReplicationStreamDescription;
 
@@ -60,6 +61,11 @@ namespace UE::ConcertSyncClient::Replication
 		 */
 		void OnObjectStreamModified(const FSoftObjectPath& Object, TArrayView<const FGuid> PutStreams);
 
+		/** Iterates every object for which there is at least one owning stream. */
+		void ForEachOwnedObject(TFunctionRef<EBreakBehavior(const FSoftObjectPath&)> Callback) const;
+		/** Writes all owning streams for ObjectPath into Paths. */
+		void AppendOwningStreamsForObject(const FSoftObjectPath& ObjectPath, TSet<FGuid>& Paths) const;
+
 		//~ Begin IReplicationDataSource Interface
 		virtual void ForEachPendingObject(TFunctionRef<void(const FObjectInStreamID&)> ProcessItemFunc) const override;
 		virtual int32 NumObjects() const override { return NumTrackedObjects; }
diff --git a/Engine/Plugins/Developer/Concert/ConcertSync/ConcertSyncClient/Source/ConcertSyncClient/Public/Replication/IConcertClientReplicationManager.h b/Engine/Plugins/Developer/Concert/ConcertSync/ConcertSyncClient/Source/ConcertSyncClient/Public/Replication/IConcertClientReplicationManager.h
--- a/Engine/Plugins/Developer/Concert/ConcertSync/ConcertSyncClient/Source/ConcertSyncClient/Public/Replication/IConcertClientReplicationManager.h
+++ b/Engine/Plugins/Developer/Concert/ConcertSync/ConcertSyncClient/Source/ConcertSyncClient/Public/Replication/IConcertClientReplicationManager.h
@@ -60,8 +60,6 @@ public:
 	/**
 	 * Joins a replication session.
 	 * Subsequent calls to JoinReplicationSession will fail until either the resulting TFuture returns or LeaveReplicationSession is called.
-	 *
-	 * @note The future may execute on any thread. Take care to synchronize correctly with the game thread if needed.
 	 */
 	virtual TFuture<UE::ConcertSyncClient::Replication::FJoinReplicatedSessionResult> JoinReplicationSession(UE::ConcertSyncClient::Replication::FJoinReplicatedSessionArgs Args) = 0;
 	/** Leaves the current replication session. */
@@ -82,14 +80,13 @@ public:
 	 * @return Whether this manager is connected to a session (Iterated) or not (NoRegisteredStreams).
 	 */
 	virtual EStreamEnumerationResult ForEachRegisteredStream(TFunctionRef<EBreakBehavior(const FReplicationStreamDescription& Stream)> Callback) const = 0;
-	/** @return Whether this manager is currently in a replication session (basically whether ForEachRegisteredStream returns EStreamEnumerationResult::Iterated). */
+	/** @return Whether this manager is has any registered streams (basically whether ForEachRegisteredStream returns EStreamEnumerationResult::Iterated). */
 	bool HasRegisteredStreams() const;
 	/** @return The streams registered with the server. */
 	TArray<FReplicationStreamDescription> GetRegisteredStreams() const;
 	
 	/**
 	 * Requests from the server to change the authority over some objects.
-	 * @note The future may execute on any thread. Take care of synchronize correctly with the game thread if needed.
 	 */
 	virtual TFuture<UE::ConcertSyncClient::Replication::FAuthorityChangeResponse> RequestAuthorityChange(UE::ConcertSyncClient::Replication::FAuthorityChangeRequest Args) = 0;
 	/** Util function that will request authority for all streams for the given objects. */
@@ -97,17 +94,44 @@ public:
 	/** Util function that will let go over all authority of the given objects. */
 	TFuture<UE::ConcertSyncClient::Replication::FAuthorityChangeResponse> ReleaseAuthorityOf(TArrayView<const FSoftObjectPath> Objects);
 
+	enum class EAuthorityEnumerationResult { NoAuthorityAvailable, Iterated };
 	/**
-	 * Requests replication info about other clients, including the streams registered and which objects they have authority over (i.e. are sending).
-	 * @note The future may execute on any thread. Take care to synchronize correctly with the game thread if needed.
+	 * Iterates through all objects that this client has authority over.
+	 * @return Whether this manager is connected to a session (Iterated) or not (NoAuthorityAvailable)
 	 */
+	virtual EAuthorityEnumerationResult ForEachClientOwnedObject(TFunctionRef<EBreakBehavior(const FSoftObjectPath& Object, TSet<FGuid>&& OwningStreams)> Callback) const = 0;
+	/** @return All the streams this client has registered that have authority for the given ObjectPath. */
+	virtual TSet<FGuid> GetClientOwnedStreamsForObject(const FSoftObjectPath& ObjectPath) const = 0;
+	/** @return All owned objects and associated owning streams. */
+	TMap<FSoftObjectPath, TSet<FGuid>> GetClientOwnedObjects() const;
+
+	/** Requests replication info about other clients, including the streams registered and which objects they have authority over (i.e. are sending). */
 	virtual TFuture<UE::ConcertSyncClient::Replication::FClientQueryResponse> QueryClientInfo(UE::ConcertSyncClient::Replication::FClientQueryRequest Args) = 0;
 
-	/**
-	 * Requests to change the client's registered streams
-	 * @note The future may execute on any thread. Take care to synchronize correctly with the game thread if needed.
-	 */
+	/** Requests to change the client's registered streams */
 	virtual TFuture<UE::ConcertSyncClient::Replication::FChangeStreamResponse> ChangeStream(UE::ConcertSyncClient::Replication::FChangeStreamRequest Args) = 0;
+
+	DECLARE_MULTICAST_DELEGATE_TwoParams(FOnPreStreamsChanged,
+		const UE::ConcertSyncClient::Replication::FChangeStreamRequest&,
+		const UE::ConcertSyncClient::Replication::FChangeStreamResponse&
+		);
+	/** Called right before the result of GetRegisteredStreams changes. */
+	virtual FOnPreStreamsChanged& OnPreStreamsChanged() = 0;
+	
+	DECLARE_MULTICAST_DELEGATE(FOnPostStreamsChanged);
+	/** Called right after the result of GetRegisteredStreams has changed. */
+	virtual FOnPostStreamsChanged& OnPostStreamsChanged() = 0;
+
+	DECLARE_MULTICAST_DELEGATE_TwoParams(FOnPreAuthorityChanged,
+		const UE::ConcertSyncClient::Replication::FAuthorityChangeRequest&,
+		const UE::ConcertSyncClient::Replication::FAuthorityChangeResponse&
+		);
+	/** Called right before GetClientOwnedObjects changes. */
+	virtual FOnPreAuthorityChanged& OnPreAuthorityChanged() = 0;
+	
+	DECLARE_MULTICAST_DELEGATE(FOnPostAuthorityChanged);
+	/** Called right after GetClientOwnedObjects has changed. */
+	virtual FOnPostAuthorityChanged& OnPostAuthorityChanged() = 0;
 	
 	virtual ~IConcertClientReplicationManager() = default;
 };
diff --git a/Engine/Plugins/Developer/Concert/ConcertSync/ConcertSyncCore/Source/ConcertSyncCore/Public/Replication/Data/ReplicationStreamDescription.h b/Engine/Plugins/Developer/Concert/ConcertSync/ConcertSyncCore/Source/ConcertSyncCore/Public/Replication/Data/ReplicationStreamDescription.h
--- a/Engine/Plugins/Developer/Concert/ConcertSync/ConcertSyncCore/Source/ConcertSyncCore/Public/Replication/Data/ReplicationStreamDescription.h
+++ b/Engine/Plugins/Developer/Concert/ConcertSync/ConcertSyncCore/Source/ConcertSyncCore/Public/Replication/Data/ReplicationStreamDescription.h
@@ -21,8 +21,6 @@ struct FSharedReplicationStreamDescription
 	/** Identifies the data that this stream will send. */
 	UPROPERTY()
 	FObjectReplicationMap ReplicationMap;
-
-	// TODO: Add override settings to objects, such as override update frequency etc. Do it via TMap<FSoftObjectPath, FObjectReplicationSettings>.
 	
 	friend bool operator==(const FSharedReplicationStreamDescription& Left, const FSharedReplicationStreamDescription& Right)
 	{
@@ -46,7 +44,14 @@ struct FReplicationStreamDescription
 	UPROPERTY()
 	FSharedReplicationStreamDescription BaseDescription;
 
-	// TODO: Add TArray<TObjectPtr<UReplicationStreamAttribute>> property
+	friend bool operator==(const FReplicationStreamDescription& Left, const FReplicationStreamDescription& Right)
+	{
+		return Left.BaseDescription == Right.BaseDescription;
+	}
+	friend bool operator!=(const FReplicationStreamDescription& Left, const FReplicationStreamDescription& Right)
+	{
+		return !(Left == Right);
+	}
 
 	/** Packs the UObjects in this description (Concert does not allow sending UObjects directly) */
 	CONCERTSYNCCORE_API struct FReplicationStreamDescription_NetPacked Pack() const;
@@ -60,12 +65,18 @@ USTRUCT()
 struct FReplicationStreamDescription_NetPacked
 {
 	GENERATED_BODY()
-
 	
 	UPROPERTY()
 	FSharedReplicationStreamDescription BaseDescription;
 
-	// TODO: Add TArray<FConcertSessionSerializedPayload> where each corresponds to a UReplicationStreamAttribute; those will be used by receive rules to decide whether a client can receive certain data.
+	friend bool operator==(const FReplicationStreamDescription_NetPacked& Left, const FReplicationStreamDescription_NetPacked& Right)
+	{
+		return Left.BaseDescription == Right.BaseDescription;
+	}
+	friend bool operator!=(const FReplicationStreamDescription_NetPacked& Left, const FReplicationStreamDescription_NetPacked& Right)
+	{
+		return !(Left == Right);
+	}
 
 	/** Unpacks this data. This function can fail, e.g if the UObject class does not exist locally. */
 	CONCERTSYNCCORE_API TOptional<FReplicationStreamDescription> Unpack(FString* OutErrorMessage = nullptr) const;
diff --git a/Engine/Plugins/Developer/Concert/ConcertSync/ConcertSyncCore/Source/ConcertSyncCore/Public/Replication/Messages/ChangeStream.h b/Engine/Plugins/Developer/Concert/ConcertSync/ConcertSyncCore/Source/ConcertSyncCore/Public/Replication/Messages/ChangeStream.h
--- a/Engine/Plugins/Developer/Concert/ConcertSync/ConcertSyncCore/Source/ConcertSyncCore/Public/Replication/Messages/ChangeStream.h
+++ b/Engine/Plugins/Developer/Concert/ConcertSync/ConcertSyncCore/Source/ConcertSyncCore/Public/Replication/Messages/ChangeStream.h
@@ -37,6 +37,16 @@ struct FConcertReplication_ChangeStream_PutObject
 	UPROPERTY()
 	FSoftClassPath ClassPath;
 
+	friend bool operator==(const FConcertReplication_ChangeStream_PutObject& Left, const FConcertReplication_ChangeStream_PutObject& Right)
+	{
+		return Left.Properties == Right.Properties && Left.ClassPath == Right.ClassPath;
+	}
+
+	friend bool operator!=(const FConcertReplication_ChangeStream_PutObject& Left, const FConcertReplication_ChangeStream_PutObject& Right)
+	{
+		return !(Left == Right);
+	}
+
 	// Intention: Ideally code dealing with PutObject requests uses these constructors / factory functions.
 	// If a property is added to FReplicatedObjectInfo, only the below code needs to be updated.
 
@@ -106,6 +116,19 @@ struct FConcertReplication_ChangeStream_Request
 	 */
 	UPROPERTY()
 	TSet<FGuid> StreamsToRemove;
+
+	friend bool operator==(const FConcertReplication_ChangeStream_Request& Left, const FConcertReplication_ChangeStream_Request& Right)
+	{
+		const auto OrderIndependentEquals = [](const auto& Left, const auto& Right){ return Left.Num() == Right.Num() && Left.Includes(Right); };
+		return OrderIndependentEquals(Left.ObjectsToRemove, Right.ObjectsToRemove)
+			&& OrderIndependentEquals(Left.StreamsToRemove, Right.StreamsToRemove)
+			&& Left.ObjectsToPut.OrderIndependentCompareEqual(Right.ObjectsToPut)
+			&& Left.StreamsToAdd == Right.StreamsToAdd;
+	}
+	friend bool operator!=(const FConcertReplication_ChangeStream_Request& Lhs, const FConcertReplication_ChangeStream_Request& RHS)
+	{
+		return !(Lhs == RHS);
+	}
 };
 
 UENUM()
