Replication view MVP:
- Replaced SSubobjectEditor with custom implementation
- Allow injection of columns to the subobject editor
- Introduced ISubmissionWorkflow, which handles changing streams and authority
- Added checkbox for toggling authority to MU replication stream editor

diff --git a/Engine/Plugins/Developer/Concert/ConcertApp/MultiUserClient/Source/MultiUserClient/Private/Replication/Authority/AuthorityChangeTracker.cpp b/Engine/Plugins/Developer/Concert/ConcertApp/MultiUserClient/Source/MultiUserClient/Private/Replication/Authority/AuthorityChangeTracker.cpp
new file mode 100644
--- /dev/null
+++ b/Engine/Plugins/Developer/Concert/ConcertApp/MultiUserClient/Source/MultiUserClient/Private/Replication/Authority/AuthorityChangeTracker.cpp
@@ -0,0 +1,83 @@
+// Copyright Epic Games, Inc. All Rights Reserved.
+
+#pragma once
+
+#include "AuthorityChangeTracker.h"
+
+#include "IClientAuthoritySynchronizer.h"
+
+namespace UE::MultiUserClient
+{
+	FAuthorityChangeTracker::FAuthorityChangeTracker(
+		IClientAuthoritySynchronizer& InAuthoritySynchronizer
+		)
+		: AuthoritySynchronizer(InAuthoritySynchronizer)
+	{
+		AuthoritySynchronizer.OnServerStateChanged().AddRaw(this, &FAuthorityChangeTracker::RefreshChanges);
+	}
+
+	FAuthorityChangeTracker::~FAuthorityChangeTracker()
+	{
+		// Strictly not needed because AuthoritySynchronizer will die at the same time as us but let's be nice
+		AuthoritySynchronizer.OnServerStateChanged().RemoveAll(this);
+	}
+
+	void FAuthorityChangeTracker::SetAuthorityIfAllowed(const FSoftObjectPath& ObjectPath, bool bNewAuthorityState)
+	{
+		if (CanSetAuthorityFor(ObjectPath))
+		{
+			NewAuthorityStates.Add(ObjectPath, bNewAuthorityState);
+		}
+	}
+
+	void FAuthorityChangeTracker::RefreshChanges()
+	{
+		for (auto It = NewAuthorityStates.CreateIterator(); It; ++It)
+		{
+			const TPair<FSoftObjectPath, bool>& NewAuthorityState = *It;
+			const FSoftObjectPath& ObjectPath = NewAuthorityState.Key;
+
+			const bool bHasAuthority = AuthoritySynchronizer.HasAuthorityOver(ObjectPath); 
+			if (bHasAuthority == NewAuthorityState.Value
+				|| !AuthoritySynchronizer.CanChangeAuthority(ObjectPath))
+			{
+				It.RemoveCurrent();
+			}
+		}
+	}
+
+	bool FAuthorityChangeTracker::GetAuthorityStateAfterApplied(const FSoftObjectPath& ObjectPath) const
+	{
+		const bool* ChangeValue = NewAuthorityStates.Find(ObjectPath);
+		return ChangeValue ? *ChangeValue : AuthoritySynchronizer.HasAuthorityOver(ObjectPath);
+	}
+
+	bool FAuthorityChangeTracker::CanSetAuthorityFor(const FSoftObjectPath& ObjectPath) const
+	{
+		return AuthoritySynchronizer.CanChangeAuthority(ObjectPath);
+	}
+
+	EAuthorityMutability FAuthorityChangeTracker::GetChangeAuthorityMutability(const FSoftObjectPath& ObjectPath) const
+	{
+		return AuthoritySynchronizer.GetChangeAuthorityMutability(ObjectPath);
+	}
+
+	ConcertSyncClient::Replication::FAuthorityChangeRequest FAuthorityChangeTracker::BuildChangeRequest(const FGuid& StreamId) const
+	{
+		ConcertSyncClient::Replication::FAuthorityChangeRequest ChangeRequest;
+
+		for (const TPair<FSoftObjectPath, bool>& NewAuthorityState : NewAuthorityStates)
+		{
+			if (NewAuthorityState.Value)
+			{
+				ChangeRequest.TakeAuthority.Add(NewAuthorityState.Key).StreamIds = { StreamId };
+			}
+			else
+			{
+				ChangeRequest.ReleaseAuthority.Add(NewAuthorityState.Key).StreamIds = { StreamId };
+			}
+		}
+		
+		return ChangeRequest;
+	}
+}
diff --git a/Engine/Plugins/Developer/Concert/ConcertApp/MultiUserClient/Source/MultiUserClient/Private/Replication/Authority/AuthorityChangeTracker.h b/Engine/Plugins/Developer/Concert/ConcertApp/MultiUserClient/Source/MultiUserClient/Private/Replication/Authority/AuthorityChangeTracker.h
new file mode 100644
--- /dev/null
+++ b/Engine/Plugins/Developer/Concert/ConcertApp/MultiUserClient/Source/MultiUserClient/Private/Replication/Authority/AuthorityChangeTracker.h
@@ -0,0 +1,53 @@
+// Copyright Epic Games, Inc. All Rights Reserved.
+
+#pragma once
+
+#include "Containers/Map.h"
+#include "Replication/IConcertClientReplicationManager.h"
+
+struct FSoftObjectPath;
+
+namespace UE::MultiUserClient
+{
+	class IClientAuthoritySynchronizer;
+	enum class EAuthorityMutability;
+	
+	/** Keeps track of authority changes the local editor instance makes to a client. */
+	class FAuthorityChangeTracker
+	{
+	public:
+
+		FAuthorityChangeTracker(IClientAuthoritySynchronizer& InAuthoritySynchronizer);
+		~FAuthorityChangeTracker();
+		
+		/** Marks that the authority should be changed to bNewAuthorityState. */
+		void SetAuthorityIfAllowed(const FSoftObjectPath& ObjectPath, bool bNewAuthorityState);
+
+		/** Diffs NewAuthorityStates to the current authority states and removes entries. */
+		void RefreshChanges();
+		
+		void ClearChanges() { NewAuthorityStates.Reset(); }
+		bool HasChanges() const { return !NewAuthorityStates.IsEmpty(); }
+		
+		/** Gets the authority state the object will have if the changes are applied. */
+		bool GetAuthorityStateAfterApplied(const FSoftObjectPath& ObjectPath) const;
+		
+		/** @return Whether it is valid to set authority for the given ObjectPath: that is if the object has properties registered. */
+		bool CanSetAuthorityFor(const FSoftObjectPath& ObjectPath) const;
+		/** @return Detailed information why the object's authority can or cannot be changed.  */
+		EAuthorityMutability GetChangeAuthorityMutability(const FSoftObjectPath& ObjectPath) const;
+
+		/** Builds a change request from the local changes. */
+		ConcertSyncClient::Replication::FAuthorityChangeRequest BuildChangeRequest(const FGuid& StreamId) const;
+		
+	private:
+
+		/** Knows the current authority state of the client. */
+		IClientAuthoritySynchronizer& AuthoritySynchronizer;
+
+		/** Object to the authority state it should have. */
+		TMap<FSoftObjectPath, bool> NewAuthorityStates;
+	};
+}
+
+
diff --git a/Engine/Plugins/Developer/Concert/ConcertApp/MultiUserClient/Source/MultiUserClient/Private/Replication/Authority/AuthorityPolicy.cpp b/Engine/Plugins/Developer/Concert/ConcertApp/MultiUserClient/Source/MultiUserClient/Private/Replication/Authority/AuthorityPolicy.cpp
deleted file mode 100644
diff --git a/Engine/Plugins/Developer/Concert/ConcertApp/MultiUserClient/Source/MultiUserClient/Private/Replication/Authority/AuthorityPolicy.h b/Engine/Plugins/Developer/Concert/ConcertApp/MultiUserClient/Source/MultiUserClient/Private/Replication/Authority/AuthorityPolicy.h
deleted file mode 100644
diff --git a/Engine/Plugins/Developer/Concert/ConcertApp/MultiUserClient/Source/MultiUserClient/Private/Replication/Authority/AuthoritySynchronizer_Base.cpp b/Engine/Plugins/Developer/Concert/ConcertApp/MultiUserClient/Source/MultiUserClient/Private/Replication/Authority/AuthoritySynchronizer_Base.cpp
new file mode 100644
--- /dev/null
+++ b/Engine/Plugins/Developer/Concert/ConcertApp/MultiUserClient/Source/MultiUserClient/Private/Replication/Authority/AuthoritySynchronizer_Base.cpp
@@ -0,0 +1,29 @@
+// Copyright Epic Games, Inc. All Rights Reserved.
+
+#include "AuthoritySynchronizer_Base.h"
+
+namespace UE::MultiUserClient
+{
+	FAuthoritySynchronizer_Base::FAuthoritySynchronizer_Base(FDoesObjectHaveProperties InDoesObjectHavePropertiesDelegate)
+		: DoesObjectHavePropertiesDelegate(MoveTemp(InDoesObjectHavePropertiesDelegate))
+	{}
+
+	EAuthorityMutability FAuthoritySynchronizer_Base::GetChangeAuthorityMutability(const FSoftObjectPath& ObjectPath) const
+	{
+		// You can always take away your own authority
+		const bool bHasAuthority = HasAuthorityOver(ObjectPath);
+		if (bHasAuthority)
+		{
+			return EAuthorityMutability::Allowed;
+		}
+		
+		const bool bHasProperties = DoesObjectHavePropertiesDelegate.Execute(ObjectPath);
+		if (!bHasProperties)
+		{
+			return EAuthorityMutability::NoProperties;
+		}
+
+		// TODO DP UE-198356 check whether this would cause a conflict with other clients
+		return EAuthorityMutability::Allowed;
+	}
+}
diff --git a/Engine/Plugins/Developer/Concert/ConcertApp/MultiUserClient/Source/MultiUserClient/Private/Replication/Authority/AuthoritySynchronizer_Base.h b/Engine/Plugins/Developer/Concert/ConcertApp/MultiUserClient/Source/MultiUserClient/Private/Replication/Authority/AuthoritySynchronizer_Base.h
new file mode 100644
--- /dev/null
+++ b/Engine/Plugins/Developer/Concert/ConcertApp/MultiUserClient/Source/MultiUserClient/Private/Replication/Authority/AuthoritySynchronizer_Base.h
@@ -0,0 +1,43 @@
+// Copyright Epic Games, Inc. All Rights Reserved.
+
+#pragma once
+
+#include "IClientAuthoritySynchronizer.h"
+
+#include "Delegates/Delegate.h"
+#include "Templates/UnrealTemplate.h"
+
+namespace UE::MultiUserClient
+{
+	class FStreamChangeTracker;
+
+	DECLARE_DELEGATE_RetVal_OneParam(bool, FDoesObjectHaveProperties, const FSoftObjectPath&);
+	
+	/**
+	 * Util base implementation of CanChangeAuthority.
+	 * Changing authority is only possible if the client has properties associated with the object.
+	 */
+	class FAuthoritySynchronizer_Base
+		: public IClientAuthoritySynchronizer
+		, public FNoncopyable
+	{
+	public:
+
+		FAuthoritySynchronizer_Base(FDoesObjectHaveProperties InDoesObjectHavePropertiesDelegate);
+		
+		//~ Begin IClientAuthoritySynchronizer Interface
+		virtual EAuthorityMutability GetChangeAuthorityMutability(const FSoftObjectPath& ObjectPath) const override;
+		virtual FOnServerStateChanged& OnServerStateChanged() override { return OnServerStateChangedDelegate; }
+		//~ End IClientAuthoritySynchronizer Interface
+
+	protected:
+
+		/** Triggered by subclasses when the authority state changes. */
+		FOnServerStateChanged OnServerStateChangedDelegate;
+		
+	private:
+
+		/** Used to check whether the client will have properties associated with the object after submitting the change. */
+		FDoesObjectHaveProperties DoesObjectHavePropertiesDelegate;
+	};
+}
diff --git a/Engine/Plugins/Developer/Concert/ConcertApp/MultiUserClient/Source/MultiUserClient/Private/Replication/Authority/AuthoritySynchronizer_LocalClient.cpp b/Engine/Plugins/Developer/Concert/ConcertApp/MultiUserClient/Source/MultiUserClient/Private/Replication/Authority/AuthoritySynchronizer_LocalClient.cpp
new file mode 100644
--- /dev/null
+++ b/Engine/Plugins/Developer/Concert/ConcertApp/MultiUserClient/Source/MultiUserClient/Private/Replication/Authority/AuthoritySynchronizer_LocalClient.cpp
@@ -0,0 +1,42 @@
+// Copyright Epic Games, Inc. All Rights Reserved.
+
+#include "AuthoritySynchronizer_LocalClient.h"
+
+#include "IConcertSyncClient.h"
+#include "Replication/IConcertClientReplicationManager.h"
+
+namespace UE::MultiUserClient
+{
+	FAuthoritySynchronizer_LocalClient::FAuthoritySynchronizer_LocalClient(
+		TSharedRef<IConcertSyncClient> InClient,
+		FDoesObjectHaveProperties InDoesObjectHaveProperties
+		)
+		: FAuthoritySynchronizer_Base(MoveTemp(InDoesObjectHaveProperties))
+		, Client(MoveTemp(InClient))
+	{
+		IConcertClientReplicationManager* ReplicationManager = Client->GetReplicationManager();
+		if (ensure(ReplicationManager))
+		{
+			ReplicationManager->OnPostAuthorityChanged().AddRaw(this, &FAuthoritySynchronizer_LocalClient::OnAuthorityChanged);
+		}
+	}
+
+	FAuthoritySynchronizer_LocalClient::~FAuthoritySynchronizer_LocalClient()
+	{
+		if (IConcertClientReplicationManager* ReplicationManager = Client->GetReplicationManager())
+		{
+			ReplicationManager->OnPostAuthorityChanged().RemoveAll(this);
+		}
+	}
+
+	bool FAuthoritySynchronizer_LocalClient::HasAuthorityOver(const FSoftObjectPath& ObjectPath) const
+	{
+		const IConcertClientReplicationManager* ReplicationManager = Client->GetReplicationManager();
+		return ensure(ReplicationManager) && !ReplicationManager->GetClientOwnedStreamsForObject(ObjectPath).IsEmpty();
+	}
+
+	void FAuthoritySynchronizer_LocalClient::OnAuthorityChanged() const
+	{
+		OnServerStateChangedDelegate.Broadcast();
+	}
+}
diff --git a/Engine/Plugins/Developer/Concert/ConcertApp/MultiUserClient/Source/MultiUserClient/Private/Replication/Authority/AuthoritySynchronizer_LocalClient.h b/Engine/Plugins/Developer/Concert/ConcertApp/MultiUserClient/Source/MultiUserClient/Private/Replication/Authority/AuthoritySynchronizer_LocalClient.h
new file mode 100644
--- /dev/null
+++ b/Engine/Plugins/Developer/Concert/ConcertApp/MultiUserClient/Source/MultiUserClient/Private/Replication/Authority/AuthoritySynchronizer_LocalClient.h
@@ -0,0 +1,32 @@
+// Copyright Epic Games, Inc. All Rights Reserved.
+
+#pragma once
+
+#include "AuthoritySynchronizer_Base.h"
+#include "Templates/UnrealTemplate.h"
+
+class IConcertSyncClient;
+
+namespace UE::MultiUserClient
+{
+	class FAuthoritySynchronizer_LocalClient : public FAuthoritySynchronizer_Base
+	{
+	public:
+		
+		FAuthoritySynchronizer_LocalClient(TSharedRef<IConcertSyncClient> InClient, FDoesObjectHaveProperties InDoesObjectHaveProperties);
+		virtual ~FAuthoritySynchronizer_LocalClient() override;
+
+		//~ Begin IClientAuthoritySynchronizer Interface
+		virtual bool HasAuthorityOver(const FSoftObjectPath& ObjectPath) const override;
+		//~ End IClientAuthoritySynchronizer Interface
+
+	private:
+
+		/** Needed to access IConcertClientReplicationManager */
+		TSharedRef<IConcertSyncClient> Client;
+		
+		void OnAuthorityChanged() const;
+	};
+}
+
+
diff --git a/Engine/Plugins/Developer/Concert/ConcertApp/MultiUserClient/Source/MultiUserClient/Private/Replication/Authority/AuthoritySynchronizer_RemoteClient.cpp b/Engine/Plugins/Developer/Concert/ConcertApp/MultiUserClient/Source/MultiUserClient/Private/Replication/Authority/AuthoritySynchronizer_RemoteClient.cpp
new file mode 100644
--- /dev/null
+++ b/Engine/Plugins/Developer/Concert/ConcertApp/MultiUserClient/Source/MultiUserClient/Private/Replication/Authority/AuthoritySynchronizer_RemoteClient.cpp
@@ -0,0 +1,44 @@
+// Copyright Epic Games, Inc. All Rights Reserved.
+
+#include "AuthoritySynchronizer_RemoteClient.h"
+
+#include "Replication/Util/RegularQueryService.h"
+
+namespace UE::MultiUserClient
+{
+	FAuthoritySynchronizer_RemoteClient::FAuthoritySynchronizer_RemoteClient(
+		const FGuid& RemoteEndpointId,
+		FRegularQueryService& InQueryService,
+		FDoesObjectHaveProperties InDoesObjectHaveProperties
+		)
+		: FAuthoritySynchronizer_Base(MoveTemp(InDoesObjectHaveProperties))
+		, QueryService(InQueryService)
+		, QueryStreamHandle(
+    		QueryService.RegisterAuthorityQuery(
+    			RemoteEndpointId,
+    			FAuthorityQueryDelegate::CreateRaw(this, &FAuthoritySynchronizer_RemoteClient::HandleAuthorityQuery)
+    			)
+    		)
+	{}
+
+	FAuthoritySynchronizer_RemoteClient::~FAuthoritySynchronizer_RemoteClient()
+	{
+		QueryService.UnregisterAuthorityQuery(QueryStreamHandle);
+	}
+
+	EAuthorityMutability FAuthoritySynchronizer_RemoteClient::GetChangeAuthorityMutability(const FSoftObjectPath& ObjectPath) const
+	{
+		// TODO DP: UE-198088 return whether it is legal to change the authority
+		return EAuthorityMutability::NotSupported;
+	}
+
+	void FAuthoritySynchronizer_RemoteClient::HandleAuthorityQuery(const TArray<FReplicationAuthorityInfo>& PerStreamAuthority)
+	{
+		LastServerState.Empty();
+
+		for (const FReplicationAuthorityInfo& Info : PerStreamAuthority)
+		{
+			LastServerState.Append(Info.AuthoredObjects);
+		}
+	}
+}
diff --git a/Engine/Plugins/Developer/Concert/ConcertApp/MultiUserClient/Source/MultiUserClient/Private/Replication/Authority/AuthoritySynchronizer_RemoteClient.h b/Engine/Plugins/Developer/Concert/ConcertApp/MultiUserClient/Source/MultiUserClient/Private/Replication/Authority/AuthoritySynchronizer_RemoteClient.h
new file mode 100644
--- /dev/null
+++ b/Engine/Plugins/Developer/Concert/ConcertApp/MultiUserClient/Source/MultiUserClient/Private/Replication/Authority/AuthoritySynchronizer_RemoteClient.h
@@ -0,0 +1,44 @@
+// Copyright Epic Games, Inc. All Rights Reserved.
+
+#pragma once
+
+#include "AuthoritySynchronizer_Base.h"
+
+#include "Containers/Set.h"
+#include "UObject/SoftObjectPath.h"
+
+struct FGuid;
+struct FReplicationAuthorityInfo;
+
+namespace UE::MultiUserClient
+{
+	class FRegularQueryService;
+
+	class FAuthoritySynchronizer_RemoteClient : public FAuthoritySynchronizer_Base
+	{
+	public:
+		
+		FAuthoritySynchronizer_RemoteClient(const FGuid& RemoteEndpointId, FRegularQueryService& InQueryService, FDoesObjectHaveProperties InDoesObjectHaveProperties);
+		virtual ~FAuthoritySynchronizer_RemoteClient();
+
+		//~ Begin IClientAuthoritySynchronizer Interface
+		virtual EAuthorityMutability GetChangeAuthorityMutability(const FSoftObjectPath& ObjectPath) const override;
+		virtual bool HasAuthorityOver(const FSoftObjectPath& ObjectPath) const override { return LastServerState.Contains(ObjectPath); }
+		//~ End IClientAuthoritySynchronizer Interface
+
+	private:
+
+		/** Queries the server in regular intervals. This services outlives our object. */
+		FRegularQueryService& QueryService;
+		
+		/** Used to unregister HandleStreamQuery upon destruction. */
+		const FDelegateHandle QueryStreamHandle;
+
+		/** The most up to date server state of the remote client's authority. */
+		TSet<FSoftObjectPath> LastServerState; 
+
+		void HandleAuthorityQuery(const TArray<FReplicationAuthorityInfo>& PerStreamAuthority);
+	};
+}
+
+
diff --git a/Engine/Plugins/Developer/Concert/ConcertApp/MultiUserClient/Source/MultiUserClient/Private/Replication/Authority/EAuthorityMutability.h b/Engine/Plugins/Developer/Concert/ConcertApp/MultiUserClient/Source/MultiUserClient/Private/Replication/Authority/EAuthorityMutability.h
new file mode 100644
--- /dev/null
+++ b/Engine/Plugins/Developer/Concert/ConcertApp/MultiUserClient/Source/MultiUserClient/Private/Replication/Authority/EAuthorityMutability.h
@@ -0,0 +1,22 @@
+// Copyright Epic Games, Inc. All Rights Reserved.
+
+#pragma once
+
+namespace UE::MultiUserClient
+{
+	/** Various reasons in which an object can be when considering take over its authority */
+	enum class EAuthorityMutability
+	{
+		/** Authority can be taken */
+		Allowed,
+
+		/** The object has no registered properties. Cannot take authority. */
+		NoProperties,
+		
+		// TODO DP UE-198356: Add function to query which properties and clients are conflicting
+		/** Another client has authority over some of the object's registered properties already. Cannot take authority. */
+		ClientConflict,
+		/** The operation is not implemented. */
+		NotSupported
+	};
+}
\ No newline at end of file
diff --git a/Engine/Plugins/Developer/Concert/ConcertApp/MultiUserClient/Source/MultiUserClient/Private/Replication/Authority/IClientAuthoritySynchronizer.h b/Engine/Plugins/Developer/Concert/ConcertApp/MultiUserClient/Source/MultiUserClient/Private/Replication/Authority/IClientAuthoritySynchronizer.h
new file mode 100644
--- /dev/null
+++ b/Engine/Plugins/Developer/Concert/ConcertApp/MultiUserClient/Source/MultiUserClient/Private/Replication/Authority/IClientAuthoritySynchronizer.h
@@ -0,0 +1,32 @@
+// Copyright Epic Games, Inc. All Rights Reserved.
+
+#pragma once
+
+#include "EAuthorityMutability.h"
+
+#include "Delegates/Delegate.h"
+
+struct FSoftObjectPath;
+
+namespace UE::MultiUserClient
+{
+	/** Synchronizes the a client's authority state with the server. */
+	class IClientAuthoritySynchronizer
+	{
+	public:
+		
+		/** @return The reason of why it is or isn't possible to change the authority for the given object. */
+		virtual EAuthorityMutability GetChangeAuthorityMutability(const FSoftObjectPath& ObjectPath) const = 0;
+		/** @return Whether the ObjectPath's authority can be changed. */
+		bool CanChangeAuthority(const FSoftObjectPath& ObjectPath) const { return GetChangeAuthorityMutability(ObjectPath) == EAuthorityMutability::Allowed; }
+
+		/** @return Whether the local instance thinks this client has authority over ObjectPath. */
+		virtual bool HasAuthorityOver(const FSoftObjectPath& ObjectPath) const = 0;
+
+		DECLARE_MULTICAST_DELEGATE(FOnServerStateChanged);
+		/** @return Event executed when authority state has been updated. */
+		virtual FOnServerStateChanged& OnServerStateChanged() = 0;
+
+		virtual ~IClientAuthoritySynchronizer() = default;
+	};
+}
\ No newline at end of file
diff --git a/Engine/Plugins/Developer/Concert/ConcertApp/MultiUserClient/Source/MultiUserClient/Private/Replication/Client/LocalReplicationClient.cpp b/Engine/Plugins/Developer/Concert/ConcertApp/MultiUserClient/Source/MultiUserClient/Private/Replication/Client/LocalReplicationClient.cpp
--- a/Engine/Plugins/Developer/Concert/ConcertApp/MultiUserClient/Source/MultiUserClient/Private/Replication/Client/LocalReplicationClient.cpp
+++ b/Engine/Plugins/Developer/Concert/ConcertApp/MultiUserClient/Source/MultiUserClient/Private/Replication/Client/LocalReplicationClient.cpp
@@ -2,6 +2,9 @@
 
 #include "LocalReplicationClient.h"
 
+#include "Replication/Authority/AuthoritySynchronizer_LocalClient.h"
+#include "Replication/Submission/SubmissionWorkflow_LocalClient.h"
+
 namespace UE::MultiUserClient
 {
 	FLocalReplicationClient::FLocalReplicationClient(
@@ -9,8 +12,18 @@ namespace UE::MultiUserClient
 		TUniquePtr<IClientStreamSynchronizer> InStreamSynchronizer,
 		TSharedRef<IConcertSyncClient> InClient
 		)
-		: FReplicationClient(InSessionContent, MoveTemp(InStreamSynchronizer))
-		, AuthorityPolicy(MoveTemp(InClient), GetStreamSynchronizer())
-	{
-	}
+		: FReplicationClient(
+			InSessionContent,
+			MoveTemp(InStreamSynchronizer),
+			MakeUnique<FAuthoritySynchronizer_LocalClient>(
+				InClient,
+				FDoesObjectHaveProperties::CreateLambda([this](const FSoftObjectPath& ObjectPath)
+				{
+					return GetStreamDiffer().DoesObjectHavePropertiesAfterSubmit(ObjectPath);
+				})),
+			[this, InClient]()
+			{
+				return MakeUnique<FSubmissionWorkflow_LocalClient>(InClient, GetStreamDiffer(), GetAuthorityDiffer(), GetStreamSynchronizer());
+			})
+	{}
 }
diff --git a/Engine/Plugins/Developer/Concert/ConcertApp/MultiUserClient/Source/MultiUserClient/Private/Replication/Client/LocalReplicationClient.h b/Engine/Plugins/Developer/Concert/ConcertApp/MultiUserClient/Source/MultiUserClient/Private/Replication/Client/LocalReplicationClient.h
--- a/Engine/Plugins/Developer/Concert/ConcertApp/MultiUserClient/Source/MultiUserClient/Private/Replication/Client/LocalReplicationClient.h
+++ b/Engine/Plugins/Developer/Concert/ConcertApp/MultiUserClient/Source/MultiUserClient/Private/Replication/Client/LocalReplicationClient.h
@@ -3,7 +3,6 @@
 #pragma once
 
 #include "ReplicationClient.h"
-#include "Replication/Authority/AuthorityPolicy.h"
 
 namespace UE::MultiUserClient
 {
@@ -17,13 +16,6 @@ namespace UE::MultiUserClient
 			TUniquePtr<IClientStreamSynchronizer> InStreamSynchronizer,
 			TSharedRef<IConcertSyncClient> InClient
 			);
-
-		FAuthorityPolicy& GetAuthorityPolicy() { return AuthorityPolicy; }
-
-	private:
-
-		/** Automatically requests authority for submitted, locally-owned objects. */
-		FAuthorityPolicy AuthorityPolicy;
 	};
 }
 
diff --git a/Engine/Plugins/Developer/Concert/ConcertApp/MultiUserClient/Source/MultiUserClient/Private/Replication/Client/RemoteReplicationClient.cpp b/Engine/Plugins/Developer/Concert/ConcertApp/MultiUserClient/Source/MultiUserClient/Private/Replication/Client/RemoteReplicationClient.cpp
--- a/Engine/Plugins/Developer/Concert/ConcertApp/MultiUserClient/Source/MultiUserClient/Private/Replication/Client/RemoteReplicationClient.cpp
+++ b/Engine/Plugins/Developer/Concert/ConcertApp/MultiUserClient/Source/MultiUserClient/Private/Replication/Client/RemoteReplicationClient.cpp
@@ -3,7 +3,9 @@
 #include "RemoteReplicationClient.h"
 
 #include "Assets/MultiUserReplicationClientPreset.h"
-#include "Replication/Stream/RemoteClientStreamSynchronizer.h"
+#include "Replication/Authority/AuthoritySynchronizer_RemoteClient.h"
+#include "Replication/Stream/StreamSynchronizer_RemoteClient.h"
+#include "Replication/Submission/SubmissionWorkflow_RemoteClient.h"
 
 namespace UE::MultiUserClient
 {
@@ -12,7 +14,18 @@ namespace UE::MultiUserClient
 		const FGuid& InConcertClientId,
 		FRegularQueryService& QueryService
 		)
-		: FReplicationClient(InSessionContent, MakeUnique<FRemoteClientStreamSynchronizer>(InConcertClientId, QueryService))
+		: FReplicationClient(
+			InSessionContent,
+			MakeUnique<FStreamSynchronizer_RemoteClient>(InConcertClientId, QueryService),
+			MakeUnique<FAuthoritySynchronizer_RemoteClient>(
+				InConcertClientId,
+				QueryService,
+				FDoesObjectHaveProperties::CreateLambda([this](const FSoftObjectPath& ObjectPath)
+				{
+					return GetStreamDiffer().DoesObjectHavePropertiesAfterSubmit(ObjectPath);
+				})),
+			[]() { return MakeUnique<FSubmissionWorkflow_RemoteClient>(); }
+			)
 		, RemoteEndpointId(InConcertClientId)
 	{
 		// When the remote client's state has changed, refresh the UI.
diff --git a/Engine/Plugins/Developer/Concert/ConcertApp/MultiUserClient/Source/MultiUserClient/Private/Replication/Client/ReplicationClient.cpp b/Engine/Plugins/Developer/Concert/ConcertApp/MultiUserClient/Source/MultiUserClient/Private/Replication/Client/ReplicationClient.cpp
--- a/Engine/Plugins/Developer/Concert/ConcertApp/MultiUserClient/Source/MultiUserClient/Private/Replication/Client/ReplicationClient.cpp
+++ b/Engine/Plugins/Developer/Concert/ConcertApp/MultiUserClient/Source/MultiUserClient/Private/Replication/Client/ReplicationClient.cpp
@@ -3,25 +3,31 @@
 #include "ReplicationClient.h"
 
 #include "Assets/MultiUserReplicationClientPreset.h"
-#include "Replication/ReplicationWidgetFactories.h"
 #include "Replication/Editor/Model/IEditableObjectToPropertiesModel.h"
-#include "Replication/Stream/LocalStreamChangeTracker.h"
+#include "Replication/ReplicationWidgetFactories.h"
+#include "Replication/Stream/StreamChangeTracker.h"
 
 namespace UE::MultiUserClient
 {
 	FReplicationClient::FReplicationClient(
 		UMultiUserReplicationClientPreset& InSessionContent,
-		TUniquePtr<IClientStreamSynchronizer> InStreamSynchronizer
+		TUniquePtr<IClientStreamSynchronizer> InStreamSynchronizer,
+		TUniquePtr<IClientAuthoritySynchronizer> InAuthoritySynchronizer,
+		TFunctionRef<TUniquePtr<ISubmissionWorkflow>()> MakeSubmissionWorkflowFunc
 		)
 		: ClientContentStorage(&InSessionContent)
 		, StreamSynchronizer(MoveTemp(InStreamSynchronizer))
+		, AuthoritySynchronizer(MoveTemp(InAuthoritySynchronizer))
 		, LocalClientEditModel(ConcertClientSharedSlate::CreatePropertySelectionModel(*ClientContentStorage->Stream, ClientContentStorage->Stream->MakeReplicationMapGetterAttribute()))
 		, LocalClientStreamDiffer(
 			GetStreamSynchronizer(),
 			ClientContentStorage->Stream->MakeReplicationMapGetterAttribute(),
-			FLocalStreamChangeTracker::FOnModifyReplicationMap::CreateLambda([this](){ ClientContentStorage->Stream->Modify(); })
+			FStreamChangeTracker::FOnModifyReplicationMap::CreateLambda([this](){ ClientContentStorage->Stream->Modify(); })
 			)
+		, LocalAuthorityDiffer(*AuthoritySynchronizer)
 	{
+		SubmissionWorkflow = MakeSubmissionWorkflowFunc();
+		
 		LocalClientEditModel->OnObjectsChanged().AddRaw(this, &FReplicationClient::OnObjectsChanged);
 		LocalClientEditModel->OnPropertiesChanged().AddRaw(this, &FReplicationClient::OnPropertiesChanged);
 		
diff --git a/Engine/Plugins/Developer/Concert/ConcertApp/MultiUserClient/Source/MultiUserClient/Private/Replication/Client/ReplicationClient.h b/Engine/Plugins/Developer/Concert/ConcertApp/MultiUserClient/Source/MultiUserClient/Private/Replication/Client/ReplicationClient.h
--- a/Engine/Plugins/Developer/Concert/ConcertApp/MultiUserClient/Source/MultiUserClient/Private/Replication/Client/ReplicationClient.h
+++ b/Engine/Plugins/Developer/Concert/ConcertApp/MultiUserClient/Source/MultiUserClient/Private/Replication/Client/ReplicationClient.h
@@ -2,11 +2,14 @@
 
 #pragma once
 
-#include "Replication/Stream/LocalStreamChangeTracker.h"
+#include "Replication/Authority/AuthorityChangeTracker.h"
+#include "Replication/Authority/IClientAuthoritySynchronizer.h"
+#include "Replication/Stream/IClientStreamSynchronizer.h"
+#include "Replication/Stream/StreamChangeTracker.h"
+#include "Replication/Submission/ISubmissionWorkflow.h"
 #include "Templates/SharedPointer.h"
 #include "Templates/UnrealTemplate.h"
 
-class IConcertSyncClient;
 class UMultiUserReplicationClientPreset;
 
 namespace UE::ConcertClientSharedSlate
@@ -17,22 +20,33 @@ namespace UE::ConcertClientSharedSlate
 
 namespace UE::MultiUserClient
 {
+	class ISubmissionWorkflow;
+	
 	/** Holds on to info about a local or remote client. */
 	class FReplicationClient : public FNoncopyable
 	{
 	public:
 
+		DECLARE_DELEGATE_RetVal(TUniquePtr<ISubmissionWorkflow>, FMakeSubmissionWorkflow);
+
 		FReplicationClient(
 			UMultiUserReplicationClientPreset& InSessionContent,
-			TUniquePtr<IClientStreamSynchronizer> InStreamSynchronizer
+			TUniquePtr<IClientStreamSynchronizer> InStreamSynchronizer,
+			TUniquePtr<IClientAuthoritySynchronizer> InAuthoritySynchronizer,
+			TFunctionRef<TUniquePtr<ISubmissionWorkflow>()> MakeSubmissionWorkflowFunc
 			);
 
 		UMultiUserReplicationClientPreset* GetClientContent() const { return ClientContentStorage; }
 		TSharedRef<ConcertClientSharedSlate::IEditableObjectToPropertiesModel> GetClientEditModel() const { return LocalClientEditModel; }
-		IClientStreamSynchronizer& GetStreamSynchronizer() const { return *StreamSynchronizer.Get(); } 
+		IClientStreamSynchronizer& GetStreamSynchronizer() const { return *StreamSynchronizer.Get(); }
+		IClientAuthoritySynchronizer& GetAuthoritySynchronizer() const { return *AuthoritySynchronizer.Get(); }
 		
-		const FLocalStreamChangeTracker& GetDiffer() const { return LocalClientStreamDiffer; }
-		FLocalStreamChangeTracker& GetDiffer() { return LocalClientStreamDiffer; }
+		const FStreamChangeTracker& GetStreamDiffer() const { return LocalClientStreamDiffer; }
+		FStreamChangeTracker& GetStreamDiffer() { return LocalClientStreamDiffer; }
+		const FAuthorityChangeTracker& GetAuthorityDiffer() const { return LocalAuthorityDiffer; }
+		FAuthorityChangeTracker& GetAuthorityDiffer() { return LocalAuthorityDiffer; }
+		const ISubmissionWorkflow& GetSubmissionWorkflow() const { return *SubmissionWorkflow; }
+		ISubmissionWorkflow& GetSubmissionWorkflow() { return *SubmissionWorkflow; }
 
 		/**
 		 * Called when the data underlying the model has changed externally. Since the change was not caused by the model,
@@ -51,8 +65,12 @@ namespace UE::MultiUserClient
 		/** The state of the server is synched up with this object and displayed in the UI. */
 		TObjectPtr<UMultiUserReplicationClientPreset> ClientContentStorage;
 		
-		/** Keeps the client's state on the server in sync. */
+		/** Keeps the client's stream state on the server in sync. */
 		TUniquePtr<IClientStreamSynchronizer> StreamSynchronizer;
+		/** Keeps the client's authority state on the server in sync.*/
+		TUniquePtr<IClientAuthoritySynchronizer> AuthoritySynchronizer;
+		/** Handles the logic of submitting and reverting for this client. */
+		TUniquePtr<ISubmissionWorkflow> SubmissionWorkflow;
 		
 		/**
 		 * Used to detect changes made to the client's config by the local editor.
@@ -64,8 +82,11 @@ namespace UE::MultiUserClient
 		 * @see FMultiUserReplicationManager::OnLeaveSession
 		 */
 		TSharedRef<ConcertClientSharedSlate::IEditableObjectToPropertiesModel> LocalClientEditModel;
+		
 		/** Tracks changes made to server's state of the client's streams and prepares to upload them using StreamSynchronizer. */
-		FLocalStreamChangeTracker LocalClientStreamDiffer;
+		FStreamChangeTracker LocalClientStreamDiffer;
+		/** Tracks changes made to the client's authority state. */
+		FAuthorityChangeTracker LocalAuthorityDiffer;
 
 		/** Called when the data underlying the model has changed (and the UI needs to be refreshed). */
 		FOnModelExternallyChanged OnModelExternallyChangedDelegate;
diff --git a/Engine/Plugins/Developer/Concert/ConcertApp/MultiUserClient/Source/MultiUserClient/Private/Replication/Client/ReplicationClientManager.cpp b/Engine/Plugins/Developer/Concert/ConcertApp/MultiUserClient/Source/MultiUserClient/Private/Replication/Client/ReplicationClientManager.cpp
--- a/Engine/Plugins/Developer/Concert/ConcertApp/MultiUserClient/Source/MultiUserClient/Private/Replication/Client/ReplicationClientManager.cpp
+++ b/Engine/Plugins/Developer/Concert/ConcertApp/MultiUserClient/Source/MultiUserClient/Private/Replication/Client/ReplicationClientManager.cpp
@@ -6,7 +6,7 @@
 #include "LocalReplicationClient.h"
 #include "RemoteReplicationClient.h"
 #include "Assets/MultiUserReplicationSessionPreset.h"
-#include "Replication/Stream/LocalClientStreamSynchronizer.h"
+#include "Replication/Stream/StreamSynchronizer_LocalClient.h"
 
 #include "UObject/Package.h"
 #include "UObject/UObjectGlobals.h"
@@ -23,7 +23,7 @@ namespace UE::MultiUserClient
 		, LocalClient([this, InClient]()
 		{
 			UMultiUserReplicationClientPreset* ClientPreset = SessionContent->AddClient();
-			return FLocalReplicationClient(*ClientPreset, MakeUnique<FLocalClientStreamSynchronizer>(InClient, ClientPreset->Stream->StreamId), InClient);
+			return FLocalReplicationClient(*ClientPreset, MakeUnique<FStreamSynchronizer_LocalClient>(InClient, ClientPreset->Stream->StreamId), InClient);
 		}())
 	{
 		InSession->OnSessionClientChanged().AddRaw(this, &FReplicationClientManager::OnSessionClientChanged);
diff --git a/Engine/Plugins/Developer/Concert/ConcertApp/MultiUserClient/Source/MultiUserClient/Private/Replication/Client/ReplicationClientManager.h b/Engine/Plugins/Developer/Concert/ConcertApp/MultiUserClient/Source/MultiUserClient/Private/Replication/Client/ReplicationClientManager.h
--- a/Engine/Plugins/Developer/Concert/ConcertApp/MultiUserClient/Source/MultiUserClient/Private/Replication/Client/ReplicationClientManager.h
+++ b/Engine/Plugins/Developer/Concert/ConcertApp/MultiUserClient/Source/MultiUserClient/Private/Replication/Client/ReplicationClientManager.h
@@ -27,7 +27,7 @@ namespace UE::MultiUserClient
 {
 	class FRemoteReplicationClient;
 	class FReplicationClient;
-	class FLocalStreamChangeTracker;
+	class FStreamChangeTracker;
 
 	/**
 	 * Keeps track of connected clients synchronizing their stream data in a UMultiUserReplicationSessionPreset.
diff --git a/Engine/Plugins/Developer/Concert/ConcertApp/MultiUserClient/Source/MultiUserClient/Private/Replication/MultiUserReplicationManager.cpp b/Engine/Plugins/Developer/Concert/ConcertApp/MultiUserClient/Source/MultiUserClient/Private/Replication/MultiUserReplicationManager.cpp
--- a/Engine/Plugins/Developer/Concert/ConcertApp/MultiUserClient/Source/MultiUserClient/Private/Replication/MultiUserReplicationManager.cpp
+++ b/Engine/Plugins/Developer/Concert/ConcertApp/MultiUserClient/Source/MultiUserClient/Private/Replication/MultiUserReplicationManager.cpp
@@ -50,13 +50,6 @@ namespace UE::MultiUserClient
 			});
 	}
 
-	FAuthorityPolicy* FMultiUserReplicationManager::GetAuthorityPolicy()
-	{
-		return ConnectedState
-			? &ConnectedState->ClientManager.GetLocalClient().GetAuthorityPolicy()
-			: nullptr;
-	}
-
 	void FMultiUserReplicationManager::OnSessionConnectionChanged(
 		IConcertClientSession& ConcertClientSession,
 		EConcertConnectionStatus ConcertConnectionStatus
diff --git a/Engine/Plugins/Developer/Concert/ConcertApp/MultiUserClient/Source/MultiUserClient/Private/Replication/MultiUserReplicationManager.h b/Engine/Plugins/Developer/Concert/ConcertApp/MultiUserClient/Source/MultiUserClient/Private/Replication/MultiUserReplicationManager.h
--- a/Engine/Plugins/Developer/Concert/ConcertApp/MultiUserClient/Source/MultiUserClient/Private/Replication/MultiUserReplicationManager.h
+++ b/Engine/Plugins/Developer/Concert/ConcertApp/MultiUserClient/Source/MultiUserClient/Private/Replication/MultiUserReplicationManager.h
@@ -2,11 +2,9 @@
 
 #pragma once
 
-#include "Replication/Authority/AuthorityPolicy.h"
-
-#include "IConcertSession.h"
 #include "Client/ReplicationClientManager.h"
 
+#include "IConcertSession.h"
 #include "Misc/Optional.h"
 #include "Templates/SharedPointer.h"
 #include "Templates/UnrealTemplate.h"
@@ -67,9 +65,6 @@ namespace UE::MultiUserClient
 		FReplicationClientManager* GetClientManager() { return ConnectedState ? &ConnectedState->ClientManager : nullptr; }
 		const FReplicationClientManager* GetClientManager() const { return ConnectedState ? &ConnectedState->ClientManager : nullptr; }
 
-		/** @note You're not supposed to keep any reference to AuthorityPolicy since it can become invalid depending on connection state. */
-		FAuthorityPolicy* GetAuthorityPolicy();
-
 		/** Called when the connection to the replication system changes. */
 		DECLARE_MULTICAST_DELEGATE_OneParam(FOnReplicationConnectionStateChanged, EMultiUserReplicationConnectionState /*NewState*/);
 		FOnReplicationConnectionStateChanged& OnReplicationConnectionStateChanged() { return OnReplicationConnectionStateChangedDelegate; }
diff --git a/Engine/Plugins/Developer/Concert/ConcertApp/MultiUserClient/Source/MultiUserClient/Private/Replication/Stream/IClientStreamSynchronizer.h b/Engine/Plugins/Developer/Concert/ConcertApp/MultiUserClient/Source/MultiUserClient/Private/Replication/Stream/IClientStreamSynchronizer.h
--- a/Engine/Plugins/Developer/Concert/ConcertApp/MultiUserClient/Source/MultiUserClient/Private/Replication/Stream/IClientStreamSynchronizer.h
+++ b/Engine/Plugins/Developer/Concert/ConcertApp/MultiUserClient/Source/MultiUserClient/Private/Replication/Stream/IClientStreamSynchronizer.h
@@ -2,51 +2,13 @@
 
 #pragma once
 
-#include "Replication/Data/ObjectIds.h"
-#include "Replication/Messages/ChangeStream.h"
-
 #include "Async/Future.h"
-#include "Replication/IConcertClientReplicationManager.h"
+#include "Delegates/Delegate.h"
 
 struct FObjectReplicationMap;
 
-namespace UE::ConcertSyncClient::Replication
-{
-	struct FChangeStreamRequest;
-	struct FChangeStreamResponse;
-}
-
 namespace UE::MultiUserClient
 {
-	struct FStreamChangelist
-	{
-		TSet<FObjectInStreamID> ObjectsToRemove;
-		TMap<FObjectInStreamID, FConcertReplication_ChangeStream_PutObject> ObjectsToPut;
-	};
-	
-	enum class ESubmitChangesErrorCode
-	{
-		/** Changes were accepted by the server */
-		Success,
-		/** Skipped because another request is in progress. */
-		CannotSendRequest
-	};
-
-	struct FCompletedChangeSubmission
-	{
-		ConcertSyncClient::Replication::FChangeStreamRequest Request;
-		ConcertSyncClient::Replication::FChangeStreamResponse Response;
-	};
-
-	struct FSubmitChangesResult
-	{
-		/** Error code of the submission. Determines whether Response is valid. */
-		ESubmitChangesErrorCode ErrorCode;
-
-		/** Valid if ErrorCode != ESubmitChangesErrorCode::AlreadyInProgress. */
-		TOptional<FCompletedChangeSubmission> SubmissionInfo;
-	};
-	
 	/**
 	 * Keeps track of a client's registered streams.
 	 * Exposes ways for observing and mutating streams.
@@ -56,25 +18,14 @@ namespace UE::MultiUserClient
 	class IClientStreamSynchronizer
 	{
 	public:
-
-		/**
-		 * Requests to change the stream. The future executes once the changes have been accepted on the server.
-		 * @note This future can execute on any thread.
-		 */
-		virtual TFuture<FSubmitChangesResult> SubmitChanges(const FStreamChangelist& Changelist) = 0;
-		/** @return Whether it's currently valid to call SubmitChanges. For example, it is not valid to call while there is already request in progress. */
-		virtual bool CanMakeSubmitRequest() const = 0;
-
+		
 		/** @return The managed client's stream ID. */
 		virtual FGuid GetStreamId() const = 0;
 		/** @return What local instance thinks the client's server state is. */
 		virtual const FObjectReplicationMap& GetServerState() const = 0;
 
-		DECLARE_MULTICAST_DELEGATE_TwoParams(FOnChangesAccepted, const FObjectReplicationMap& OldServerState, const ConcertSyncClient::Replication::FChangeStreamRequest& AcceptedRequest);
-		/** @return Event executed when a change request that was in transit was accepted by the server. This is executed before OnServerStateSynched: the local confirmed state is not updated, yet. */
-		virtual FOnChangesAccepted& OnChangesAccepted() = 0;
 		DECLARE_MULTICAST_DELEGATE(FOnServerStateChanged);
-		/** @return Event executed when the result of GetServerState has been synched. Called after OnChangesAccepted_AnyThread. */
+		/** @return Event executed when the result of GetServerState has been updated. */
 		virtual FOnServerStateChanged& OnServerStateChanged() = 0;
 		
 		virtual ~IClientStreamSynchronizer() = default;
diff --git a/Engine/Plugins/Developer/Concert/ConcertApp/MultiUserClient/Source/MultiUserClient/Private/Replication/Stream/LocalClientStreamSynchronizer.cpp b/Engine/Plugins/Developer/Concert/ConcertApp/MultiUserClient/Source/MultiUserClient/Private/Replication/Stream/LocalClientStreamSynchronizer.cpp
deleted file mode 100644
diff --git a/Engine/Plugins/Developer/Concert/ConcertApp/MultiUserClient/Source/MultiUserClient/Private/Replication/Stream/LocalClientStreamSynchronizer.h b/Engine/Plugins/Developer/Concert/ConcertApp/MultiUserClient/Source/MultiUserClient/Private/Replication/Stream/LocalClientStreamSynchronizer.h
deleted file mode 100644
diff --git a/Engine/Plugins/Developer/Concert/ConcertApp/MultiUserClient/Source/MultiUserClient/Private/Replication/Stream/LocalStreamChangeTracker.cpp b/Engine/Plugins/Developer/Concert/ConcertApp/MultiUserClient/Source/MultiUserClient/Private/Replication/Stream/LocalStreamChangeTracker.cpp
deleted file mode 100644
diff --git a/Engine/Plugins/Developer/Concert/ConcertApp/MultiUserClient/Source/MultiUserClient/Private/Replication/Stream/LocalStreamChangeTracker.h b/Engine/Plugins/Developer/Concert/ConcertApp/MultiUserClient/Source/MultiUserClient/Private/Replication/Stream/LocalStreamChangeTracker.h
deleted file mode 100644
diff --git a/Engine/Plugins/Developer/Concert/ConcertApp/MultiUserClient/Source/MultiUserClient/Private/Replication/Stream/RemoteClientStreamSynchronizer.cpp b/Engine/Plugins/Developer/Concert/ConcertApp/MultiUserClient/Source/MultiUserClient/Private/Replication/Stream/RemoteClientStreamSynchronizer.cpp
deleted file mode 100644
diff --git a/Engine/Plugins/Developer/Concert/ConcertApp/MultiUserClient/Source/MultiUserClient/Private/Replication/Stream/RemoteClientStreamSynchronizer.h b/Engine/Plugins/Developer/Concert/ConcertApp/MultiUserClient/Source/MultiUserClient/Private/Replication/Stream/RemoteClientStreamSynchronizer.h
deleted file mode 100644
diff --git a/Engine/Plugins/Developer/Concert/ConcertApp/MultiUserClient/Source/MultiUserClient/Private/Replication/Stream/StreamChangeTracker.cpp b/Engine/Plugins/Developer/Concert/ConcertApp/MultiUserClient/Source/MultiUserClient/Private/Replication/Stream/StreamChangeTracker.cpp
new file mode 100644
--- /dev/null
+++ b/Engine/Plugins/Developer/Concert/ConcertApp/MultiUserClient/Source/MultiUserClient/Private/Replication/Stream/StreamChangeTracker.cpp
@@ -0,0 +1,119 @@
+// Copyright Epic Games, Inc. All Rights Reserved.
+
+#include "StreamChangeTracker.h"
+
+#include "Replication/ChangeStreamSharedUtils.h"
+#include "Replication/IConcertClientReplicationManager.h"
+
+#include "ScopedTransaction.h"
+
+#define LOCTEXT_NAMESPACE "FLocalClientStreamDiffer"
+
+namespace UE::MultiUserClient
+{
+	FStreamChangeTracker::FStreamChangeTracker(
+		IClientStreamSynchronizer& InStreamSynchronizer,
+		TAttribute<FObjectReplicationMap*> InStreamWithInProgressChangesAttribute,
+		FOnModifyReplicationMap InOnModifyReplicationMapDelegate
+		)
+		: StreamSynchronizer(InStreamSynchronizer)
+		, StreamWithInProgressChangesAttribute(MoveTemp(InStreamWithInProgressChangesAttribute))
+		, OnModifyReplicationMapDelegate(MoveTemp(InOnModifyReplicationMapDelegate))
+	{
+		check(InStreamWithInProgressChangesAttribute.IsBound() || InStreamWithInProgressChangesAttribute.IsSet());
+		StreamSynchronizer.OnServerStateChanged().AddRaw(this, &FStreamChangeTracker::RefreshChangesCache);
+	}
+
+	FStreamChangeTracker::~FStreamChangeTracker()
+	{
+		StreamSynchronizer.OnServerStateChanged().RemoveAll(this);
+	}
+
+	void FStreamChangeTracker::RefreshChangesCache()
+	{
+		CachedDeltaChange = DiffChanges();
+	}
+
+	void FStreamChangeTracker::RevertCachedChanges()
+	{
+		if (HasChanges())
+		{
+			FScopedTransaction Transaction(LOCTEXT("RevertCachedChanges", "Revert replication changes"));
+			OnModifyReplicationMapDelegate.ExecuteIfBound();
+			
+			*StreamWithInProgressChangesAttribute.Get() = StreamSynchronizer.GetServerState();
+			CachedDeltaChange = {};
+
+			check(IsInGameThread());
+			OnChangesRevertedDelegate.Broadcast();
+			RefreshChangesCache();
+		}
+	}
+
+	bool FStreamChangeTracker::HasChanges() const
+	{
+		return !CachedDeltaChange.ObjectsToPut.IsEmpty() || !CachedDeltaChange.ObjectsToRemove.IsEmpty();
+	}
+
+	bool FStreamChangeTracker::DoesObjectHavePropertiesAfterSubmit(const FSoftObjectPath& ObjectPath) const
+	{
+		const EObjectChangeType ChangeType = GetObjectChanges(ObjectPath);
+		switch (ChangeType)
+		{
+		case EObjectChangeType::NoChange: return StreamSynchronizer.GetServerState().ReplicatedObjects.Contains(ObjectPath);
+			
+		case EObjectChangeType::PropertiesModified: return true;
+		case EObjectChangeType::Added: return true;
+			
+		case EObjectChangeType::Removed: return false;
+		default:
+			checkNoEntry();
+			return false;
+		}
+	}
+
+	FStreamChangeTracker::EObjectChangeType FStreamChangeTracker::GetObjectChanges(const FSoftObjectPath& Object) const
+	{
+		const FGuid StreamId = StreamSynchronizer.GetStreamId();
+		const FObjectInStreamID ObjectId { StreamId, Object };
+		if (CachedDeltaChange.ObjectsToRemove.Contains(ObjectId))
+		{
+			return EObjectChangeType::Removed;
+		}
+
+		const FConcertReplication_ChangeStream_PutObject* PutObject = CachedDeltaChange.ObjectsToPut.Find(ObjectId);
+		if (!PutObject)
+		{
+			return EObjectChangeType::NoChange;
+		}
+
+		const bool bObjectIsOnServer = StreamSynchronizer.GetServerState().ReplicatedObjects.Contains(Object);
+		return bObjectIsOnServer ? EObjectChangeType::PropertiesModified : EObjectChangeType::Added;
+	}
+
+	FStreamChangeTracker::EPropertyChangeType FStreamChangeTracker::GetPropertyChanges(const FSoftObjectPath& Object, const FConcertPropertyChain& PropertyChain) const
+	{
+		// TODO DP:
+		return EPropertyChangeType::NoChange;
+	}
+	
+	FStreamChangelist FStreamChangeTracker::DiffChanges(
+		const FGuid& StreamId,
+		const FObjectReplicationMap& Base,
+		const FObjectReplicationMap& Changed
+		)
+	{
+		// BuildRequestFromDiff does not tolerate any invalid entries (like empty properties, which the UI generates right after you add an object to the list)
+		FObjectReplicationMap Cleansed = Changed;
+		ConcertSyncCore::Replication::ChangeStreamUtils::IterateInvalidEntries(Changed, [&Cleansed](const FSoftObjectPath& InvalidObject, const FReplicatedObjectInfo&)
+		{
+			Cleansed.ReplicatedObjects.Remove(InvalidObject);
+			return EBreakBehavior::Continue;
+		});
+
+		FConcertReplication_ChangeStream_Request Request = ConcertSyncCore::Replication::ChangeStreamUtils::BuildRequestFromDiff(StreamId, Base, Cleansed);
+		return { MoveTemp(Request.ObjectsToRemove), MoveTemp(Request.ObjectsToPut) };
+	}
+}
+
+#undef LOCTEXT_NAMESPACE
\ No newline at end of file
diff --git a/Engine/Plugins/Developer/Concert/ConcertApp/MultiUserClient/Source/MultiUserClient/Private/Replication/Stream/StreamChangeTracker.h b/Engine/Plugins/Developer/Concert/ConcertApp/MultiUserClient/Source/MultiUserClient/Private/Replication/Stream/StreamChangeTracker.h
new file mode 100644
--- /dev/null
+++ b/Engine/Plugins/Developer/Concert/ConcertApp/MultiUserClient/Source/MultiUserClient/Private/Replication/Stream/StreamChangeTracker.h
@@ -0,0 +1,130 @@
+// Copyright Epic Games, Inc. All Rights Reserved.
+
+#pragma once
+
+#include "IClientStreamSynchronizer.h"
+#include "Replication/Data/ObjectIds.h"
+#include "Replication/Data/ObjectReplicationMap.h"
+#include "Replication/Messages/ChangeStream.h"
+
+#include "Async/Future.h"
+#include "Containers/Set.h"
+#include "Containers/Map.h"
+#include "Misc/Attribute.h"
+#include "Templates/UnrealTemplate.h"
+
+class IConcertSyncClient;
+
+namespace UE::ConcertSyncClient::Replication
+{
+	struct FChangeStreamResponse;
+	struct FChangeStreamRequest;
+}
+
+namespace UE::MultiUserClient
+{
+	enum class EObjectWarningFlags
+	{
+		Ok = 0,
+		MissingProperties = 1 << 0
+	};
+	ENUM_CLASS_FLAGS(EObjectWarningFlags);
+	
+	/** Describes changes that MU client makes. */
+	struct FStreamChangelist
+	{
+		TSet<FObjectInStreamID> ObjectsToRemove;
+		TMap<FObjectInStreamID, FConcertReplication_ChangeStream_PutObject> ObjectsToPut;
+	};
+	
+	/**
+	 * Knows of the local client's registered replication streams and builds a changelist. The changelist tracks the
+	 * unconfirmed changes to the client's streams and is updates when the server confirms the change.
+	 */
+	class FStreamChangeTracker : public FNoncopyable
+	{
+	public:
+
+		DECLARE_DELEGATE(FOnModifyReplicationMap);
+		FStreamChangeTracker(
+			IClientStreamSynchronizer& InStreamSynchronizer,
+			TAttribute<FObjectReplicationMap*> InStreamWithInProgressChangesAttribute,
+			FOnModifyReplicationMap InOnModifyReplicationMapDelegate
+			);
+		~FStreamChangeTracker();
+
+		/**
+		 * Clears and rebuilds CachedDeltaChange in response to the underlying StreamWithInProgressChangesAttribute having been changed.
+		 * 
+		 * This iterates through all elements, which is suboptimal. However, this is needed e.g. after an undo & redo command.
+		 * In the future, we could implement a more efficient version that just responds to the relevant events (like adding an object).
+		 * While this function does redundant work, it should really not be a big performance bog after all. Also it is called
+		 * infrequently e.g. in response to a user pressing a button so we should be fine.
+		 */
+		void RefreshChangesCache();
+		const FStreamChangelist& GetCachedDeltaChange() const { return CachedDeltaChange; }
+		
+		/** Reverts the external stream StreamWithInProgressChangesAttribute to be the version registered on the server. */
+		void RevertCachedChanges();
+
+		/** @return Whether there are any changes that that can be submitted to the server (excludes those with warnings). */
+		bool HasChanges() const;
+
+		/** @return If change are submitted, is ObjectPath in the stream? */
+		bool DoesObjectHavePropertiesAfterSubmit(const FSoftObjectPath& ObjectPath) const;
+
+		enum class EObjectChangeType
+		{
+			NoChange,
+			Added,
+			Removed,
+			PropertiesModified
+		};
+		/** @return Whether this object's local configuration differs from the server version. */
+		EObjectChangeType GetObjectChanges(const FSoftObjectPath& Object) const;
+		
+		enum class EPropertyChangeType
+		{
+			NoChange,
+			Added,
+			Removed
+		};
+		/** @return Gets how the given property config differs from what is on the server. */
+		EPropertyChangeType GetPropertyChanges(const FSoftObjectPath& Object, const FConcertPropertyChain& PropertyChain) const;
+
+		/** Called when RevertCachedChanges is called. The UI should be refreshed. */
+		DECLARE_MULTICAST_DELEGATE(FOnChangesReverted);
+		FOnChangesReverted& OnChangesReverted_GameThread() { return OnChangesRevertedDelegate; }
+		
+	private:
+
+		/** Keeps track of the edited stream's server state. Outlives the FStreamChangeTracker instance. */
+		IClientStreamSynchronizer& StreamSynchronizer;
+
+		/**
+		 * Represents ConfirmedServerState with changes made to it. These changes have not been sent to the server, yet.
+		 * This is cleared when a change is sent to the server. 
+		 */
+		const TAttribute<FObjectReplicationMap*> StreamWithInProgressChangesAttribute;
+		
+		/**
+		 * The changes that if applied to ConfirmedServerState would result in StreamWithInProgressChangesAttribute.
+		 * This is updated every time the local client makes changes.
+		 */
+		FStreamChangelist CachedDeltaChange;
+		
+		/** Called when StreamWithInProgressChangesAttribute is about to be modified. */
+		FOnModifyReplicationMap OnModifyReplicationMapDelegate;
+		/** Called when StreamWithInProgressChangesAttribute was modified by this FLocalClientStreamDiffer. */
+		FOnChangesReverted OnChangesRevertedDelegate;
+
+		FStreamChangelist DiffChanges() const
+		{
+			const FObjectReplicationMap* Map = StreamWithInProgressChangesAttribute.Get();
+			return Map ? DiffChanges(StreamSynchronizer.GetStreamId(), StreamSynchronizer.GetServerState(), *Map) : FStreamChangelist{};
+		}
+		/** Builds the changelist to get from Base to Changed. */
+		static FStreamChangelist DiffChanges(const FGuid& StreamId, const FObjectReplicationMap& Base, const FObjectReplicationMap& Changed);
+	};
+}
+
diff --git a/Engine/Plugins/Developer/Concert/ConcertApp/MultiUserClient/Source/MultiUserClient/Private/Replication/Stream/StreamSynchronizer_LocalClient.cpp b/Engine/Plugins/Developer/Concert/ConcertApp/MultiUserClient/Source/MultiUserClient/Private/Replication/Stream/StreamSynchronizer_LocalClient.cpp
new file mode 100644
--- /dev/null
+++ b/Engine/Plugins/Developer/Concert/ConcertApp/MultiUserClient/Source/MultiUserClient/Private/Replication/Stream/StreamSynchronizer_LocalClient.cpp
@@ -0,0 +1,65 @@
+// Copyright Epic Games, Inc. All Rights Reserved.
+
+#include "StreamSynchronizer_LocalClient.h"
+
+#include "IConcertSyncClient.h"
+#include "Replication/IConcertClientReplicationManager.h"
+#include "Replication/Messages/ChangeStream.h"
+
+#define LOCTEXT_NAMESPACE "FLocalClientStreamDiffer"
+
+namespace UE::MultiUserClient
+{
+	FStreamSynchronizer_LocalClient::FStreamSynchronizer_LocalClient(
+		TSharedRef<IConcertSyncClient> InLocalClient,
+		const FGuid& InLocalClientStreamId
+		)
+		: LocalClient(MoveTemp(InLocalClient))
+		, LocalClientStreamId(InLocalClientStreamId)
+	{
+		IConcertClientReplicationManager* ReplicationManager = LocalClient->GetReplicationManager();
+		if (ensure(ReplicationManager))
+		{
+			ReplicationManager->OnPostStreamsChanged().AddRaw(this, &FStreamSynchronizer_LocalClient::OnPostStreamsChanged);
+		}
+	}
+
+	FStreamSynchronizer_LocalClient::~FStreamSynchronizer_LocalClient()
+	{
+		if (IConcertClientReplicationManager* ReplicationManager = LocalClient->GetReplicationManager())
+		{
+			ReplicationManager->OnPostStreamsChanged().RemoveAll(this);
+		}
+	}
+
+	const FObjectReplicationMap& FStreamSynchronizer_LocalClient::GetServerState() const
+	{
+		IConcertClientReplicationManager* ReplicationManager = LocalClient->GetReplicationManager();
+		if (!ensure(ReplicationManager))
+		{
+			return EmptyState;
+		}
+		
+		const FObjectReplicationMap* Result = nullptr;
+		ReplicationManager->ForEachRegisteredStream([this, &Result](const FReplicationStreamDescription& Stream)
+		{
+			if (Stream.BaseDescription.Identifier == LocalClientStreamId)
+			{
+				Result = &Stream.BaseDescription.ReplicationMap;
+				return EBreakBehavior::Break;
+			}
+			return EBreakBehavior::Continue;
+		});
+		return Result
+			? *Result
+			: EmptyState;
+	}
+
+	void FStreamSynchronizer_LocalClient::OnPostStreamsChanged()
+	{
+		// Note that this fires even for non MU-related changes, e.g. due to some other system's Concert replication API usage.
+		OnServerStateChangedDelegate.Broadcast();
+	}
+}
+
+#undef LOCTEXT_NAMESPACE
\ No newline at end of file
diff --git a/Engine/Plugins/Developer/Concert/ConcertApp/MultiUserClient/Source/MultiUserClient/Private/Replication/Stream/StreamSynchronizer_LocalClient.h b/Engine/Plugins/Developer/Concert/ConcertApp/MultiUserClient/Source/MultiUserClient/Private/Replication/Stream/StreamSynchronizer_LocalClient.h
new file mode 100644
--- /dev/null
+++ b/Engine/Plugins/Developer/Concert/ConcertApp/MultiUserClient/Source/MultiUserClient/Private/Replication/Stream/StreamSynchronizer_LocalClient.h
@@ -0,0 +1,43 @@
+// Copyright Epic Games, Inc. All Rights Reserved.
+
+#pragma once
+
+#include "IClientStreamSynchronizer.h"
+#include "Replication/Data/ObjectReplicationMap.h"
+#include "Templates/UnrealTemplate.h"
+
+class IConcertSyncClient;
+
+namespace UE::MultiUserClient
+{
+	/** Asks the local client's IConcertClientReplicationManager about the current state of the streams. */
+	class FStreamSynchronizer_LocalClient : public IClientStreamSynchronizer, public FNoncopyable
+	{
+	public:
+
+		FStreamSynchronizer_LocalClient(TSharedRef<IConcertSyncClient> InLocalClient, const FGuid& InLocalClientStreamId);
+		virtual ~FStreamSynchronizer_LocalClient() override;
+
+		//~ Begin IClientStreamSynchronizer Interface
+		virtual FGuid GetStreamId() const override { return LocalClientStreamId; }
+		virtual const FObjectReplicationMap& GetServerState() const override;
+		virtual FOnServerStateChanged& OnServerStateChanged() override { return OnServerStateChangedDelegate; }
+		//~ End IClientStreamSynchronizer Interface
+
+	private:
+		
+		/** Returned by GetServerState when there is no registered stream. */
+		const FObjectReplicationMap EmptyState; 
+		
+		/** Owning client. Used to send change requests to the server. */
+		const TSharedRef<IConcertSyncClient> LocalClient;
+		/** The ID of the local client's stream this FLocalClientStreamDiffer is managing. */
+		const FGuid LocalClientStreamId;
+		
+		/** Event executed when the result of GetServerState has been synched. */
+		FOnServerStateChanged OnServerStateChangedDelegate;
+		
+		void OnPostStreamsChanged();
+	};
+}
+
diff --git a/Engine/Plugins/Developer/Concert/ConcertApp/MultiUserClient/Source/MultiUserClient/Private/Replication/Stream/StreamSynchronizer_RemoteClient.cpp b/Engine/Plugins/Developer/Concert/ConcertApp/MultiUserClient/Source/MultiUserClient/Private/Replication/Stream/StreamSynchronizer_RemoteClient.cpp
new file mode 100644
--- /dev/null
+++ b/Engine/Plugins/Developer/Concert/ConcertApp/MultiUserClient/Source/MultiUserClient/Private/Replication/Stream/StreamSynchronizer_RemoteClient.cpp
@@ -0,0 +1,55 @@
+// Copyright Epic Games, Inc. All Rights Reserved.
+
+#include "StreamSynchronizer_RemoteClient.h"
+
+#include "ConcertLogGlobal.h"
+#include "Assets/MultiUserReplicationClientPreset.h"
+#include "Replication/Util/RegularQueryService.h"
+
+namespace UE::MultiUserClient
+{
+	constexpr float QueryTimeInterval = 1.f;
+	
+	FStreamSynchronizer_RemoteClient::FStreamSynchronizer_RemoteClient(const FGuid& RemoteEndpointId, FRegularQueryService& InQueryService)
+		: QueryService(InQueryService)
+		, QueryStreamHandle(
+			QueryService.RegisterStreamQuery(
+				RemoteEndpointId,
+				FStreamQueryDelegate::CreateRaw(this, &FStreamSynchronizer_RemoteClient::HandleStreamQuery)
+				)
+			)
+	{}
+
+	FStreamSynchronizer_RemoteClient::~FStreamSynchronizer_RemoteClient()
+	{
+		QueryService.UnregisterStreamQuery(QueryStreamHandle);
+	}
+
+	FGuid FStreamSynchronizer_RemoteClient::GetStreamId() const
+	{
+		return UMultiUserReplicationClientPreset::MultiUserStreamID;
+	}
+
+	void FStreamSynchronizer_RemoteClient::HandleStreamQuery(const TArray<FSharedReplicationStreamDescription>& Streams)
+	{
+		// MU clients use UMultiUserReplicationClientPreset::MultiUserStreamID for streams.
+		// That handles the (unlikely) case in which some external logic had added streams to the same client which we must differentiate
+		const FSharedReplicationStreamDescription* MultiUserDescription = Streams.FindByPredicate([](const FSharedReplicationStreamDescription& Description)
+		{
+			return Description.Identifier == UMultiUserReplicationClientPreset::MultiUserStreamID;
+		});
+		if (MultiUserDescription
+			// Avoid unnecessary Broadcast()s
+			&& LastKnownServerState != MultiUserDescription->ReplicationMap)
+		{
+			LastKnownServerState = MultiUserDescription->ReplicationMap;
+			OnServerStateChangedDelegate.Broadcast();
+		}
+		// If the remote client removed the last object, the stream is implicitly deleted.
+		else if (!MultiUserDescription && !LastKnownServerState.ReplicatedObjects.IsEmpty())
+		{
+			LastKnownServerState = {};
+			OnServerStateChangedDelegate.Broadcast();
+		}
+	}
+}
diff --git a/Engine/Plugins/Developer/Concert/ConcertApp/MultiUserClient/Source/MultiUserClient/Private/Replication/Stream/StreamSynchronizer_RemoteClient.h b/Engine/Plugins/Developer/Concert/ConcertApp/MultiUserClient/Source/MultiUserClient/Private/Replication/Stream/StreamSynchronizer_RemoteClient.h
new file mode 100644
--- /dev/null
+++ b/Engine/Plugins/Developer/Concert/ConcertApp/MultiUserClient/Source/MultiUserClient/Private/Replication/Stream/StreamSynchronizer_RemoteClient.h
@@ -0,0 +1,47 @@
+// Copyright Epic Games, Inc. All Rights Reserved.
+
+#pragma once
+
+#include "IClientStreamSynchronizer.h"
+#include "Replication/Data/ObjectReplicationMap.h"
+#include "Replication/Data/ReplicationStreamDescription.h"
+
+#include "Containers/Array.h"
+#include "Templates/UnrealTemplate.h"
+
+namespace UE::MultiUserClient
+{
+	class FRegularQueryService;
+	
+	/** Tracks the state of a remote client by querying the client's state in regular intervals. */
+	class FStreamSynchronizer_RemoteClient : public IClientStreamSynchronizer, public FNoncopyable
+	{
+	public:
+		
+		FStreamSynchronizer_RemoteClient(const FGuid& RemoteEndpointId, FRegularQueryService& InQueryService);
+		virtual ~FStreamSynchronizer_RemoteClient() override;
+
+		//~ Begin IClientStreamSynchronizer Interface
+		virtual FGuid GetStreamId() const override;
+		virtual const FObjectReplicationMap& GetServerState() const override { return LastKnownServerState; }
+		virtual FOnServerStateChanged& OnServerStateChanged() override { return OnServerStateChangedDelegate; }
+		//~ End IClientStreamSynchronizer Interface
+
+	private:
+
+		/** Queries the server in regular intervals. This services outlives our object. */
+		FRegularQueryService& QueryService;
+		/** Used to unregister HandleStreamQuery upon destruction. */
+		const FDelegateHandle QueryStreamHandle; 
+
+		/** Represents what the local client thinks the replication map on the server currently looks like. */
+		FObjectReplicationMap LastKnownServerState;
+		
+		/** Event executed when the result of GetServerState has been synched. */
+		FOnServerStateChanged OnServerStateChangedDelegate;
+
+		/** Called in regular intervals with the contents of the remote client. */
+		void HandleStreamQuery(const TArray<FSharedReplicationStreamDescription>& Streams);
+	};
+}
+
diff --git a/Engine/Plugins/Developer/Concert/ConcertApp/MultiUserClient/Source/MultiUserClient/Private/Replication/Submission/Data/AuhtoritySubmission.cpp b/Engine/Plugins/Developer/Concert/ConcertApp/MultiUserClient/Source/MultiUserClient/Private/Replication/Submission/Data/AuhtoritySubmission.cpp
new file mode 100644
--- /dev/null
+++ b/Engine/Plugins/Developer/Concert/ConcertApp/MultiUserClient/Source/MultiUserClient/Private/Replication/Submission/Data/AuhtoritySubmission.cpp
@@ -0,0 +1,25 @@
+// Copyright Epic Games, Inc. All Rights Reserved.
+
+#include "Replication/Submission/Data/AuthoritySubmission.h"
+
+#define LOCTEXT_NAMESPACE "EAuthoritySubmissionErrorCode"
+
+namespace UE::MultiUserClient
+{
+	FText LexToText(EAuthoritySubmissionErrorCode ErrorCode)
+	{
+		switch (ErrorCode)
+		{
+		case EAuthoritySubmissionErrorCode::Success: return LOCTEXT("Success", "Success");
+		case EAuthoritySubmissionErrorCode::NoChange: return LOCTEXT("NoChange", "No Change");
+		case EAuthoritySubmissionErrorCode::Timeout: return LOCTEXT("Timeout", "Timeout");
+		case EAuthoritySubmissionErrorCode::CancelledDueToStreamUpdate: return LOCTEXT("CancelledDueToStreamUpdate", "Failed because stream dependency could not be updated");
+		case EAuthoritySubmissionErrorCode::Cancelled: return LOCTEXT("Cancelled", "Cancelled");
+		default:
+			checkNoEntry();
+			return LOCTEXT("Unknown", "Unknown");;
+		}
+	}
+}
+
+#undef LOCTEXT_NAMESPACE
\ No newline at end of file
diff --git a/Engine/Plugins/Developer/Concert/ConcertApp/MultiUserClient/Source/MultiUserClient/Private/Replication/Submission/Data/AuthoritySubmission.h b/Engine/Plugins/Developer/Concert/ConcertApp/MultiUserClient/Source/MultiUserClient/Private/Replication/Submission/Data/AuthoritySubmission.h
new file mode 100644
--- /dev/null
+++ b/Engine/Plugins/Developer/Concert/ConcertApp/MultiUserClient/Source/MultiUserClient/Private/Replication/Submission/Data/AuthoritySubmission.h
@@ -0,0 +1,52 @@
+// Copyright Epic Games, Inc. All Rights Reserved.
+
+#pragma once
+
+#include "Replication/IConcertClientReplicationManager.h"
+
+namespace UE::MultiUserClient
+{
+	class ISubmissionOperation;
+	
+	enum class EAuthoritySubmissionErrorCode
+	{
+		/** The operation completed as expected */
+		Success,
+
+		/** No change was sent to the server because there were no local changes. */
+		NoChange, 
+
+		/** The request timed out */
+		Timeout,
+
+		/** Cancelled because the required stream update could not be made. */
+		CancelledDueToStreamUpdate,
+		
+		/** The request was cancelled (e.g. because client disconnected while operation was in progress) */
+		Cancelled
+	};
+
+	FText LexToText(EAuthoritySubmissionErrorCode ErrorCode);
+
+	struct FSubmitAuthorityChangesRequest
+	{
+		EAuthoritySubmissionErrorCode ErrorCode;
+		
+		/** The operation this was executed as part of. You cannot keep any reference to this. */
+		ISubmissionOperation& OperationContext;
+		
+		/** Only valid if ErrorCode == EAuthoritySubmissionErrorCode::Success */
+		TOptional<ConcertSyncClient::Replication::FAuthorityChangeRequest> Request;
+	};
+	
+	struct FSubmitAuthorityChangesResponse
+	{
+		EAuthoritySubmissionErrorCode ErrorCode;
+		
+		/** The operation this was executed as part of. You cannot keep any reference to this. */
+		ISubmissionOperation& OperationContext;
+		
+		/** Only valid if ErrorCode == EAuthoritySubmissionErrorCode::Success */
+		TOptional<ConcertSyncClient::Replication::FAuthorityChangeResponse> Response;
+	};
+}
diff --git a/Engine/Plugins/Developer/Concert/ConcertApp/MultiUserClient/Source/MultiUserClient/Private/Replication/Submission/Data/EChangeRevertability.h b/Engine/Plugins/Developer/Concert/ConcertApp/MultiUserClient/Source/MultiUserClient/Private/Replication/Submission/Data/EChangeRevertability.h
new file mode 100644
--- /dev/null
+++ b/Engine/Plugins/Developer/Concert/ConcertApp/MultiUserClient/Source/MultiUserClient/Private/Replication/Submission/Data/EChangeRevertability.h
@@ -0,0 +1,17 @@
+// Copyright Epic Games, Inc. All Rights Reserved.
+
+#pragma once
+
+namespace UE::MultiUserClient
+{
+	enum class EChangeRevertability
+	{
+		/** Can click the revert button */
+		Revertable,
+
+		/** There are no changes to upload */
+		NoChanges,
+		/** An upload operation is in progress */
+		UploadInProgress,
+	};
+}
\ No newline at end of file
diff --git a/Engine/Plugins/Developer/Concert/ConcertApp/MultiUserClient/Source/MultiUserClient/Private/Replication/Submission/Data/EChangeUploadability.h b/Engine/Plugins/Developer/Concert/ConcertApp/MultiUserClient/Source/MultiUserClient/Private/Replication/Submission/Data/EChangeUploadability.h
new file mode 100644
--- /dev/null
+++ b/Engine/Plugins/Developer/Concert/ConcertApp/MultiUserClient/Source/MultiUserClient/Private/Replication/Submission/Data/EChangeUploadability.h
@@ -0,0 +1,21 @@
+// Copyright Epic Games, Inc. All Rights Reserved.
+
+#pragma once
+
+namespace UE::MultiUserClient
+{
+	enum class EChangeUploadability
+	{
+		/** Can click the upload button */
+		Ready,
+
+		/** There are no changes to upload */
+		NoChanges,
+		/** A previous upload operation is in progress */
+		InProgress,
+		
+		// TODO DP UE-197435: Remove once remote client streams can be changed
+		// TODO DP UE-198088: Remove once remote changing is implemented
+		NotImplemented
+	};
+}
\ No newline at end of file
diff --git a/Engine/Plugins/Developer/Concert/ConcertApp/MultiUserClient/Source/MultiUserClient/Private/Replication/Submission/Data/StreamSubmission.h b/Engine/Plugins/Developer/Concert/ConcertApp/MultiUserClient/Source/MultiUserClient/Private/Replication/Submission/Data/StreamSubmission.h
new file mode 100644
--- /dev/null
+++ b/Engine/Plugins/Developer/Concert/ConcertApp/MultiUserClient/Source/MultiUserClient/Private/Replication/Submission/Data/StreamSubmission.h
@@ -0,0 +1,42 @@
+// Copyright Epic Games, Inc. All Rights Reserved.
+
+#pragma once
+
+#include "Replication/IConcertClientReplicationManager.h"
+
+namespace UE::MultiUserClient
+{
+	class ISubmissionOperation;
+
+	enum class EStreamSubmissionErrorCode
+	{
+		/** Changes were submitted to and processed by the server */
+		Success,
+
+		/** No change was sent to the server because there were no local changes. */
+		NoChange,
+
+		/** The request was cancelled (e.g. because client disconnected while operation was in progress) */
+		Cancelled
+	};
+
+	/** Contains the original request and response of a completed stream change. */
+	struct FCompletedChangeSubmission
+	{
+		ConcertSyncClient::Replication::FChangeStreamRequest Request;
+		ConcertSyncClient::Replication::FChangeStreamResponse Response;
+	};
+
+	/** Resulting of submitting stream changes */
+	struct FSubmitStreamChangesResponse
+	{
+		/** Error code of the submission. Determines whether Response is valid. */
+		EStreamSubmissionErrorCode ErrorCode;
+
+		/** The operation this was executed as part of. You cannot keep any reference to this. */
+		ISubmissionOperation& OperationContext; 
+
+		/** Valid if ErrorCode != ESubmitChangesErrorCode::AlreadyInProgress. */
+		TOptional<FCompletedChangeSubmission> SubmissionInfo;
+	};
+}
diff --git a/Engine/Plugins/Developer/Concert/ConcertApp/MultiUserClient/Source/MultiUserClient/Private/Replication/Submission/ISubmissionOperation.h b/Engine/Plugins/Developer/Concert/ConcertApp/MultiUserClient/Source/MultiUserClient/Private/Replication/Submission/ISubmissionOperation.h
new file mode 100644
--- /dev/null
+++ b/Engine/Plugins/Developer/Concert/ConcertApp/MultiUserClient/Source/MultiUserClient/Private/Replication/Submission/ISubmissionOperation.h
@@ -0,0 +1,42 @@
+// Copyright Epic Games, Inc. All Rights Reserved.
+
+#pragma once
+
+#include "Async/Future.h"
+#include "Data/AuthoritySubmission.h"
+#include "Data/StreamSubmission.h"
+
+namespace UE::MultiUserClient
+{
+	/**
+	 * Exposes the stages of submitting to the server.
+	 * Every time ISubmissionWorkflow::SubmitChanges is called a new instance is created.
+	 */
+	class ISubmissionOperation
+	{
+	public:
+
+		/** @return Whether a request for changing streams has been or will be sent. */
+		virtual bool IsModifyingStreams() const = 0;
+		
+		/**
+		 * Completes when the operation of changing streams has completed.
+		 * @note This can be called at most once; subsequent calls result in an unset future.
+		 */
+		virtual TFuture<FSubmitStreamChangesResponse> OnStreamChangesSubmittedFuture() = 0;
+
+		/**
+		 * Completes when the authority change request has been sent to the server.
+		 * @note This can be called at most once; subsequent calls result in an unset future.
+		 */
+		virtual TFuture<FSubmitAuthorityChangesRequest> OnAuthorityChangeRequestedFuture() = 0;
+		
+		/**
+		 * Completes when the operation of changing authority has completed.
+		 * @note This can be called at most once; subsequent calls result in an unset future.
+		 */
+		virtual TFuture<FSubmitAuthorityChangesResponse> OnAuthorityChangeResponseReceivedFuture() = 0;
+		
+		virtual ~ISubmissionOperation() = default;
+	};
+}
diff --git a/Engine/Plugins/Developer/Concert/ConcertApp/MultiUserClient/Source/MultiUserClient/Private/Replication/Submission/ISubmissionWorkflow.h b/Engine/Plugins/Developer/Concert/ConcertApp/MultiUserClient/Source/MultiUserClient/Private/Replication/Submission/ISubmissionWorkflow.h
new file mode 100644
--- /dev/null
+++ b/Engine/Plugins/Developer/Concert/ConcertApp/MultiUserClient/Source/MultiUserClient/Private/Replication/Submission/ISubmissionWorkflow.h
@@ -0,0 +1,45 @@
+// Copyright Epic Games, Inc. All Rights Reserved.
+
+#pragma once
+
+#include "Data/EChangeRevertability.h"
+#include "Data/EChangeUploadability.h"
+
+namespace UE::MultiUserClient
+{
+	class ISubmissionOperation;
+	struct FSubmitStreamChangesResponse;
+	
+	/**
+	 * Manages the flow of changing stream and authority on the server.
+	 * Handles both submitting and reverting.
+	 *
+	 * The general flow is as follows:
+	 *	1. Submit stream changes
+	 *	2. Submit authority changes
+	 */
+	class ISubmissionWorkflow
+	{
+	public:
+
+		/**
+		 * Synchronizes the server with the locally made changes streams and authority.
+		 * @return Operation object which emits special events. You cannot keep any reference to it. Null if there is already an operation in progress.
+		 * @see CanSubmit
+		 */
+		virtual ISubmissionOperation* SubmitChanges() = 0;
+
+		/** Reverts all local changes */
+		virtual void RevertChanges() = 0;
+
+		/** @return Detailed information about whether Submit can be called */
+		virtual EChangeUploadability GetUploadability() const = 0;
+		bool CanSubmit() const { return GetUploadability() == EChangeUploadability::Ready; }
+
+		/** @return Detailed information about whether Revert can be called */
+		virtual EChangeRevertability GetRevertability() const = 0;
+		bool CanRevert() const { return GetRevertability() == EChangeRevertability::Revertable; }
+
+		virtual ~ISubmissionWorkflow() = default;
+	};
+}
diff --git a/Engine/Plugins/Developer/Concert/ConcertApp/MultiUserClient/Source/MultiUserClient/Private/Replication/Submission/SingleClientSubmissionOperation.h b/Engine/Plugins/Developer/Concert/ConcertApp/MultiUserClient/Source/MultiUserClient/Private/Replication/Submission/SingleClientSubmissionOperation.h
new file mode 100644
--- /dev/null
+++ b/Engine/Plugins/Developer/Concert/ConcertApp/MultiUserClient/Source/MultiUserClient/Private/Replication/Submission/SingleClientSubmissionOperation.h
@@ -0,0 +1,75 @@
+// Copyright Epic Games, Inc. All Rights Reserved.
+
+#pragma once
+
+#include "ISubmissionOperation.h"
+#include "Replication/Stream/IClientStreamSynchronizer.h"
+
+#include "Templates/UnrealTemplate.h"
+
+namespace UE::MultiUserClient
+{
+	/** Keeps a promise for each relevant event. Handles cancelling them when destroyed. */
+	class FSingleClientSubmissionOperation
+		: public ISubmissionOperation
+		, public FNoncopyable
+	{
+	public:
+
+		FSingleClientSubmissionOperation(bool bInModifiesStreams)
+			: bModifiesStreams(bInModifiesStreams)
+		{}
+		
+		virtual ~FSingleClientSubmissionOperation() override
+		{
+			if (!bStreamPromiseWasSet)
+			{
+				StreamChangesPromise.EmplaceValue(FSubmitStreamChangesResponse{ EStreamSubmissionErrorCode::Cancelled, *this });
+			}
+			if (!bAuthorityRequestPromiseWasSet)
+			{
+				AuthorityChangeRequestPromise.EmplaceValue(FSubmitAuthorityChangesRequest{ EAuthoritySubmissionErrorCode::Cancelled, *this });
+			}
+			if (!bAuthorityResponsePromiseWasSet)
+			{
+				AuthorityChangeResponsePromise.EmplaceValue(FSubmitAuthorityChangesResponse{ EAuthoritySubmissionErrorCode::Cancelled, *this });
+			}
+		}
+
+		void EmplaceStreamPromise(FSubmitStreamChangesResponse&& Result)
+		{
+			bStreamPromiseWasSet = true;
+			StreamChangesPromise.EmplaceValue(MoveTemp(Result));
+		}
+		void EmplaceAuthorityRequestPromise(FSubmitAuthorityChangesRequest&& Result)
+		{
+			bAuthorityRequestPromiseWasSet = true;
+			AuthorityChangeRequestPromise.EmplaceValue(MoveTemp(Result));
+		}
+		void EmplaceAuthorityResponsePromise(FSubmitAuthorityChangesResponse&& Result)
+		{
+			bAuthorityResponsePromiseWasSet = true;
+			AuthorityChangeResponsePromise.EmplaceValue(MoveTemp(Result));
+		}
+
+		//~ Begin ISubmissionOperation Interface
+		virtual bool IsModifyingStreams() const override { return bModifiesStreams; }
+		virtual TFuture<FSubmitStreamChangesResponse> OnStreamChangesSubmittedFuture() override { return StreamChangesPromise.GetFuture();  }
+		virtual TFuture<FSubmitAuthorityChangesRequest> OnAuthorityChangeRequestedFuture() override { return AuthorityChangeRequestPromise.GetFuture(); }
+		virtual TFuture<FSubmitAuthorityChangesResponse> OnAuthorityChangeResponseReceivedFuture() override { return AuthorityChangeResponsePromise.GetFuture(); }
+		//~ End ISubmissionOperation Interface
+
+	private:
+
+		const bool bModifiesStreams;
+
+		bool bStreamPromiseWasSet = false;
+		bool bAuthorityRequestPromiseWasSet = false;
+		bool bAuthorityResponsePromiseWasSet = false;
+		
+		// All fulfilled by the owning FSingleClientSubmissionWorkflow
+		TPromise<FSubmitStreamChangesResponse> StreamChangesPromise;
+		TPromise<FSubmitAuthorityChangesRequest> AuthorityChangeRequestPromise;
+		TPromise<FSubmitAuthorityChangesResponse> AuthorityChangeResponsePromise;
+	};
+}
diff --git a/Engine/Plugins/Developer/Concert/ConcertApp/MultiUserClient/Source/MultiUserClient/Private/Replication/Submission/SubmissionWorkflow_LocalClient.cpp b/Engine/Plugins/Developer/Concert/ConcertApp/MultiUserClient/Source/MultiUserClient/Private/Replication/Submission/SubmissionWorkflow_LocalClient.cpp
new file mode 100644
--- /dev/null
+++ b/Engine/Plugins/Developer/Concert/ConcertApp/MultiUserClient/Source/MultiUserClient/Private/Replication/Submission/SubmissionWorkflow_LocalClient.cpp
@@ -0,0 +1,155 @@
+// Copyright Epic Games, Inc. All Rights Reserved.
+
+#include "SubmissionWorkflow_LocalClient.h"
+
+#include "IConcertSyncClient.h"
+#include "Replication/Authority/AuthorityChangeTracker.h"
+#include "Replication/Stream/StreamChangeTracker.h"
+#include "Replication/Util/StreamRequestUtils.h"
+
+namespace UE::MultiUserClient
+{
+	FSubmissionWorkflow_LocalClient::FSubmissionWorkflow_LocalClient(
+		TSharedRef<IConcertSyncClient> InClient,
+		FStreamChangeTracker& InStreamChangeTracker,
+		FAuthorityChangeTracker& InAuthorityChangeTracker,
+		IClientStreamSynchronizer& InStreamSynchronizer
+		)
+		: Client(MoveTemp(InClient))
+		, StreamChangeTracker(InStreamChangeTracker)
+		, AuthorityChangeTracker(InAuthorityChangeTracker)
+		, StreamSynchronizer(InStreamSynchronizer)
+	{}
+
+	void FSubmissionWorkflow_LocalClient::RevertChanges()
+	{
+		StreamChangeTracker.RevertCachedChanges();
+		AuthorityChangeTracker.ClearChanges();
+	}
+
+	EChangeUploadability FSubmissionWorkflow_LocalClient::GetUploadability() const
+	{
+		const bool bOperationInProgress = InProgressOperation.IsSet(); 
+		if (bOperationInProgress)
+		{
+			return EChangeUploadability::InProgress;
+		}
+
+		const bool bHasChanges = StreamChangeTracker.HasChanges() || AuthorityChangeTracker.HasChanges();
+		return bHasChanges ? EChangeUploadability::Ready : EChangeUploadability::NoChanges;
+	}
+
+	EChangeRevertability FSubmissionWorkflow_LocalClient::GetRevertability() const
+	{
+		if (GetUploadability() == EChangeUploadability::InProgress)
+		{
+			return EChangeRevertability::UploadInProgress;
+		}
+
+		const bool bStreamHasChanges = StreamChangeTracker.HasChanges();
+		const bool bAuthorityHasChanges = AuthorityChangeTracker.HasChanges();
+		const bool bHasChanges = bStreamHasChanges || bAuthorityHasChanges;
+		return bHasChanges ? EChangeRevertability::Revertable : EChangeRevertability::NoChanges;
+	}
+
+	ISubmissionOperation* FSubmissionWorkflow_LocalClient::SubmitChanges()
+	{
+		using namespace ConcertSyncClient::Replication;
+		IConcertClientReplicationManager* ReplicationManager = Client->GetReplicationManager();
+		if (!CanSubmit() || !ensure(ReplicationManager))
+		{
+			return nullptr;
+		}
+
+		// TODO DP: We should check the authority change for conflicts one more time here, in case the remote client's streams have changed since the user last edited authority
+		// The authority request is pre-built now to avoid sending changes the local client makes while we're waiting for the latent server responses
+		FAuthorityChangeRequest AuthorityChangeRequest = AuthorityChangeTracker.BuildChangeRequest(GetLocalClientStreamId());
+		
+		const FStreamChangelist& Changelist = StreamChangeTracker.GetCachedDeltaChange();
+		const bool bIsChangelistEmpty = Changelist.ObjectsToPut.IsEmpty() && Changelist.ObjectsToRemove.IsEmpty();
+		const bool bModifyStreams = !bIsChangelistEmpty;
+		InProgressOperation.Emplace(bModifyStreams);
+		
+		if (bIsChangelistEmpty)
+		{
+			InProgressOperation->EmplaceStreamPromise(FSubmitStreamChangesResponse{ EStreamSubmissionErrorCode::NoChange, *InProgressOperation });
+			SendAuthorityChangeRequest(MoveTemp(AuthorityChangeRequest));
+		}
+		else
+		{
+			FChangeStreamRequest Request = StreamSynchronizer.GetServerState().ReplicatedObjects.IsEmpty()
+				? StreamRequestUtils::BuildChangeRequest_CreateNewStream(GetLocalClientStreamId(), Changelist)
+				: StreamRequestUtils::BuildChangeRequest_UpdateExistingStream(Changelist);
+		
+			ReplicationManager->ChangeStream(Request)
+				.Next([this, DestructionDetection = LifetimeToken->AsWeak(), Request, AuthorityChangeRequest = MoveTemp(AuthorityChangeRequest)](FChangeStreamResponse&& Response)
+				{
+					const FSubmitStreamChangesResponse SubmissionResult { EStreamSubmissionErrorCode::Success, *InProgressOperation, { FCompletedChangeSubmission{Request, Response } } };
+					// The request might execute after we're destroyed, e.g. by leaving session while request is on the way.
+					// In that case, the Concert session triggers the OnSessionConnectionChanged which destroys us. Only after that, all the requests are timed out.
+					if (DestructionDetection.IsValid())
+					{
+						OnStreamChangeCompleted(Request, Response, AuthorityChangeRequest);
+					}
+
+					return SubmissionResult;
+				});
+		}
+		
+		return InProgressOperation.GetPtrOrNull();
+	}
+
+	void FSubmissionWorkflow_LocalClient::OnStreamChangeCompleted(
+		const ConcertSyncClient::Replication::FChangeStreamRequest& Request,
+		const ConcertSyncClient::Replication::FChangeStreamResponse& ChangeStreamResponse,
+		ConcertSyncClient::Replication::FAuthorityChangeRequest AuthorityChangeRequest
+		)
+	{
+		InProgressOperation->EmplaceStreamPromise(FSubmitStreamChangesResponse{
+			EStreamSubmissionErrorCode::Success,
+			*InProgressOperation,
+			{ FCompletedChangeSubmission{ Request, ChangeStreamResponse } }
+		});
+		
+		if (!ChangeStreamResponse.IsSuccess())
+		{
+			InProgressOperation->EmplaceAuthorityRequestPromise(FSubmitAuthorityChangesRequest{ EAuthoritySubmissionErrorCode::CancelledDueToStreamUpdate, *InProgressOperation });
+			InProgressOperation->EmplaceAuthorityResponsePromise(FSubmitAuthorityChangesResponse{ EAuthoritySubmissionErrorCode::CancelledDueToStreamUpdate, *InProgressOperation });
+			InProgressOperation.Reset();
+		}
+		else
+		{
+			SendAuthorityChangeRequest(MoveTemp(AuthorityChangeRequest));
+		}
+	}
+
+	void FSubmissionWorkflow_LocalClient::SendAuthorityChangeRequest(ConcertSyncClient::Replication::FAuthorityChangeRequest AuthorityChangeRequest)
+	{
+		using namespace ConcertSyncClient::Replication;
+		IConcertClientReplicationManager* ReplicationManager = Client->GetReplicationManager();
+		if (!ensure(ReplicationManager))
+		{
+			InProgressOperation.Reset(); // Automatically cancels the pending promises
+			return;
+		}
+
+		const bool bHasNoChanges = AuthorityChangeRequest.ReleaseAuthority.IsEmpty() && AuthorityChangeRequest.TakeAuthority.IsEmpty();
+		if (bHasNoChanges)
+		{
+			InProgressOperation->EmplaceAuthorityRequestPromise(FSubmitAuthorityChangesRequest{ EAuthoritySubmissionErrorCode::NoChange, *InProgressOperation });
+			InProgressOperation->EmplaceAuthorityResponsePromise(FSubmitAuthorityChangesResponse{ EAuthoritySubmissionErrorCode::NoChange, *InProgressOperation });
+			return;
+		}
+		
+		InProgressOperation->EmplaceAuthorityRequestPromise(FSubmitAuthorityChangesRequest{ EAuthoritySubmissionErrorCode::Success, *InProgressOperation, AuthorityChangeRequest });
+		ReplicationManager->RequestAuthorityChange(MoveTemp(AuthorityChangeRequest))
+			.Next([this, DestructionDetection = LifetimeToken->AsWeak()](FAuthorityChangeResponse&& Response)
+			{
+				if (DestructionDetection.IsValid())
+				{
+					InProgressOperation->EmplaceAuthorityResponsePromise(FSubmitAuthorityChangesResponse{ EAuthoritySubmissionErrorCode::Success, *InProgressOperation, MoveTemp(Response) });
+					InProgressOperation.Reset();
+				}
+			});
+	}
+}
diff --git a/Engine/Plugins/Developer/Concert/ConcertApp/MultiUserClient/Source/MultiUserClient/Private/Replication/Submission/SubmissionWorkflow_LocalClient.h b/Engine/Plugins/Developer/Concert/ConcertApp/MultiUserClient/Source/MultiUserClient/Private/Replication/Submission/SubmissionWorkflow_LocalClient.h
new file mode 100644
--- /dev/null
+++ b/Engine/Plugins/Developer/Concert/ConcertApp/MultiUserClient/Source/MultiUserClient/Private/Replication/Submission/SubmissionWorkflow_LocalClient.h
@@ -0,0 +1,71 @@
+// Copyright Epic Games, Inc. All Rights Reserved.
+
+#pragma once
+
+#include "ISubmissionWorkflow.h"
+#include "Replication/Stream/IClientStreamSynchronizer.h"
+#include "Replication/IToken.h"
+#include "SingleClientSubmissionOperation.h"
+
+class IConcertSyncClient;
+
+namespace UE::MultiUserClient
+{
+	class IClientAuthoritySynchronizer;
+	class IClientStreamSynchronizer;
+	class FAuthorityChangeTracker;
+	class FStreamChangeTracker;
+	
+	/** Handles the submission workflow for a single client. */
+	class FSubmissionWorkflow_LocalClient : public ISubmissionWorkflow
+	{
+	public:
+		
+		FSubmissionWorkflow_LocalClient(
+			TSharedRef<IConcertSyncClient> InClient,
+			FStreamChangeTracker& InStreamChangeTracker,
+			FAuthorityChangeTracker& InAuthorityChangeTracker,
+			IClientStreamSynchronizer& InStreamSynchronizer
+			);
+		
+		//~ Begin ISubmissionWorkflow Interface
+		virtual ISubmissionOperation* SubmitChanges() override;
+		virtual void RevertChanges() override;
+		virtual EChangeUploadability GetUploadability() const override;
+		virtual EChangeRevertability GetRevertability() const override;
+		//~ End ISubmissionWorkflow Interface
+	
+	private:
+
+		/** Latent operations keep a weak reference to this to detect whether this FSubmissionWorkflow_LocalClient was destroyed. */
+		const TSharedRef<FToken> LifetimeToken = FToken::Make();
+		
+		/** Used to send authority requests to the server. */
+		const TSharedRef<IConcertSyncClient> Client;
+
+		/** Used to get changes made to the stream */
+		FStreamChangeTracker& StreamChangeTracker;
+		/** Used to get changes made to the authority */
+		FAuthorityChangeTracker& AuthorityChangeTracker;
+
+		/** Used to change the streams */
+		IClientStreamSynchronizer& StreamSynchronizer;
+
+		/**
+		 * Set for as long as there is a SubmitChanges operation in progress.
+		 * Automatically cancels pending promises when destroyed.
+		 */
+		TOptional<FSingleClientSubmissionOperation> InProgressOperation;
+
+		FGuid GetLocalClientStreamId() const { return StreamSynchronizer.GetStreamId(); }
+
+		/** Advances the request by requesting authority. */
+		void OnStreamChangeCompleted(
+			const ConcertSyncClient::Replication::FChangeStreamRequest& Request,
+			const ConcertSyncClient::Replication::FChangeStreamResponse& ChangeStreamResponse,
+			ConcertSyncClient::Replication::FAuthorityChangeRequest AuthorityChangeRequest
+			);
+		void SendAuthorityChangeRequest(ConcertSyncClient::Replication::FAuthorityChangeRequest AuthorityChangeRequest);
+	};
+}
+
diff --git a/Engine/Plugins/Developer/Concert/ConcertApp/MultiUserClient/Source/MultiUserClient/Private/Replication/Submission/SubmissionWorkflow_RemoteClient.cpp b/Engine/Plugins/Developer/Concert/ConcertApp/MultiUserClient/Source/MultiUserClient/Private/Replication/Submission/SubmissionWorkflow_RemoteClient.cpp
new file mode 100644
--- /dev/null
+++ b/Engine/Plugins/Developer/Concert/ConcertApp/MultiUserClient/Source/MultiUserClient/Private/Replication/Submission/SubmissionWorkflow_RemoteClient.cpp
@@ -0,0 +1,30 @@
+// Copyright Epic Games, Inc. All Rights Reserved.
+
+#include "SubmissionWorkflow_RemoteClient.h"
+
+namespace UE::MultiUserClient
+{
+	ISubmissionOperation* FSubmissionWorkflow_RemoteClient::SubmitChanges()
+	{
+		// TODO DP UE-180657: Changing remote client's stream is not implemented for now
+		// TODO DP UE-198088 changing remote client authority
+		return nullptr;
+	}
+
+	void FSubmissionWorkflow_RemoteClient::RevertChanges()
+	{}
+
+	EChangeUploadability FSubmissionWorkflow_RemoteClient::GetUploadability() const
+	{
+		// TODO DP UE-180657: Changing remote client's stream is not implemented for now
+		// TODO DP UE-198088 changing remote client authority
+		return EChangeUploadability::NotImplemented;
+	}
+
+	EChangeRevertability FSubmissionWorkflow_RemoteClient::GetRevertability() const
+	{
+		// TODO DP UE-180657: Changing remote client's stream is not implemented for now
+		// TODO DP UE-198088 changing remote client authority
+		return EChangeRevertability::NoChanges;
+	}
+}
diff --git a/Engine/Plugins/Developer/Concert/ConcertApp/MultiUserClient/Source/MultiUserClient/Private/Replication/Submission/SubmissionWorkflow_RemoteClient.h b/Engine/Plugins/Developer/Concert/ConcertApp/MultiUserClient/Source/MultiUserClient/Private/Replication/Submission/SubmissionWorkflow_RemoteClient.h
new file mode 100644
--- /dev/null
+++ b/Engine/Plugins/Developer/Concert/ConcertApp/MultiUserClient/Source/MultiUserClient/Private/Replication/Submission/SubmissionWorkflow_RemoteClient.h
@@ -0,0 +1,26 @@
+// Copyright Epic Games, Inc. All Rights Reserved.
+
+#pragma once
+
+#include "ISubmissionWorkflow.h"
+#include "Templates/UnrealTemplate.h"
+
+class IConcertSyncClient;
+
+namespace UE::MultiUserClient
+{
+	class FSubmissionWorkflow_RemoteClient
+		: public ISubmissionWorkflow
+		, public FNoncopyable
+	{
+	public:
+
+		//~ Begin ISubmissionWorkflow Interface
+		virtual ISubmissionOperation* SubmitChanges() override;
+		virtual void RevertChanges() override;
+		virtual EChangeUploadability GetUploadability() const override;
+		virtual EChangeRevertability GetRevertability() const override;
+		//~ End ISubmissionWorkflow Interface
+	};
+}
+
diff --git a/Engine/Plugins/Developer/Concert/ConcertApp/MultiUserClient/Source/MultiUserClient/Private/Replication/Util/StreamRequestUtils.cpp b/Engine/Plugins/Developer/Concert/ConcertApp/MultiUserClient/Source/MultiUserClient/Private/Replication/Util/StreamRequestUtils.cpp
new file mode 100644
--- /dev/null
+++ b/Engine/Plugins/Developer/Concert/ConcertApp/MultiUserClient/Source/MultiUserClient/Private/Replication/Util/StreamRequestUtils.cpp
@@ -0,0 +1,42 @@
+// Copyright Epic Games, Inc. All Rights Reserved.
+
+#include "StreamRequestUtils.h"
+
+#include "Replication/Stream/IClientStreamSynchronizer.h"
+#include "Replication/Stream/StreamChangeTracker.h"
+
+namespace UE::MultiUserClient::StreamRequestUtils
+{
+	ConcertSyncClient::Replication::FChangeStreamRequest BuildChangeRequest_CreateNewStream(const FGuid& StreamId, const FStreamChangelist& FromChangelist)
+	{
+		const TMap<FObjectInStreamID, FConcertReplication_ChangeStream_PutObject>& ObjectsToPut = FromChangelist.ObjectsToPut;
+		
+		ConcertSyncClient::Replication::FChangeStreamRequest Request;
+		Request.StreamsToAdd.Emplace();
+		FReplicationStreamDescription_NetPacked& NewStream = Request.StreamsToAdd[0];
+		NewStream.BaseDescription.Identifier = StreamId;
+			
+		// If creating a new stream, the objects must be supplied in the description instead of in PutObjects!
+		FObjectReplicationMap& ReplicationMap = NewStream.BaseDescription.ReplicationMap;
+		ReplicationMap.ReplicatedObjects.Reserve(ObjectsToPut.Num());
+		for (const TPair<FObjectInStreamID, FConcertReplication_ChangeStream_PutObject>& PutObjectPair : ObjectsToPut)
+		{
+			const TOptional<FReplicatedObjectInfo> NewObjectInfo = PutObjectPair.Value.MakeObjectInfoIfValid();
+			// The editing UI allows adding objects without properties (to make UX easier) - do not submit those to the server.
+			if (!NewObjectInfo)
+			{
+				continue;
+			}
+			checkf(PutObjectPair.Key.StreamId == StreamId, TEXT("BuildChangeRequest should always use LocalClientStreamId ID!"));
+				
+			ReplicationMap.ReplicatedObjects.Add(PutObjectPair.Key.Object, *NewObjectInfo);
+		}
+		
+		return Request;
+	}
+		
+	ConcertSyncClient::Replication::FChangeStreamRequest BuildChangeRequest_UpdateExistingStream(FStreamChangelist FromChangelist)
+	{
+		return { MoveTemp(FromChangelist.ObjectsToRemove), MoveTemp(FromChangelist.ObjectsToPut) };
+	}
+}
diff --git a/Engine/Plugins/Developer/Concert/ConcertApp/MultiUserClient/Source/MultiUserClient/Private/Replication/Util/StreamRequestUtils.h b/Engine/Plugins/Developer/Concert/ConcertApp/MultiUserClient/Source/MultiUserClient/Private/Replication/Util/StreamRequestUtils.h
new file mode 100644
--- /dev/null
+++ b/Engine/Plugins/Developer/Concert/ConcertApp/MultiUserClient/Source/MultiUserClient/Private/Replication/Util/StreamRequestUtils.h
@@ -0,0 +1,31 @@
+// Copyright Epic Games, Inc. All Rights Reserved.
+
+#pragma once
+
+#include "Replication/IConcertClientReplicationManager.h"
+
+namespace UE::ConcertSyncClient::Replication
+{
+	struct FChangeStreamRequest;
+}
+
+namespace UE::MultiUserClient
+{
+	struct FStreamChangelist;
+}
+
+namespace UE::MultiUserClient::StreamRequestUtils
+{
+	/**
+	 * Builds a request for creating a new stream.
+	 * @param StreamId The stream that should be modified
+	 * @param FromChangelist The changes to be made
+	 */
+	ConcertSyncClient::Replication::FChangeStreamRequest BuildChangeRequest_CreateNewStream(const FGuid& StreamId, const FStreamChangelist& FromChangelist);
+	
+	/**
+	 * Builds a request for updating a preexisting stream.
+	 * @param FromChangelist The changes to be made
+	 */
+	ConcertSyncClient::Replication::FChangeStreamRequest BuildChangeRequest_UpdateExistingStream(FStreamChangelist FromChangelist);
+}
diff --git a/Engine/Plugins/Developer/Concert/ConcertApp/MultiUserClient/Source/MultiUserClient/Private/Widgets/ActiveSession/Replication/Client/SSubmissionControls.cpp b/Engine/Plugins/Developer/Concert/ConcertApp/MultiUserClient/Source/MultiUserClient/Private/Widgets/ActiveSession/Replication/Client/SSubmissionControls.cpp
new file mode 100644
--- /dev/null
+++ b/Engine/Plugins/Developer/Concert/ConcertApp/MultiUserClient/Source/MultiUserClient/Private/Widgets/ActiveSession/Replication/Client/SSubmissionControls.cpp
@@ -0,0 +1,307 @@
+// Copyright Epic Games, Inc. All Rights Reserved.
+
+#include "SSubmissionControls.h"
+
+#include "ConcertLogGlobal.h"
+#include "Replication/IConcertClientReplicationManager.h"
+#include "Replication/MultiUserReplicationManager.h"
+#include "Replication/Submission/ISubmissionWorkflow.h"
+#include "Replication/Submission/Data/EChangeRevertability.h"
+
+#include "Framework/Notifications/NotificationManager.h"
+#include "Misc/ScopeExit.h"
+#include "SNegativeActionButton.h"
+#include "Styling/AppStyle.h"
+#include "Styling/StyleColors.h"
+#include "Widgets/Images/SImage.h"
+#include "Widgets/Input/SButton.h"
+#include "Widgets/Layout/SBox.h"
+#include "Widgets/Layout/SSpacer.h"
+#include "Widgets/Notifications/SNotificationList.h"
+#include "Widgets/SBoxPanel.h"
+#include "Widgets/Images/SThrobber.h"
+#include "Widgets/Layout/SWidgetSwitcher.h"
+#include "Widgets/Text/STextBlock.h"
+
+#define LOCTEXT_NAMESPACE "SSubmissionControls"
+
+namespace UE::MultiUserClient
+{
+	void SSubmissionControls::Construct(const FArguments& InArgs)
+	{
+		SubmissionWorkflowAttribute = InArgs._SubmissionWorkflow;
+		check(SubmissionWorkflowAttribute.IsBound() || SubmissionWorkflowAttribute.IsSet());
+		
+		ChildSlot
+		[
+			SNew(SHorizontalBox)
+			
+			// Upload button
+			+SHorizontalBox::Slot()
+			.AutoWidth()
+			[
+				BuildUploadButton()
+			]
+
+			// Revert changes
+			+SHorizontalBox::Slot()
+			.AutoWidth()
+			[
+				BuildRevertButton()
+			]
+		];
+	}
+
+	TSharedRef<SWidget> SSubmissionControls::BuildUploadButton()
+	{
+		return SNew(SButton)
+			.ForegroundColor(FSlateColor::UseStyle())
+			.IsEnabled_Lambda([this](){ return GetSubmissionWorkflow().CanSubmit(); })
+			.ToolTipText(this, &SSubmissionControls::GetUploadButtonToolTipText)
+			.HAlign(HAlign_Center)
+			.VAlign(VAlign_Center)
+			.OnClicked(this, &SSubmissionControls::OnUploadButtonClicked)
+			[
+				SNew(SWidgetSwitcher)
+				.WidgetIndex_Lambda([this](){ return GetSubmissionWorkflow().GetUploadability() == EChangeUploadability::InProgress ? 0 : 1; })
+				
+				// InProgress
+				+SWidgetSwitcher::Slot()
+				[
+					SNew(SHorizontalBox)
+					+ SHorizontalBox::Slot()
+					.HAlign(HAlign_Center)
+					.VAlign(VAlign_Center)
+					.AutoWidth()
+					[
+						SNew(SCircularThrobber)
+					]
+					+ SHorizontalBox::Slot()
+					.Padding(FMargin(3, 0, 0, 0))
+					.VAlign(VAlign_Center)
+					.AutoWidth()
+					[
+						SNew(STextBlock)
+						.TextStyle(FAppStyle::Get(), "SmallButtonText")
+						.Text(LOCTEXT("UploadChanges.Label.Uploading", "Uploading"))
+					]
+				]
+				// Any other EChangeUploadability
+				+SWidgetSwitcher::Slot()
+				[
+					SNew(SHorizontalBox)
+					+SHorizontalBox::Slot()
+					.HAlign(HAlign_Center)
+					.VAlign(VAlign_Center)
+					.AutoWidth()
+					[
+						SNew(SImage)
+						.Image(FAppStyle::Get().GetBrush("Icons.Plus"))
+						.ColorAndOpacity(FStyleColors::AccentGreen)
+					]
+					+ SHorizontalBox::Slot()
+					.Padding(FMargin(3, 0, 0, 0))
+					.VAlign(VAlign_Center)
+					.AutoWidth()
+					[
+						SNew(STextBlock)
+						.TextStyle(FAppStyle::Get(), "SmallButtonText")
+						.Text(LOCTEXT("UploadChanges.Label.Upload", "Upload"))
+					]
+				]
+			];
+	}
+
+	FText SSubmissionControls::GetUploadButtonToolTipText() const
+	{
+		switch(GetSubmissionWorkflow().GetUploadability())
+		{
+			case EChangeUploadability::Ready: return LOCTEXT("UploadChanges.ToolTip.Upload", "Upload your local changes to the server.");
+			case EChangeUploadability::NoChanges: return LOCTEXT("UploadChanges.ToolTip.NoChanges", "No submittable changes.");
+			case EChangeUploadability::InProgress: return LOCTEXT("UploadChanges.ToolTip.InTransit", "Awaiting server response.");
+			case EChangeUploadability::NotImplemented: return LOCTEXT("UploadChanges.ToolTip.NotImplemented", "This feature is currently not implemented");
+			default: checkNoEntry(); return FText::GetEmpty();
+		}
+	}
+
+	FReply SSubmissionControls::OnUploadButtonClicked() const
+	{
+		ISubmissionWorkflow& Workflow = GetSubmissionWorkflow();
+		// SubmitChanges implicitly checks whether Workflow.CanSubmit()
+		if (ISubmissionOperation* SubmissionOperation = Workflow.SubmitChanges())
+		{
+			SetupNotificationsFor(*SubmissionOperation);
+		}
+		return FReply::Handled();
+	}
+
+	void SSubmissionControls::SetupNotificationsFor(ISubmissionOperation& SubmissionOperation) const
+	{
+		HandleStreamUpdatedNotification(SubmissionOperation);
+		HandleAuthorityChangeRequestedNotification(SubmissionOperation);
+	}
+
+	void SSubmissionControls::HandleStreamUpdatedNotification(ISubmissionOperation& SubmissionOperation) const
+	{
+		if (!SubmissionOperation.IsModifyingStreams())
+		{
+			return;
+		}
+		
+		FNotificationInfo NotificationInfo(LOCTEXT("Uploading.InProgress", "Uploading stream changes."));
+		NotificationInfo.bUseThrobber = true;
+		NotificationInfo.bUseSuccessFailIcons = true;
+		NotificationInfo.bFireAndForget = false;
+		
+		const TSharedPtr<SNotificationItem> Notification = FSlateNotificationManager::Get().AddNotification(NotificationInfo);
+		SubmissionOperation.OnStreamChangesSubmittedFuture()
+			.Next([Notification](FSubmitStreamChangesResponse&& SubmitChangesResult)
+			{
+				if (SubmitChangesResult.ErrorCode == EStreamSubmissionErrorCode::Cancelled)
+				{
+					Notification->SetCompletionState(SNotificationItem::CS_Fail);
+					Notification->SetText(LOCTEXT("Uploading.Failure", "Updating Stream cancelled."));
+					return;
+				}
+
+				const TOptional<FCompletedChangeSubmission>& SubmissionInfo = SubmitChangesResult.SubmissionInfo;
+				if (SubmissionInfo.IsSet() && SubmissionInfo->Response.IsSuccess())
+				{
+					Notification->SetCompletionState(SNotificationItem::CS_Success);
+					Notification->SetText(LOCTEXT("Uploading.Success", "Changes accepted."));
+
+					const ConcertSyncClient::Replication::FChangeStreamRequest& Request = SubmissionInfo->Request;
+					// Multi-User adds only one stream so we only need to look at index 0
+					const int32 NumAdditions = Request.StreamsToAdd.IsEmpty() ? 0 : Request.StreamsToAdd[0].BaseDescription.ReplicationMap.ReplicatedObjects.Num();
+					const int32 NumPuts = Request.ObjectsToPut.Num() + NumAdditions;
+					Notification->SetSubText(FText::Format(LOCTEXT("Uploading.Success.SubText", "{0} puts, {1} removals"),
+						NumPuts,
+						Request.ObjectsToRemove.Num()
+						));
+				}
+				else
+				{
+					Notification->SetCompletionState(SNotificationItem::CS_Fail);
+					Notification->SetText(LOCTEXT("Uploading.Failure", "Changes rejected."));
+					Notification->SetSubText(LOCTEXT("Uploading.Failure.SubText", "See log for details."));
+
+					if (SubmissionInfo.IsSet())
+					{
+						FStringOutputDevice Errors;
+						SubmissionInfo->Response.LogErrors(Errors);
+						UE_LOG(LogConcert, Error, TEXT("Errors uploading stream changes:\n%s"), *Errors);
+					}
+					else
+					{
+						UE_LOG(LogConcert, Error, TEXT("Failed to submit changes. ESubmitChangesErrorCode: %d"), SubmitChangesResult.ErrorCode);
+					}
+				}
+
+				Notification->ExpireAndFadeout();
+           });
+	}
+
+	void SSubmissionControls::HandleAuthorityChangeRequestedNotification(ISubmissionOperation& SubmissionOperation) const
+	{
+		SubmissionOperation.OnAuthorityChangeRequestedFuture()
+			.Next([WeakThis = TWeakPtr<const SSubmissionControls>(SharedThis(this))](FSubmitAuthorityChangesRequest&& RequestMetaData)
+			{
+				if (RequestMetaData.ErrorCode != EAuthoritySubmissionErrorCode::Success)
+				{
+					return;
+				}
+
+				FNotificationInfo NotificationInfo(LOCTEXT("Authority.InProgress.Text", "Requesting authority."));
+				NotificationInfo.bUseThrobber = true;
+				NotificationInfo.bUseSuccessFailIcons = true;
+				NotificationInfo.bFireAndForget = false;
+				const TSharedPtr<SNotificationItem> AuthorityChangeNotification = FSlateNotificationManager::Get().AddNotification(NotificationInfo);
+
+				const ConcertSyncClient::Replication::FAuthorityChangeRequest& Request = *RequestMetaData.Request ;
+				AuthorityChangeNotification->SetSubText(
+					FText::Format(LOCTEXT("Authority.InProgress.SubTextFmt", "Taking {0}|plural(one=object,other\nReleasing {1}|plural(one=object,other=objects)"),
+						Request.TakeAuthority.Num(),
+						Request.ReleaseAuthority.Num()
+					));
+				
+				RequestMetaData.OperationContext.OnAuthorityChangeResponseReceivedFuture()
+					.Next([Request, AuthorityChangeNotification](FSubmitAuthorityChangesResponse&& Response)
+					{
+						HandleAuthorityChangedResponseNotification(Request, Response, AuthorityChangeNotification.ToSharedRef());
+					});
+			});
+	}
+
+	void SSubmissionControls::HandleAuthorityChangedResponseNotification(
+		const ConcertSyncClient::Replication::FAuthorityChangeRequest& Request,
+		const FSubmitAuthorityChangesResponse& ResponseMetaData,
+		const TSharedRef<SNotificationItem>& NotificationItem
+		)
+	{
+		if (ResponseMetaData.ErrorCode == EAuthoritySubmissionErrorCode::NoChange)
+		{
+			return;
+		}
+		
+		ON_SCOPE_EXIT{ NotificationItem->ExpireAndFadeout(); };
+		if (ResponseMetaData.ErrorCode != EAuthoritySubmissionErrorCode::Success)
+		{
+			NotificationItem->SetCompletionState(SNotificationItem::CS_Fail);
+			NotificationItem->SetText(LOCTEXT("Authority.Failed.Text", "Authority request failed."));
+			NotificationItem->SetSubText(FText::Format(LOCTEXT("Auhtority.Failed.SubText", "Error code: {0}"), LexToText(ResponseMetaData.ErrorCode)));
+			return;
+		}
+		
+		const ConcertSyncClient::Replication::FAuthorityChangeResponse& Response = *ResponseMetaData.Response;
+		
+		const bool bSuccess = Response.RejectedObjects.IsEmpty();
+		const FText Text = bSuccess
+			? LOCTEXT("Authority.Text.Accepted", "Authority change accepted.")
+			: LOCTEXT("Authority.Text.Rejection", "Authority change had rejections.");
+		const FText SubText = bSuccess
+			? FText::GetEmpty()
+			: FText::Format(LOCTEXT("Authority.SubText.RejectionFmt", "{0} accepted {1} rejected\nSee logs for details."),
+				Request.TakeAuthority.Num() - Response.RejectedObjects.Num(),
+				Response.RejectedObjects.Num()
+				);
+				
+		NotificationItem->SetCompletionState(bSuccess ? SNotificationItem::CS_Success : SNotificationItem::CS_Fail);
+		NotificationItem->SetText(Text);
+		NotificationItem->SetSubText(SubText);
+
+		if (!bSuccess)
+		{
+			const FString Error = FString::JoinBy(Response.RejectedObjects, TEXT("\n"), [](const TPair<FSoftObjectPath, FConcertStreamArray>& Pair)
+			{
+				return Pair.Key.ToString();
+			});
+			UE_LOG(LogConcert, Error, TEXT("Authority change rejected objects:\n%s"), *Error);
+		}
+	}
+
+	TSharedRef<SWidget> SSubmissionControls::BuildRevertButton()
+	{
+		return SNew(SNegativeActionButton)
+			.Text(LOCTEXT("RevertChanges.Label", "Revert"))
+			.ToolTipText(this, &SSubmissionControls::GetRevertButtonToolTipText)
+			.IsEnabled_Lambda([this](){ return GetSubmissionWorkflow().CanRevert(); })
+			.OnClicked_Lambda([this]()
+			{
+				GetSubmissionWorkflow().RevertChanges();
+				return FReply::Handled();
+			});
+	}
+
+	FText SSubmissionControls::GetRevertButtonToolTipText() const
+	{
+		switch(GetSubmissionWorkflow().GetRevertability())
+		{
+		case EChangeRevertability::Revertable: return LOCTEXT("RevertChanges.ToolTip.Upload", "Revert all local changes you've made.");
+		case EChangeRevertability::NoChanges: return LOCTEXT("RevertChanges.ToolTip.NoChanges", "Nothing to revert.");
+		case EChangeRevertability::UploadInProgress: return LOCTEXT("RevertChanges.ToolTip.InTransit", "Disabled because upload is in progress.");
+		default: checkNoEntry(); return FText::GetEmpty();
+		}
+	}
+}
+
+#undef LOCTEXT_NAMESPACE
\ No newline at end of file
diff --git a/Engine/Plugins/Developer/Concert/ConcertApp/MultiUserClient/Source/MultiUserClient/Private/Widgets/ActiveSession/Replication/Client/SSubmissionControls.h b/Engine/Plugins/Developer/Concert/ConcertApp/MultiUserClient/Source/MultiUserClient/Private/Widgets/ActiveSession/Replication/Client/SSubmissionControls.h
new file mode 100644
--- /dev/null
+++ b/Engine/Plugins/Developer/Concert/ConcertApp/MultiUserClient/Source/MultiUserClient/Private/Widgets/ActiveSession/Replication/Client/SSubmissionControls.h
@@ -0,0 +1,64 @@
+// Copyright Epic Games, Inc. All Rights Reserved.
+
+#pragma once
+
+#include "Replication/Submission/ISubmissionOperation.h"
+#include "Widgets/DeclarativeSyntaxSupport.h"
+#include "Widgets/SCompoundWidget.h"
+
+class SNotificationItem;
+
+namespace UE::MultiUserClient
+{
+	class ISubmissionWorkflow;
+	/**
+	 * Displays a ISubmissionWorkflow.
+	 * 
+	 * Contains an upload and revert button.
+	 * Displays notifications to the user.
+	 */
+	class SSubmissionControls : public SCompoundWidget
+	{
+	public:
+
+		SLATE_BEGIN_ARGS(SSubmissionControls)
+		{}
+			SLATE_ATTRIBUTE(ISubmissionWorkflow*, SubmissionWorkflow)
+		SLATE_END_ARGS()
+
+		void Construct(const FArguments& InArgs);
+
+	private:
+
+		/** Used to get the submission workflow that is being displayed. */
+		TAttribute<ISubmissionWorkflow*> SubmissionWorkflowAttribute;
+		
+		ISubmissionWorkflow& GetSubmissionWorkflow() const
+		{
+			ISubmissionWorkflow* Workflow = SubmissionWorkflowAttribute.Get();
+			check(Workflow);
+			return *Workflow;
+		}
+
+		// Upload button
+		TSharedRef<SWidget> BuildUploadButton();
+		FText GetUploadButtonToolTipText() const;
+		FReply OnUploadButtonClicked() const;
+
+		// Handle submission events
+		void SetupNotificationsFor(ISubmissionOperation& SubmissionOperation) const;
+		void HandleStreamUpdatedNotification(ISubmissionOperation& SubmissionOperation) const;
+		void HandleAuthorityChangeRequestedNotification(ISubmissionOperation& SubmissionOperation) const;
+		static void HandleAuthorityChangedResponseNotification(
+			const ConcertSyncClient::Replication::FAuthorityChangeRequest& Request,
+			const FSubmitAuthorityChangesResponse& ResponseMetaData,
+			const TSharedRef<SNotificationItem>& NotificationItem
+			);
+
+		// Revert button
+		TSharedRef<SWidget> BuildRevertButton();
+		FText GetRevertButtonToolTipText() const;
+	};
+}
+
+
diff --git a/Engine/Plugins/Developer/Concert/ConcertApp/MultiUserClient/Source/MultiUserClient/Private/Widgets/ActiveSession/Replication/Client/Single/SReplicationClientView.cpp b/Engine/Plugins/Developer/Concert/ConcertApp/MultiUserClient/Source/MultiUserClient/Private/Widgets/ActiveSession/Replication/Client/Single/SReplicationClientView.cpp
new file mode 100644
--- /dev/null
+++ b/Engine/Plugins/Developer/Concert/ConcertApp/MultiUserClient/Source/MultiUserClient/Private/Widgets/ActiveSession/Replication/Client/Single/SReplicationClientView.cpp
@@ -0,0 +1,81 @@
+// Copyright Epic Games, Inc. All Rights Reserved.
+
+#include "SReplicationClientView.h"
+
+#include "StreamEditorColumns.h"
+#include "Replication/Client/ReplicationClient.h"
+#include "Replication/Editor/Model/IEditableObjectToPropertiesModel.h"
+#include "Replication/Editor/Model/Object/EditorObjectSelectionSourceModel.h"
+#include "Replication/Editor/Model/Property/SelectPropertyFromUClassModel.h"
+#include "Replication/Editor/View/IReplicationStreamEditor.h"
+#include "Replication/ReplicationWidgetFactories.h"
+#include "Replication/Submission/ISubmissionWorkflow.h"
+#include "Widgets/ActiveSession/Replication/Client/Single/SSingleClientToolbar.h"
+
+#include "Widgets/SBoxPanel.h"
+
+#define LOCTEXT_NAMESPACE "SReplicationClientView"
+
+namespace UE::MultiUserClient
+{
+	void SReplicationClientView::Construct(const FArguments& InArgs)
+	{
+		GetReplicationClientAttribute = InArgs._GetReplicationClient;
+		FReplicationClient* ReplicationClient = GetReplicationClientAttribute.Get();
+		check(ReplicationClient);
+
+		ConcertClientSharedSlate::IObjectToPropertiesModel& PropertyModel = *ReplicationClient->GetClientEditModel();
+		FAuthorityChangeTracker& AuthorityTracker = ReplicationClient->GetAuthorityDiffer();
+		ISubmissionWorkflow& SubmissionWorkflow = ReplicationClient->GetSubmissionWorkflow();
+
+		// Add checkboxes in front of top level and subobject rows for changing authority
+		using namespace ConcertClientSharedSlate;
+		const FCreateSubobjectViewParams SubobjectViewParams
+		{
+			.AdditionalColumns = { StreamEditorColumns::ToggleSubobjectAuthority(AuthorityTracker, SubmissionWorkflow) }
+		};
+		const FCreateEditorParams ReplicationEditorCreationParams
+		{
+			.DataModel = ReplicationClient->GetClientEditModel(),
+			.ObjectSource = MakeShared<FEditorObjectSelectionSourceModel>(),
+			.PropertySource = MakeShared<FSelectPropertyFromUClassModel>(),
+			.SubobjectView = CreateUnrealEditorSubobjectView(SubobjectViewParams),
+			.AdditionalObjectColumns = { StreamEditorColumns::ToggleTopLevelAuthority(PropertyModel, AuthorityTracker, SubmissionWorkflow) },
+			.IsEditingEnabled = TAttribute<bool>::CreateLambda([&SubmissionWorkflow](){ return SubmissionWorkflow.GetUploadability() != EChangeUploadability::NotImplemented; }),
+			.EditingDisabledToolTipText = LOCTEXT("Editing.NotImplemented", "Editing remote clients is not implemented. You can only edit the local client.")
+		};
+
+		EditorView = CreateDefaultStreamEditor(ReplicationEditorCreationParams);
+		ChildSlot
+		[
+			SNew(SVerticalBox)
+
+			// Toolbar
+			+SVerticalBox::Slot()
+			.AutoHeight()
+			.Padding(2.f)
+			[
+				SNew(SSingleClientToolbar)
+				.SubmissionWorkflow(&ReplicationClient->GetSubmissionWorkflow())
+				.ViewSelectionArea() [ InArgs._ViewSelectionArea.Widget ]
+			]
+
+			// Editor
+			+SVerticalBox::Slot()
+			.FillHeight(1.f)
+			[
+				EditorView.ToSharedRef()
+			]
+		];
+		
+		// Refresh UI if streams change externally, e.g. a remote client changed what they sent
+		ReplicationClient->OnModelExternallyChanged().AddSP(this, &SReplicationClientView::OnModelChanged);
+	}
+	
+	void SReplicationClientView::OnModelChanged() const
+	{
+		EditorView->Refresh();
+	}
+}
+
+#undef LOCTEXT_NAMESPACE
\ No newline at end of file
diff --git a/Engine/Plugins/Developer/Concert/ConcertApp/MultiUserClient/Source/MultiUserClient/Private/Widgets/ActiveSession/Replication/Client/Single/SReplicationClientView.h b/Engine/Plugins/Developer/Concert/ConcertApp/MultiUserClient/Source/MultiUserClient/Private/Widgets/ActiveSession/Replication/Client/Single/SReplicationClientView.h
new file mode 100644
--- /dev/null
+++ b/Engine/Plugins/Developer/Concert/ConcertApp/MultiUserClient/Source/MultiUserClient/Private/Widgets/ActiveSession/Replication/Client/Single/SReplicationClientView.h
@@ -0,0 +1,44 @@
+// Copyright Epic Games, Inc. All Rights Reserved.
+
+#pragma once
+
+#include "Widgets/DeclarativeSyntaxSupport.h"
+#include "Widgets/SCompoundWidget.h"
+
+namespace UE::ConcertClientSharedSlate
+{
+	class IReplicationStreamEditor;
+}
+
+namespace UE::MultiUserClient
+{
+	class FReplicationClient;
+	
+	/** Displays the contents of a client. */
+	class SReplicationClientView : public SCompoundWidget
+	{
+	public:
+
+		SLATE_BEGIN_ARGS(SReplicationClientView)
+		{}
+			/** The client to depict. Should always return true. If the client is destroyed, so should this widget be. */
+			SLATE_ATTRIBUTE(FReplicationClient*, GetReplicationClient)
+			/** Dedicated space for a widget with which to change the view. */
+			SLATE_NAMED_SLOT(FArguments, ViewSelectionArea)
+		SLATE_END_ARGS()
+
+		void Construct(const FArguments& InArgs);
+
+	private:
+		
+		/** The editor view of the replication content. */
+		TSharedPtr<ConcertClientSharedSlate::IReplicationStreamEditor> EditorView;
+		
+		/** The client to depict. Should always return true. If the client is destroyed, so should this widget be. */
+		TAttribute<FReplicationClient*> GetReplicationClientAttribute;
+		
+		/** Called when any of the streams change. */
+		void OnModelChanged() const;
+	};
+}
+
diff --git a/Engine/Plugins/Developer/Concert/ConcertApp/MultiUserClient/Source/MultiUserClient/Private/Widgets/ActiveSession/Replication/Client/Single/SSingleClientToolbar.cpp b/Engine/Plugins/Developer/Concert/ConcertApp/MultiUserClient/Source/MultiUserClient/Private/Widgets/ActiveSession/Replication/Client/Single/SSingleClientToolbar.cpp
new file mode 100644
--- /dev/null
+++ b/Engine/Plugins/Developer/Concert/ConcertApp/MultiUserClient/Source/MultiUserClient/Private/Widgets/ActiveSession/Replication/Client/Single/SSingleClientToolbar.cpp
@@ -0,0 +1,40 @@
+// Copyright Epic Games, Inc. All Rights Reserved.
+
+#include "SSingleClientToolbar.h"
+
+#include "Widgets/ActiveSession/Replication/Client/SSubmissionControls.h"
+
+#include "Widgets/Layout/SSpacer.h"
+#include "Widgets/SBoxPanel.h"
+
+namespace UE::MultiUserClient
+{
+	void SSingleClientToolbar::Construct(const FArguments& InArgs)
+	{
+		ChildSlot
+		[
+			SNew(SHorizontalBox)
+
+			+SHorizontalBox::Slot()
+			.AutoWidth()
+			.Padding(5.f, 0.f)
+			[
+				InArgs._ViewSelectionArea.Widget
+			]
+
+			+SHorizontalBox::Slot()
+			.FillWidth(1.f)
+			[
+				SNew(SSpacer)
+			]
+
+			// Submission controls
+			+SHorizontalBox::Slot()
+			.AutoWidth()
+			[
+				SNew(SSubmissionControls)
+				.SubmissionWorkflow(InArgs._SubmissionWorkflow)
+			]
+		];
+	}
+}
\ No newline at end of file
diff --git a/Engine/Plugins/Developer/Concert/ConcertApp/MultiUserClient/Source/MultiUserClient/Private/Widgets/ActiveSession/Replication/Client/Single/SSingleClientToolbar.h b/Engine/Plugins/Developer/Concert/ConcertApp/MultiUserClient/Source/MultiUserClient/Private/Widgets/ActiveSession/Replication/Client/Single/SSingleClientToolbar.h
new file mode 100644
--- /dev/null
+++ b/Engine/Plugins/Developer/Concert/ConcertApp/MultiUserClient/Source/MultiUserClient/Private/Widgets/ActiveSession/Replication/Client/Single/SSingleClientToolbar.h
@@ -0,0 +1,35 @@
+// Copyright Epic Games, Inc. All Rights Reserved.
+
+#pragma once
+
+#include "Widgets/DeclarativeSyntaxSupport.h"
+#include "Widgets/SCompoundWidget.h"
+
+namespace UE::MultiUserClient
+{
+	class FStreamChangeTracker;
+	class ISubmissionWorkflow;
+
+	/**
+	 * Contains a bunch of actions that can be performed on client view.
+	 *
+	 * Layout:
+	 * | AdditionalToolbarWidgets                   | Submit | Revert |
+	 */
+	class SSingleClientToolbar : public SCompoundWidget
+	{
+	public:
+		
+		SLATE_BEGIN_ARGS(SSingleClientToolbar)
+		{}
+			/** Manages submission and reverting of changes. Used for upload and revert buttons. */
+			SLATE_ATTRIBUTE(ISubmissionWorkflow*, SubmissionWorkflow)
+
+			/** Dedicated space for a widget with which to change the view. */
+			SLATE_NAMED_SLOT(FArguments, ViewSelectionArea)
+		SLATE_END_ARGS()
+
+		void Construct(const FArguments& InArgs);
+	};
+}
+
diff --git a/Engine/Plugins/Developer/Concert/ConcertApp/MultiUserClient/Source/MultiUserClient/Private/Widgets/ActiveSession/Replication/Client/Single/StreamEditorColumns.cpp b/Engine/Plugins/Developer/Concert/ConcertApp/MultiUserClient/Source/MultiUserClient/Private/Widgets/ActiveSession/Replication/Client/Single/StreamEditorColumns.cpp
new file mode 100644
--- /dev/null
+++ b/Engine/Plugins/Developer/Concert/ConcertApp/MultiUserClient/Source/MultiUserClient/Private/Widgets/ActiveSession/Replication/Client/Single/StreamEditorColumns.cpp
@@ -0,0 +1,213 @@
+// Copyright Epic Games, Inc. All Rights Reserved.
+
+#include "StreamEditorColumns.h"
+
+#include "Replication/Authority/AuthorityChangeTracker.h"
+#include "Replication/Authority/EAuthorityMutability.h"
+#include "Replication/Authority/IClientAuthoritySynchronizer.h"
+#include "Replication/Editor/Model/IObjectToPropertiesModel.h"
+#include "Replication/Editor/View/ReplicationColumnsUtils.h"
+#include "Replication/Submission/ISubmissionWorkflow.h"
+
+#define LOCTEXT_NAMESPACE "StreamEditorColumns"
+
+namespace UE::MultiUserClient::StreamEditorColumns
+{
+	const FName ToggleTopLevelAuthorityColumnId = TEXT("ToggleTopLevelAuthorityColumn");
+	const FName ToggleSubobjectAuthorityColumnId = TEXT("ToggleSubobjectAuthorityColumn");
+
+	namespace Private::ToggleTopLevelAuthority
+	{
+		static ECheckBoxState GetCheckboxState(
+			const ConcertClientSharedSlate::FReplicatedObjectData& ObjectData,
+			ConcertClientSharedSlate::IObjectToPropertiesModel* ClientStreamModel,
+			FAuthorityChangeTracker* ChangeTracker
+			)
+		{
+			const FSoftObjectPath& ObjectPath = ObjectData.GetObjectPath();
+
+			// Init the checkbox state to the top level object, if the top level object has properties associated with it...
+			const bool bCanSetTopLevelState = ChangeTracker->CanSetAuthorityFor(ObjectPath)
+				|| ChangeTracker->GetChangeAuthorityMutability(ObjectPath) == EAuthorityMutability::NotSupported;
+			TOptional<ECheckBoxState> CheckBoxState = bCanSetTopLevelState
+				? ChangeTracker->GetAuthorityStateAfterApplied(ObjectPath) ? ECheckBoxState::Checked : ECheckBoxState::Unchecked
+				: TOptional<ECheckBoxState>{};
+
+			// ... and then determine whether all the other subobjects share that state (Check / Unchecked) or not (Undetermined)
+			ClientStreamModel->ForEachSubobject(ObjectPath, [&CheckBoxState, &ChangeTracker](const FSoftObjectPath& Child)
+			{
+				if (!ChangeTracker->CanSetAuthorityFor(Child)
+					&& ChangeTracker->GetChangeAuthorityMutability(Child) != EAuthorityMutability::NotSupported)
+				{
+					return EBreakBehavior::Continue;
+				}
+
+				// Init checkbox state
+				const bool bCurrentAuthorityState = ChangeTracker->GetAuthorityStateAfterApplied(Child);
+				const bool bIsFirstObject = !CheckBoxState.IsSet(); 
+				if (bIsFirstObject)
+				{
+					CheckBoxState = bCurrentAuthorityState ? ECheckBoxState::Checked : ECheckBoxState::Unchecked;
+					return EBreakBehavior::Continue;
+				}
+
+				// Return the checkbox state that all subobjects share, or Undetermined if mixed
+				const bool bMustBeUndetermined = (*CheckBoxState == ECheckBoxState::Unchecked && bCurrentAuthorityState)
+					|| (*CheckBoxState == ECheckBoxState::Checked && !bCurrentAuthorityState);
+				if (bMustBeUndetermined)
+				{
+					CheckBoxState = ECheckBoxState::Undetermined;
+					return EBreakBehavior::Break;
+				}
+				return EBreakBehavior::Continue;
+			});
+						
+			return CheckBoxState
+				? *CheckBoxState
+				// If this case happens, there are no objects to check. In that case, it visually "looks better" to be unchecked
+				: ECheckBoxState::Unchecked;
+		}
+
+		static void OnCheckboxStateChanged(
+			bool bIsChecked,
+			const ConcertClientSharedSlate::FReplicatedObjectData& ObjectData,
+			ConcertClientSharedSlate::IObjectToPropertiesModel* ClientStreamModel,
+			FAuthorityChangeTracker* ChangeTracker
+			)
+		{
+			const FSoftObjectPath& ObjectPath = ObjectData.GetObjectPath();
+			ChangeTracker->SetAuthorityIfAllowed(ObjectPath, bIsChecked);
+			ClientStreamModel->ForEachSubobject(ObjectPath, [&ChangeTracker, bIsChecked](const FSoftObjectPath& Child)
+			{
+				ChangeTracker->SetAuthorityIfAllowed(Child, bIsChecked);
+				return EBreakBehavior::Continue;
+			});
+		}
+
+		static bool IsEnabled(
+			const ConcertClientSharedSlate::FReplicatedObjectData& ObjectData,
+			ConcertClientSharedSlate::IObjectToPropertiesModel* ClientStreamModel,
+			FAuthorityChangeTracker* ChangeTracker
+			)
+		{
+			const FSoftObjectPath& TopLevelObjectPath = ObjectData.GetObjectPath();
+			const bool bCanChangeTopLevel = ChangeTracker->CanSetAuthorityFor(TopLevelObjectPath);
+			if (bCanChangeTopLevel)
+			{
+				return true;
+			}
+						
+			bool bCanChangeAnySubobject = false;
+			ClientStreamModel->ForEachSubobject(TopLevelObjectPath, [&ChangeTracker, &bCanChangeAnySubobject](const FSoftObjectPath& Child)
+			{
+				bCanChangeAnySubobject |= ChangeTracker->CanSetAuthorityFor(Child);
+				return bCanChangeAnySubobject ? EBreakBehavior::Break : EBreakBehavior::Continue;
+			});
+			return bCanChangeAnySubobject;
+		}
+	}
+	
+	ConcertClientSharedSlate::ReplicationColumns::FReplicationTopLevelObjectColumn ToggleTopLevelAuthority(
+		ConcertClientSharedSlate::IObjectToPropertiesModel& ClientStreamModel,
+		FAuthorityChangeTracker& ChangeTracker,
+		ISubmissionWorkflow& SubmissionWorkflow
+		)
+	{
+		using namespace ConcertClientSharedSlate;
+		using namespace Private::ToggleTopLevelAuthority;
+		using FTopLevelColumnDelegates = TReplicationColumnDelegates<FReplicatedObjectData>;
+		
+		return MakeCheckboxColumn<FReplicatedObjectData>(
+			ToggleTopLevelAuthorityColumnId,
+				FTopLevelColumnDelegates(
+					FTopLevelColumnDelegates::FGetColumnCheckboxState::CreateStatic(&GetCheckboxState, &ClientStreamModel, &ChangeTracker),
+					FTopLevelColumnDelegates::FOnColumnCheckboxChanged::CreateStatic(&OnCheckboxStateChanged, &ClientStreamModel, &ChangeTracker),
+					FTopLevelColumnDelegates::FGetToolTipText::CreateLambda([&ClientStreamModel, &ChangeTracker, &SubmissionWorkflow](const FReplicatedObjectData& ObjectData)
+					{
+						
+						if (SubmissionWorkflow.GetUploadability() == EChangeUploadability::NotImplemented)
+						{
+							const ECheckBoxState CheckBoxState = GetCheckboxState(ObjectData, &ClientStreamModel, &ChangeTracker);
+							const FText CheckBoxText = [CheckBoxState]()
+							{
+								switch (CheckBoxState)
+								{
+								case ECheckBoxState::Unchecked: return LOCTEXT("TopLevel.ChangeAuthority.ToolTip.NotSupported.Unchecked", "None of the subobjects are being replicated.");
+								case ECheckBoxState::Checked: return LOCTEXT("TopLevel.ChangeAuthority.ToolTip.NotSupported.Unchecked", "All of the subobjects are being replicated.");
+								case ECheckBoxState::Undetermined: return LOCTEXT("TopLevel.ChangeAuthority.ToolTip.NotSupported.Unchecked", "Some of the subobjects are being replicated.");
+								default: checkNoEntry(); return FText::GetEmpty();
+								}
+							}();
+							
+							return FText::Format(
+								LOCTEXT("TopLevel.ChangeAuthority.ToolTip.NotSupportedFmt", "{0}\nEditing remote clients is not implemented. You can only edit the local client."),
+								CheckBoxText
+								);
+						}
+						
+						return IsEnabled(ObjectData, &ClientStreamModel, &ChangeTracker)
+							? LOCTEXT("TopLevel.ChangeAuthority.Allowed", "Toggles whether this object and its subobjects will replicate the assigned properties.")
+							: LOCTEXT("TopLevel.ChangeAuthority.Disallowed", "Toggles whether this object and its subobjects will replicate the assigned properties.\nDisabled because none of the subobjects can be replicated with their current configuration.");
+					}),
+					FTopLevelColumnDelegates::FIsEnabled::CreateStatic(&IsEnabled, &ClientStreamModel, &ChangeTracker)
+				),
+			FText::GetEmpty(),
+			0
+		);
+	}
+
+	ConcertClientSharedSlate::ReplicationColumns::FReplicationSubobjectObjectColumn ToggleSubobjectAuthority(
+		FAuthorityChangeTracker& ChangeTracker,
+		ISubmissionWorkflow& SubmissionWorkflow
+		)
+	{
+		using namespace ConcertClientSharedSlate;
+		using FSubobjectColumnDelegates = TReplicationColumnDelegates<FReplicatedObjectData>;
+		return MakeCheckboxColumn<FReplicatedObjectData>(
+			ToggleTopLevelAuthorityColumnId,
+				FSubobjectColumnDelegates(
+					FSubobjectColumnDelegates::FGetColumnCheckboxState::CreateLambda(
+					[&ChangeTracker](const FReplicatedObjectData& ObjectData)
+					{
+						const bool bHasAuthority = ChangeTracker.GetAuthorityStateAfterApplied(ObjectData.GetObjectPath());
+						return bHasAuthority ? ECheckBoxState::Checked : ECheckBoxState::Unchecked;
+					}),
+					FSubobjectColumnDelegates::FOnColumnCheckboxChanged::CreateLambda(
+					[&ChangeTracker](bool bIsChecked, const FReplicatedObjectData& ObjectData)
+					{
+						ChangeTracker.SetAuthorityIfAllowed(ObjectData.GetObjectPath(), bIsChecked);
+					}),
+					FSubobjectColumnDelegates::FGetToolTipText::CreateLambda([&ChangeTracker, &SubmissionWorkflow](const FReplicatedObjectData& ObjectData)
+					{
+						const bool bHasAuthority = ChangeTracker.GetAuthorityStateAfterApplied(ObjectData.GetObjectPath());
+						const FText NotSupportedText = FText::Format(
+							LOCTEXT("Subobject.ChangeAuthority.ToolTip.NotSupportedFmt", "{0}\nEditing remote clients is not implemented. You can only edit the local client."),
+							bHasAuthority ? LOCTEXT("Subobject.ChangeAuthority.Replicating.True", "This object is being replicated.") : LOCTEXT("Subobject.ChangeAuthority.Replicating.True", "This object is not currently being replicated.")
+							);
+						if (SubmissionWorkflow.GetUploadability() == EChangeUploadability::NotImplemented)
+						{
+							return NotSupportedText;
+						}
+						
+						switch (ChangeTracker.GetChangeAuthorityMutability(ObjectData.GetObjectPath()))
+						{
+						case EAuthorityMutability::Allowed: return LOCTEXT("Subobject.ChangeAuthority.ToolTip.Allowed", "Toggles whether this object should be replicated.");
+						case EAuthorityMutability::NoProperties: return LOCTEXT("Subobject.ChangeAuthority.ToolTip.NoProperties", "Toggles whether this object should be replicated.\nAssign properties to this object first.\n");
+						// TODO DP UE-198356: Make the tooltip info about the conflicting client
+						case EAuthorityMutability::ClientConflict: return LOCTEXT("Subobject.ChangeAuthority.ToolTip.ClientConflict", "Toggles whether this object should be replicated.\nAnother client already is already replicating some of the assigned properties.");
+						case EAuthorityMutability::NotSupported: return NotSupportedText;
+						default: checkNoEntry(); return FText::GetEmpty();
+						}
+					}),
+					FSubobjectColumnDelegates::FIsEnabled::CreateLambda([&ChangeTracker](const FReplicatedObjectData& ObjectData)
+					{
+						return ChangeTracker.CanSetAuthorityFor(ObjectData.GetObjectPath());
+					})
+				),
+			FText::GetEmpty(),
+			0
+		);
+	}
+}
+
+#undef LOCTEXT_NAMESPACE
\ No newline at end of file
diff --git a/Engine/Plugins/Developer/Concert/ConcertApp/MultiUserClient/Source/MultiUserClient/Private/Widgets/ActiveSession/Replication/Client/Single/StreamEditorColumns.h b/Engine/Plugins/Developer/Concert/ConcertApp/MultiUserClient/Source/MultiUserClient/Private/Widgets/ActiveSession/Replication/Client/Single/StreamEditorColumns.h
new file mode 100644
--- /dev/null
+++ b/Engine/Plugins/Developer/Concert/ConcertApp/MultiUserClient/Source/MultiUserClient/Private/Widgets/ActiveSession/Replication/Client/Single/StreamEditorColumns.h
@@ -0,0 +1,54 @@
+// Copyright Epic Games, Inc. All Rights Reserved.
+
+#pragma once
+
+#include "Replication/Editor/View/PredefinedReplicationColumns.h"
+#include "Widgets/DeclarativeSyntaxSupport.h"
+#include "Widgets/SCompoundWidget.h"
+
+namespace UE::ConcertClientSharedSlate
+{
+	class IObjectToPropertiesModel;
+}
+
+namespace UE::MultiUserClient
+{
+	class FAuthorityChangeTracker;
+	class IClientAuthoritySynchronizer;
+	class ISubmissionWorkflow;
+}
+
+namespace UE::MultiUserClient::StreamEditorColumns
+{
+	extern const FName ToggleTopLevelAuthorityColumnId;
+	extern const FName ToggleSubobjectAuthorityColumnId;
+
+	/**
+	 * Checkbox placed in the top-level view.
+	 * It gives / removes authority for the top-level object and all of its subobjects.
+	 *
+	 * @param ClientStreamModel Used to discover the client's registered subobjects
+	 * @param ChangeTracker Used to determine the checkbox state and change authority.
+	 * @param SubmissionWorkflow Used to determine whether changing authority is at all enabled
+	 * @return Column that can be placed in the table
+	 */
+	ConcertClientSharedSlate::ReplicationColumns::FReplicationTopLevelObjectColumn ToggleTopLevelAuthority(
+		ConcertClientSharedSlate::IObjectToPropertiesModel& ClientStreamModel,
+		FAuthorityChangeTracker& ChangeTracker,
+		ISubmissionWorkflow& SubmissionWorkflow
+		);
+
+	/**
+	 * Checkbox placed in the subobject view
+	 * It gives / removes authority for the object it is placed next to.
+	 * 
+	 * @param AuthoritySynchronizer Determines whether the client has authority or not.
+	 * @param ChangeTracker Used to determine the checkbox state and change authority.
+	 * @param AuthoritySynchronizer Used to determine whether changing authority is at all enabled
+	 * @return Column that can be placed in the table
+	 */
+	ConcertClientSharedSlate::ReplicationColumns::FReplicationSubobjectObjectColumn ToggleSubobjectAuthority(
+		FAuthorityChangeTracker& ChangeTracker,
+		ISubmissionWorkflow& SubmissionWorkflow
+		);
+}
diff --git a/Engine/Plugins/Developer/Concert/ConcertApp/MultiUserClient/Source/MultiUserClient/Private/Widgets/ActiveSession/Replication/Joined/SReplicationClientView.cpp b/Engine/Plugins/Developer/Concert/ConcertApp/MultiUserClient/Source/MultiUserClient/Private/Widgets/ActiveSession/Replication/Joined/SReplicationClientView.cpp
deleted file mode 100644
diff --git a/Engine/Plugins/Developer/Concert/ConcertApp/MultiUserClient/Source/MultiUserClient/Private/Widgets/ActiveSession/Replication/Joined/SReplicationClientView.h b/Engine/Plugins/Developer/Concert/ConcertApp/MultiUserClient/Source/MultiUserClient/Private/Widgets/ActiveSession/Replication/Joined/SReplicationClientView.h
deleted file mode 100644
diff --git a/Engine/Plugins/Developer/Concert/ConcertApp/MultiUserClient/Source/MultiUserClient/Private/Widgets/ActiveSession/Replication/Joined/SReplicationClientViewToolbar.cpp b/Engine/Plugins/Developer/Concert/ConcertApp/MultiUserClient/Source/MultiUserClient/Private/Widgets/ActiveSession/Replication/Joined/SReplicationClientViewToolbar.cpp
deleted file mode 100644
diff --git a/Engine/Plugins/Developer/Concert/ConcertApp/MultiUserClient/Source/MultiUserClient/Private/Widgets/ActiveSession/Replication/Joined/SReplicationClientViewToolbar.h b/Engine/Plugins/Developer/Concert/ConcertApp/MultiUserClient/Source/MultiUserClient/Private/Widgets/ActiveSession/Replication/Joined/SReplicationClientViewToolbar.h
deleted file mode 100644
diff --git a/Engine/Plugins/Developer/Concert/ConcertApp/MultiUserClient/Source/MultiUserClient/Private/Widgets/ActiveSession/Replication/Joined/SReplicationJoinedView.cpp b/Engine/Plugins/Developer/Concert/ConcertApp/MultiUserClient/Source/MultiUserClient/Private/Widgets/ActiveSession/Replication/Joined/SReplicationJoinedView.cpp
--- a/Engine/Plugins/Developer/Concert/ConcertApp/MultiUserClient/Source/MultiUserClient/Private/Widgets/ActiveSession/Replication/Joined/SReplicationJoinedView.cpp
+++ b/Engine/Plugins/Developer/Concert/ConcertApp/MultiUserClient/Source/MultiUserClient/Private/Widgets/ActiveSession/Replication/Joined/SReplicationJoinedView.cpp
@@ -4,9 +4,9 @@
 
 #include "ConcertLogGlobal.h"
 #include "IConcertSyncClient.h"
-#include "SReplicationClientView.h"
 #include "Replication/MultiUserReplicationManager.h"
 #include "SSelectClientViewComboButton.h"
+#include "Widgets/ActiveSession/Replication/Client/Single/SReplicationClientView.h"
 
 #include "Framework/Notifications/NotificationManager.h"
 #include "Replication/Client/RemoteReplicationClient.h"
@@ -36,7 +36,7 @@ namespace UE::MultiUserClient
 			[
 				SNew(SReplicationClientView)
 				.GetReplicationClient_Lambda([this](){ return &ReplicationManager->GetClientManager()->GetLocalClient(); })
-				.AdditionalToolbarWidgets()
+				.ViewSelectionArea()
 				[
 					MakeClientSelectionArea()
 				]
@@ -45,10 +45,6 @@ namespace UE::MultiUserClient
 		
 		RefreshClientViewSwitcher();
 		ReplicationManager->GetClientManager()->OnRemoteClientsChanged().AddSP(this, &SReplicationJoinedView::RefreshClientViewSwitcher);
-		
-		// Show notifications about changing authority
-		ReplicationManager->GetAuthorityPolicy()->OnAuthorityRequestSent_AnyThread().AddSP(this, &SReplicationJoinedView::OnAuthorityRequestSent_AnyThread);
-		ReplicationManager->GetAuthorityPolicy()->OnAuthorityResponseReceived_AnyThread().AddSP(this, &SReplicationJoinedView::OnAuthorityResponseReceived_AnyThread);
 	}
 
 	void SReplicationJoinedView::RefreshClientViewSwitcher()
@@ -121,7 +117,7 @@ namespace UE::MultiUserClient
         				// It is unsafe to simply capture RemoteClient because the containing TArray may reallocate its location
         				return ReplicationManager->GetClientManager()->FindRemoteClient(EndpointId);
         			})
-        			.AdditionalToolbarWidgets()
+        			.ViewSelectionArea()
         			[
         				MakeClientSelectionArea()
         			]
@@ -210,67 +206,6 @@ namespace UE::MultiUserClient
 		}
 		return {};
 	}
-
-	void SReplicationJoinedView::OnAuthorityRequestSent_AnyThread(
-		const ConcertSyncClient::Replication::FAuthorityChangeRequest& Request
-		)
-	{
-		ExecuteOnGameThread(TEXT("OnAuthorityRequestSent"), [WeakThis = TWeakPtr<SReplicationJoinedView>(SharedThis(this)), Request]()
-		{
-			if (const TSharedPtr<SReplicationJoinedView> ThisPin = WeakThis.Pin())
-			{
-				FNotificationInfo NotificationInfo(LOCTEXT("Authority.InProgress.Text", "Requesting authority."));
-				NotificationInfo.bUseThrobber = true;
-				NotificationInfo.bUseSuccessFailIcons = true;
-				NotificationInfo.bFireAndForget = false;
-				ThisPin->AuthorityChangeNotification = FSlateNotificationManager::Get().AddNotification(NotificationInfo);
-				ThisPin->AuthorityChangeNotification->SetSubText(
-					FText::Format(LOCTEXT("Authority.InProgress.SubTextFmt", "Taking {0}|plural(one=object,other\nReleasing {1}|plural(one=object,other=objects)"),
-						Request.TakeAuthority.Num(),
-						Request.ReleaseAuthority.Num()
-					));
-			}
-		});
-	}
-
-	void SReplicationJoinedView::OnAuthorityResponseReceived_AnyThread(
-		const ConcertSyncClient::Replication::FAuthorityChangeRequest& Request,
-		const ConcertSyncClient::Replication::FAuthorityChangeResponse& Response
-		)
-	{
-		ExecuteOnGameThread(TEXT("OnAuthorityResponseReceived"), [WeakThis = TWeakPtr<SReplicationJoinedView>(SharedThis(this)), Request, Response]()
-		{
-			if (const TSharedPtr<SReplicationJoinedView> ThisPin = WeakThis.Pin()
-				; ThisPin && ensure(ThisPin->AuthorityChangeNotification))
-			{
-				const bool bSuccess = Response.RejectedObjects.IsEmpty();
-				const FText Text = bSuccess
-					? LOCTEXT("Authority.Text.Accepted", "Authority change accepted.")
-					: LOCTEXT("Authority.Text.Rejection", "Authority change had rejections.");
-				const FText SubText = bSuccess
-					? FText::GetEmpty()
-					: FText::Format(LOCTEXT("Authority.SubText.RejectionFmt", "{0} accepted {1} rejected\nSee logs for details."),
-						Request.TakeAuthority.Num() - Response.RejectedObjects.Num(),
-						Response.RejectedObjects.Num()
-						);
-				
-				ThisPin->AuthorityChangeNotification->SetCompletionState(bSuccess ? SNotificationItem::CS_Success : SNotificationItem::CS_Fail);
-				ThisPin->AuthorityChangeNotification->SetText(Text);
-				ThisPin->AuthorityChangeNotification->SetSubText(SubText);
-
-				if (!bSuccess)
-				{
-					const FString Error = FString::JoinBy(Response.RejectedObjects, TEXT("\n"), [](const TPair<FSoftObjectPath, FConcertStreamArray>& Pair)
-					{
-						return Pair.Key.ToString();
-					});
-					UE_LOG(LogConcert, Error, TEXT("Authority change rejected objects:\n%s"), *Error);
-				}
-				
-				ThisPin->AuthorityChangeNotification->ExpireAndFadeout();
-			}
-		});
-	}
 }
 
 #undef LOCTEXT_NAMESPACE
\ No newline at end of file
diff --git a/Engine/Plugins/Developer/Concert/ConcertApp/MultiUserClient/Source/MultiUserClient/Private/Widgets/ActiveSession/Replication/Joined/SReplicationJoinedView.h b/Engine/Plugins/Developer/Concert/ConcertApp/MultiUserClient/Source/MultiUserClient/Private/Widgets/ActiveSession/Replication/Joined/SReplicationJoinedView.h
--- a/Engine/Plugins/Developer/Concert/ConcertApp/MultiUserClient/Source/MultiUserClient/Private/Widgets/ActiveSession/Replication/Joined/SReplicationJoinedView.h
+++ b/Engine/Plugins/Developer/Concert/ConcertApp/MultiUserClient/Source/MultiUserClient/Private/Widgets/ActiveSession/Replication/Joined/SReplicationJoinedView.h
@@ -10,19 +10,15 @@ class IConcertSyncClient;
 class SNotificationItem;
 class SWidgetSwitcher;
 
-namespace UE::MultiUserClient
-{
-	class SReplicationClientView;
-}
-
 namespace UE::ConcertClientSharedSlate
 {
-	class IReplicationEditorView;
+	class IReplicationStreamEditor;
 }
 
 namespace UE::MultiUserClient
 {
 	class FMultiUserReplicationManager;
+	class SReplicationClientView;
 
 	/** This widget is displayed by SReplicationRootWidget when the client has joined replication. */
 	class SReplicationJoinedView : public SCompoundWidget
@@ -65,11 +61,5 @@ namespace UE::MultiUserClient
 		/** Creates a combobox with which the content of ClientViewSwitcher can be changed. */
 		TSharedRef<SWidget> MakeClientSelectionComboBox();
 		TOptional<FGuid> GetRemoteClientBySwitcherIndex(int32 WidgetSwitcherIndex) const;
-
-		void OnAuthorityRequestSent_AnyThread(const ConcertSyncClient::Replication::FAuthorityChangeRequest& Request);
-		void OnAuthorityResponseReceived_AnyThread(
-			const ConcertSyncClient::Replication::FAuthorityChangeRequest& Request,
-			const ConcertSyncClient::Replication::FAuthorityChangeResponse& Response
-			);
 	};
-}
+}
\ No newline at end of file
diff --git a/Engine/Plugins/Developer/Concert/ConcertSync/ConcertSyncClient/Source/ConcertSyncClient/Private/Replication/Manager/ReplicationManagerState_Connected.cpp b/Engine/Plugins/Developer/Concert/ConcertSync/ConcertSyncClient/Source/ConcertSyncClient/Private/Replication/Manager/ReplicationManagerState_Connected.cpp
--- a/Engine/Plugins/Developer/Concert/ConcertSync/ConcertSyncClient/Source/ConcertSyncClient/Private/Replication/Manager/ReplicationManagerState_Connected.cpp
+++ b/Engine/Plugins/Developer/Concert/ConcertSync/ConcertSyncClient/Source/ConcertSyncClient/Private/Replication/Manager/ReplicationManagerState_Connected.cpp
@@ -255,7 +255,8 @@ namespace UE::ConcertSyncClient::Replication
 			const FSoftObjectPath& ReplicatedObject = TakeAuthority.Key;
 			// Request will be discarded so ...
 			FConcertStreamArray& ReplicatedStreams = TakeAuthority.Value;
-
+			
+			UE_CLOG(ReplicatedStreams.StreamIds.IsEmpty(), LogConcert, Warning, TEXT("Your FAuthorityChangeRequest::TakeAuthority request contained empty stream ID array for object %s"), *ReplicatedObject.ToString());
 			const FConcertStreamArray* RejectedStreams = Response.RejectedObjects.Find(ReplicatedObject);
 			if (RejectedStreams)
 			{
diff --git a/Engine/Plugins/Developer/Concert/ConcertUI/ConcertClientSharedSlate/Source/ConcertClientSharedSlate/ConcertClientSharedSlate.Build.cs b/Engine/Plugins/Developer/Concert/ConcertUI/ConcertClientSharedSlate/Source/ConcertClientSharedSlate/ConcertClientSharedSlate.Build.cs
--- a/Engine/Plugins/Developer/Concert/ConcertUI/ConcertClientSharedSlate/Source/ConcertClientSharedSlate/ConcertClientSharedSlate.Build.cs
+++ b/Engine/Plugins/Developer/Concert/ConcertUI/ConcertClientSharedSlate/Source/ConcertClientSharedSlate/ConcertClientSharedSlate.Build.cs
@@ -34,7 +34,6 @@ namespace UnrealBuildTool.Rules
 					"Slate",
 					"SlateCore",
 					"SubobjectDataInterface",
-					"SubobjectEditor",
 					"ToolMenus",
 					"ToolWidgets",
 					"UnrealEd", 
diff --git a/Engine/Plugins/Developer/Concert/ConcertUI/ConcertClientSharedSlate/Source/ConcertClientSharedSlate/Private/Replication/Editor/Model/DisplayUtils.cpp b/Engine/Plugins/Developer/Concert/ConcertUI/ConcertClientSharedSlate/Source/ConcertClientSharedSlate/Private/Replication/Editor/Model/DisplayUtils.cpp
--- a/Engine/Plugins/Developer/Concert/ConcertUI/ConcertClientSharedSlate/Source/ConcertClientSharedSlate/Private/Replication/Editor/Model/DisplayUtils.cpp
+++ b/Engine/Plugins/Developer/Concert/ConcertUI/ConcertClientSharedSlate/Source/ConcertClientSharedSlate/Private/Replication/Editor/Model/DisplayUtils.cpp
@@ -16,13 +16,22 @@ namespace UE::ConcertClientSharedSlate::DisplayUtils
 	FText GetObjectDisplayText(const FSoftObjectPath& Object)
 	{
 		// Important! The object may not be loaded, yet. This could be if the asset is using a level that was not opened. 
-		UObject* LoadedObject = Object.TryLoad();
-		return LoadedObject
-			? FText::FromString(GetObjectDisplayString(*LoadedObject))
-			: FText::GetEmpty();
+		if (const UObject* LoadedObject = Object.ResolveObject())
+		{
+			return FText::FromString(GetObjectDisplayString(*LoadedObject));
+		}
+
+		// Subpath looks like this PersistentLevel.Actor.Component
+		const FString& Subpath = Object.GetSubPathString();
+		const int32 LastDotIndex = Subpath.Find(TEXT("."), ESearchCase::CaseSensitive, ESearchDir::FromEnd);
+		if (LastDotIndex == INDEX_NONE)
+		{
+			return FText::GetEmpty();
+		}
+		return FText::FromString(Subpath.RightChop(LastDotIndex + 1));
 	}
 
-	FString GetObjectDisplayString(UObject& Object)
+	FString GetObjectDisplayString(const UObject& Object)
 	{
 		if (const AActor* AsActor = Cast<AActor>(&Object))
 		{
diff --git a/Engine/Plugins/Developer/Concert/ConcertUI/ConcertClientSharedSlate/Source/ConcertClientSharedSlate/Private/Replication/Editor/Model/DisplayUtils.h b/Engine/Plugins/Developer/Concert/ConcertUI/ConcertClientSharedSlate/Source/ConcertClientSharedSlate/Private/Replication/Editor/Model/DisplayUtils.h
--- a/Engine/Plugins/Developer/Concert/ConcertUI/ConcertClientSharedSlate/Source/ConcertClientSharedSlate/Private/Replication/Editor/Model/DisplayUtils.h
+++ b/Engine/Plugins/Developer/Concert/ConcertUI/ConcertClientSharedSlate/Source/ConcertClientSharedSlate/Private/Replication/Editor/Model/DisplayUtils.h
@@ -20,7 +20,7 @@ namespace UE::ConcertClientSharedSlate::DisplayUtils
 	/** @return The text to use for displaying this object's name */
 	FText GetObjectDisplayText(const FSoftObjectPath& Object);
 	/** @return More lightweight version of GetReplicatedObjectDisplayName which does not construct any FText. */
-	FString GetObjectDisplayString(UObject& Object);
+	FString GetObjectDisplayString(const UObject& Object);
 	
 	/** @return The text to use for displaying this object's type */
 	FText GetObjectTypeText(const IObjectToPropertiesModel& Model, const FSoftObjectPath& Object);
diff --git a/Engine/Plugins/Developer/Concert/ConcertUI/ConcertClientSharedSlate/Source/ConcertClientSharedSlate/Private/Replication/Editor/Model/Subobject/ComponentHierarchySubobjectModel.cpp b/Engine/Plugins/Developer/Concert/ConcertUI/ConcertClientSharedSlate/Source/ConcertClientSharedSlate/Private/Replication/Editor/Model/Subobject/ComponentHierarchySubobjectModel.cpp
new file mode 100644
--- /dev/null
+++ b/Engine/Plugins/Developer/Concert/ConcertUI/ConcertClientSharedSlate/Source/ConcertClientSharedSlate/Private/Replication/Editor/Model/Subobject/ComponentHierarchySubobjectModel.cpp
@@ -0,0 +1,114 @@
+// Copyright Epic Games, Inc. All Rights Reserved.
+
+#include "ComponentHierarchySubobjectModel.h"
+
+#include "SubobjectDataSubsystem.h"
+#include "Components/SceneComponent.h"
+#include "Misc/EBreakBehavior.h"
+
+#include "GameFramework/Actor.h"
+
+namespace UE::ConcertClientSharedSlate
+{
+	namespace Private
+	{
+		static bool IsValidComponent(const UActorComponent& Component)
+		{
+			return !Component.IsVisualizationComponent();
+		}
+	}
+	
+	const FName FComponentHierarchySubobjectModel::SceneComponentsCategory(TEXT("SceneComponents"));
+	const FName FComponentHierarchySubobjectModel::ActorComponentsCategory(TEXT("ActorComponents"));
+	
+	void FComponentHierarchySubobjectModel::SetTopLevelObject(const FSoftObjectPath& InTopLevelObject)
+	{
+		if (InTopLevelObject != TopLevelObject)
+		{
+			TopLevelObject = InTopLevelObject;
+
+			if (AActor* Actor = GetTopLevelObjectAsActor())
+			{
+				USubobjectDataSubsystem* SubobjectDataSubsystem = USubobjectDataSubsystem::Get();
+				TArray<FSubobjectDataHandle> Handles;
+				SubobjectDataSubsystem->GatherSubobjectData(Actor, Handles);
+				
+				ObjectMetaData.Empty();
+				for (const FSubobjectDataHandle& Handle : Handles)
+				{
+					const FSubobjectData* SubobjectData = Handle.GetData();
+					const UObject* Object = SubobjectData->GetObject();
+					ObjectMetaData.Add(Object, { FText::FromString(SubobjectData->GetDisplayString()) });
+				}
+			}
+			
+			OnHierarchyChangedDelegate.Broadcast();
+		}
+	}
+
+	FText FComponentHierarchySubobjectModel::GetSubobjectDisplayName(const FSoftObjectPath& ObjectPath) const
+	{
+		const FObjectMetaData* MetaData = ObjectMetaData.Find(ObjectPath);
+		return MetaData ? MetaData->DisplayLabel : FText::GetEmpty();
+	}
+
+	void FComponentHierarchySubobjectModel::ForEachRootSubobject(FName Category, TFunctionRef<EBreakBehavior(const FSoftObjectPath& Object)> Callback)
+	{
+		const AActor* Actor = GetTopLevelObjectAsActor();
+		if (!Actor)
+		{
+			return;
+		}
+
+		if (Category == SceneComponentsCategory)
+		{
+			if (Actor->GetRootComponent())
+			{
+				Callback(Actor->GetRootComponent());
+			}
+		}
+		else if (Category == ActorComponentsCategory)
+		{
+			for (const UActorComponent* Component : Actor->GetComponents())
+			{
+				if (Private::IsValidComponent(*Component)
+					&& !Component->IsA<USceneComponent>()
+					&& Callback(Component) == EBreakBehavior::Break)
+				{
+					break;
+				}
+			}
+		}
+	}
+
+	void FComponentHierarchySubobjectModel::ForEachDirectChildSubobject(const FSoftObjectPath& Parent, TFunctionRef<EBreakBehavior(const FSoftObjectPath&)> Callback)
+	{
+		const UObject* Object = Parent.ResolveObject();
+		const UObject* TopLevel = GetTopLevelObjectAsActor();
+		if (!Object || !ensure(TopLevel && Object->IsIn(TopLevel)))
+		{
+			return;
+		}
+		
+		const USceneComponent* SceneComponent = Cast<USceneComponent>(Object);
+		if (!SceneComponent)
+		{
+			return;
+		}
+
+		for (int32 i = 0; i < SceneComponent->GetNumChildrenComponents(); ++i)
+		{
+			const USceneComponent* ChildComponent = SceneComponent->GetChildComponent(i);
+			if (Private::IsValidComponent(*ChildComponent)
+				&& Callback(ChildComponent) == EBreakBehavior::Break)
+			{
+				break;
+			}
+		}
+	}
+
+	AActor* FComponentHierarchySubobjectModel::GetTopLevelObjectAsActor() const
+	{
+		return Cast<AActor>(TopLevelObject.ResolveObject());
+	}
+}
diff --git a/Engine/Plugins/Developer/Concert/ConcertUI/ConcertClientSharedSlate/Source/ConcertClientSharedSlate/Private/Replication/Editor/Model/Subobject/ComponentHierarchySubobjectModel.h b/Engine/Plugins/Developer/Concert/ConcertUI/ConcertClientSharedSlate/Source/ConcertClientSharedSlate/Private/Replication/Editor/Model/Subobject/ComponentHierarchySubobjectModel.h
new file mode 100644
--- /dev/null
+++ b/Engine/Plugins/Developer/Concert/ConcertUI/ConcertClientSharedSlate/Source/ConcertClientSharedSlate/Private/Replication/Editor/Model/Subobject/ComponentHierarchySubobjectModel.h
@@ -0,0 +1,53 @@
+// Copyright Epic Games, Inc. All Rights Reserved.
+
+#pragma once
+
+#include "ISubobjectModel.h"
+#include "UObject/SoftObjectPtr.h"
+
+class AActor;
+
+namespace UE::ConcertClientSharedSlate
+{
+	/**
+	 * Builds a similar tree hierarchy as SSubobjectEditor.
+	 * This model only inspects UActorComponents on the root object, which must be an actor.
+	 */
+	class FComponentHierarchySubobjectModel : public ISubobjectModel
+	{
+	public:
+
+		/** Instances of USceneComponent are grouped to this category. */
+		static const FName SceneComponentsCategory;
+		/** Instances of UActorComponent which are not USceneComponents are grouped to this category. */
+		static const FName ActorComponentsCategory;
+
+		//~ Begin ISubobjectModel Interface
+		virtual void SetTopLevelObject(const FSoftObjectPath& InTopLevelObject) override;
+		virtual FSoftObjectPath GetTopLevelObject() const override { return TopLevelObject; }
+		virtual TArray<FName> GetCategories() const override { return { SceneComponentsCategory, ActorComponentsCategory }; }
+		virtual FText GetSubobjectDisplayName(const FSoftObjectPath& ObjectPath) const override;
+		virtual void ForEachRootSubobject(FName Category, TFunctionRef<EBreakBehavior(const FSoftObjectPath& Object)> Callback) override;
+		virtual void ForEachDirectChildSubobject(const FSoftObjectPath& Parent, TFunctionRef<EBreakBehavior(const FSoftObjectPath&)> Callback) override;
+		virtual FOnHierarchyChanged& OnHierarchyChanged() override { return OnHierarchyChangedDelegate; }
+		//~ End ISubobjectModel Interface
+
+	private:
+
+		/** The object from which the hierarchy is generated. */
+		FSoftObjectPath TopLevelObject;
+
+		struct FObjectMetaData
+		{
+			FText DisplayLabel;
+		};
+		TMap<FSoftObjectPath,FObjectMetaData> ObjectMetaData;
+
+		/** Called when the hierarchy has changed, e.g. due to calling SetTopLevelObject. */
+		FOnHierarchyChanged OnHierarchyChangedDelegate;
+
+		AActor* GetTopLevelObjectAsActor() const;
+	};
+}
+
+
diff --git a/Engine/Plugins/Developer/Concert/ConcertUI/ConcertClientSharedSlate/Source/ConcertClientSharedSlate/Private/Replication/Editor/Model/Subobject/ISubobjectModel.h b/Engine/Plugins/Developer/Concert/ConcertUI/ConcertClientSharedSlate/Source/ConcertClientSharedSlate/Private/Replication/Editor/Model/Subobject/ISubobjectModel.h
new file mode 100644
--- /dev/null
+++ b/Engine/Plugins/Developer/Concert/ConcertUI/ConcertClientSharedSlate/Source/ConcertClientSharedSlate/Private/Replication/Editor/Model/Subobject/ISubobjectModel.h
@@ -0,0 +1,50 @@
+// Copyright Epic Games, Inc. All Rights Reserved.
+
+#pragma once
+
+#include "Delegates/Delegate.h"
+#include "Templates/Function.h"
+
+enum class EBreakBehavior : uint8;
+struct FSoftObjectPath;
+
+namespace UE::ConcertClientSharedSlate
+{
+	/**
+	 * Decides on the contents for the subobject tree view.
+	 */
+	class ISubobjectModel
+	{
+	public:
+
+		/** Sets the object for which to build the subobject hierarchy. */
+		virtual void SetTopLevelObject(const FSoftObjectPath& TopLevelObject) = 0;
+		virtual FSoftObjectPath GetTopLevelObject() const = 0;
+
+		/**
+		 * Gets the categories that can be passed to ForEachRootSubobject.
+		 * Objects are grouped by category. The categories are separated by separator lines.
+		 * Example: One category could be the the USceneComponent hierarchy and another could be all other UActorComponents (like in the SSubobjectEditor).
+		 */
+		virtual TArray<FName> GetCategories() const = 0;
+
+		/** Gets the display name for the subobject */
+		virtual FText GetSubobjectDisplayName(const FSoftObjectPath& ObjectPath) const = 0;
+
+		/**
+		 * Gets the direct subobjects of the top level objects.
+		 * @param Category The category for the objects.
+		 * @param Callback Callback to invoke for each found root object.
+		 * @see GetCategories for valid categories.
+		 */
+		virtual void ForEachRootSubobject(FName Category, TFunctionRef<EBreakBehavior(const FSoftObjectPath& Object)> Callback) = 0;
+		/** Gets the direct subobject children of another subobject. Child subobjects are implicitly in the same category as Parent. */
+		virtual void ForEachDirectChildSubobject(const FSoftObjectPath& Parent, TFunctionRef<EBreakBehavior(const FSoftObjectPath& Object)> Callback) = 0;
+
+		DECLARE_MULTICAST_DELEGATE(FOnHierarchyChanged)
+		/** Called when the hierarchy has changed, e.g. due to calling SetTopLevelObject. */
+		virtual FOnHierarchyChanged& OnHierarchyChanged() = 0;
+
+		virtual ~ISubobjectModel() = default;
+	};
+}
diff --git a/Engine/Plugins/Developer/Concert/ConcertUI/ConcertClientSharedSlate/Source/ConcertClientSharedSlate/Private/Replication/Editor/View/ObjectEditor/SBaseReplicationStreamEditor.cpp b/Engine/Plugins/Developer/Concert/ConcertUI/ConcertClientSharedSlate/Source/ConcertClientSharedSlate/Private/Replication/Editor/View/ObjectEditor/SBaseReplicationStreamEditor.cpp
--- a/Engine/Plugins/Developer/Concert/ConcertUI/ConcertClientSharedSlate/Source/ConcertClientSharedSlate/Private/Replication/Editor/View/ObjectEditor/SBaseReplicationStreamEditor.cpp
+++ b/Engine/Plugins/Developer/Concert/ConcertUI/ConcertClientSharedSlate/Source/ConcertClientSharedSlate/Private/Replication/Editor/View/ObjectEditor/SBaseReplicationStreamEditor.cpp
@@ -36,6 +36,8 @@ namespace UE::ConcertClientSharedSlate
 		EditablePropertiesModel->OnPropertiesChanged().AddSP(this, &SBaseReplicationStreamEditor::OnPropertiesChanged);
 		PropertiesModelAdapter = MakeShared<FFakeObjectToPropertiesEditorModel>(EditablePropertiesModel.ToSharedRef(), PropertySelectionSource.ToSharedRef());
 
+		IsEditingEnabledAttribute = InArgs._IsEditingEnabled;
+		EditingDisabledToolTipTextAttribute = InArgs._EditingDisabledToolTipText;
 		OnExtendObjectsContextMenuDelegate = InArgs._OnExtendObjectsContextMenu;
 		
 		ChildSlot
@@ -85,6 +87,19 @@ namespace UE::ConcertClientSharedSlate
 		return ReplicationViewer->GetObjectsBeingPropertyEdited();
 	}
 
+	bool SBaseReplicationStreamEditor::IsEditingDisabled() const
+	{
+		return (IsEditingEnabledAttribute.IsBound() || IsEditingEnabledAttribute.IsSet())
+			&& !IsEditingEnabledAttribute.Get();
+	}
+
+	FText SBaseReplicationStreamEditor::GetEditingDisabledText() const
+	{
+		return (EditingDisabledToolTipTextAttribute.IsBound() || EditingDisabledToolTipTextAttribute.IsSet())
+			? EditingDisabledToolTipTextAttribute.Get()
+			: FText::GetEmpty();
+	}
+
 	void SBaseReplicationStreamEditor::OnObjectsChanged(TConstArrayView<UObject*> AddedObjects, TConstArrayView<FSoftObjectPath> RemovedObjects, EReplicatedObjectChangeReason ChangeReason)
 	{
 		ReplicationViewer->RefreshObjectData();
@@ -127,6 +142,11 @@ namespace UE::ConcertClientSharedSlate
 
 	void SBaseReplicationStreamEditor::OnDeleteObjects(const TArray<TSharedPtr<FReplicatedObjectData>>& ObjectsToDelete) const
 	{
+		if (IsEditingDisabled())
+		{
+			return;
+		}
+		
 		TArray<FSoftObjectPath> DeleteObjectPaths;
 		Algo::Transform(ObjectsToDelete, DeleteObjectPaths, [](const TSharedPtr<FReplicatedObjectData>& Data) { return Data->GetObjectPath(); });
 
@@ -156,9 +176,12 @@ namespace UE::ConcertClientSharedSlate
 		AddObjectSourceContextMenuOptions(MenuBuilder);
 		MenuBuilder.AddMenuEntry(
 				LOCTEXT("DeleteItems", "Delete"),
-				FText::GetEmpty(),
+				TAttribute<FText>::CreateLambda([this](){ return GetEditingDisabledText(); }),
 				FSlateIcon(),
-				FUIAction(FExecuteAction::CreateSP(this, &SBaseReplicationStreamEditor::OnDeleteObjects_PassByValue, ReplicationViewer->GetSelectedOutlinerObjects())),
+				FUIAction(
+					FExecuteAction::CreateSP(this, &SBaseReplicationStreamEditor::OnDeleteObjects_PassByValue, ReplicationViewer->GetSelectedOutlinerObjects()),
+					FCanExecuteAction::CreateLambda([this]() { return !IsEditingDisabled(); })
+					),
 				NAME_None,
 				EUserInterfaceActionType::Button
 			);
@@ -207,22 +230,25 @@ namespace UE::ConcertClientSharedSlate
 	ConcertSharedSlate::FSourceModelBuilders<FSelectableObjectInfo>::FItemPickerArgs SBaseReplicationStreamEditor::MakeObjectSourceBuilderArgs() const
 	{
 		using namespace ConcertSharedSlate;
+		using FBuilderDelegates = FSourceModelBuilders<FSelectableObjectInfo>;
 		
 		return FSourceModelBuilders<FSelectableObjectInfo>::FItemPickerArgs
 		{
-			FSourceModelBuilders<FSelectableObjectInfo>::FOnItemsSelected::CreateSP(this, &SBaseReplicationStreamEditor::OnObjectsSelectedForAdding),
-			FSourceModelBuilders<FSelectableObjectInfo>::FGetItemDisplayString::CreateLambda([](const FSelectableObjectInfo& Item)
+			FBuilderDelegates::FOnItemsSelected::CreateSP(this, &SBaseReplicationStreamEditor::OnObjectsSelectedForAdding),
+			FBuilderDelegates::FGetItemDisplayString::CreateLambda([](const FSelectableObjectInfo& Item)
 			{
 				return Item.Object.IsValid() ? DisplayUtils::GetObjectDisplayString(*Item.Object.Get()) : TEXT("");
 			}),
-			FSourceModelBuilders<FSelectableObjectInfo>::FGetItemIcon::CreateLambda([](const FSelectableObjectInfo& Item)
+			FBuilderDelegates::FGetItemIcon::CreateLambda([](const FSelectableObjectInfo& Item)
 			{
 				return Item.Object.IsValid() ? DisplayUtils::GetObjectIcon(*Item.Object.Get()) : FSlateIcon();
 			}),
-			FSourceModelBuilders<FSelectableObjectInfo>::FIsItemSelected::CreateLambda([this](const FSelectableObjectInfo& Item)
+			FBuilderDelegates::FIsItemSelected::CreateLambda([this](const FSelectableObjectInfo& Item)
 			{
 				return EditablePropertiesModel->ContainsObjects({ Item.Object.Get() } );
-			})
+			}),
+			TAttribute<bool>::CreateLambda([this]() { return !IsEditingDisabled(); }),
+			TAttribute<FText>::CreateLambda([this]() { return GetEditingDisabledText(); })
 		}; 
 	}
 }
diff --git a/Engine/Plugins/Developer/Concert/ConcertUI/ConcertClientSharedSlate/Source/ConcertClientSharedSlate/Private/Replication/Editor/View/ObjectEditor/SBaseReplicationStreamEditor.h b/Engine/Plugins/Developer/Concert/ConcertUI/ConcertClientSharedSlate/Source/ConcertClientSharedSlate/Private/Replication/Editor/View/ObjectEditor/SBaseReplicationStreamEditor.h
--- a/Engine/Plugins/Developer/Concert/ConcertUI/ConcertClientSharedSlate/Source/ConcertClientSharedSlate/Private/Replication/Editor/View/ObjectEditor/SBaseReplicationStreamEditor.h
+++ b/Engine/Plugins/Developer/Concert/ConcertUI/ConcertClientSharedSlate/Source/ConcertClientSharedSlate/Private/Replication/Editor/View/ObjectEditor/SBaseReplicationStreamEditor.h
@@ -50,9 +50,9 @@ namespace UE::ConcertClientSharedSlate
 		SLATE_BEGIN_ARGS(SBaseReplicationStreamEditor)
 		{}
 			/** Additional columns to add to the object view */
-			SLATE_ARGUMENT(TArray<ReplicationObjectColumns::FReplicationObjectColumn>, AdditionalObjectColumns)
+			SLATE_ARGUMENT(TArray<ReplicationColumns::FReplicationTopLevelObjectColumn>, AdditionalObjectColumns)
 			/** Additional columns to add to the property view */
-			SLATE_ARGUMENT(TArray<ReplicationPropertyColumns::FReplicationPropertyColumn>, AdditionalPropertyColumns)
+			SLATE_ARGUMENT(TArray<ReplicationColumns::FReplicationPropertyColumn>, AdditionalPropertyColumns)
 		
 			/** Optional. Placed between root object outliner and property editor. */
 			SLATE_ARGUMENT(TSharedPtr<IReplicationSubobjectView>, SubobjectView)
@@ -67,6 +67,11 @@ namespace UE::ConcertClientSharedSlate
 			SLATE_NAMED_SLOT(FArguments, LeftOfObjectSearchBar)
 			/** Optional widget to add to the left of the property list search bar. */
 			SLATE_NAMED_SLOT(FArguments, LeftOfPropertySearchBar)
+
+			/** Optional. Determines whether all UI for changing the model should be disabled. */
+			SLATE_ATTRIBUTE(bool, IsEditingEnabled)
+			/** Optional. Whenever IsEditingEnabled returns true, this tooltip is displayed for relevant, disabled UI. */
+			SLATE_ATTRIBUTE(FText, EditingDisabledToolTipText)
 		SLATE_END_ARGS()
 
 		void Construct(const FArguments& InArgs,
@@ -99,8 +104,16 @@ namespace UE::ConcertClientSharedSlate
 		/** For deciding which properties can be added to EditablePropertiesModel. */
 		TSharedPtr<IPropertySelectionSourceModel> PropertySelectionSource;
 
+		/** Optional. Determines whether all UI for changing the model should be disabled. */
+		TAttribute<bool> IsEditingEnabledAttribute;
+		/** Optional. Whenever IsEditingEnabled returns true, this tooltip is displayed for relevant, disabled UI. */
+		TAttribute<FText> EditingDisabledToolTipTextAttribute;
+		
 		/** Generates additional entries */
 		FExtendMenu OnExtendObjectsContextMenuDelegate;
+
+		bool IsEditingDisabled() const;
+		FText GetEditingDisabledText() const;
 		
 		void OnObjectsChanged(TConstArrayView<UObject*> AddedObjects, TConstArrayView<FSoftObjectPath> RemovedObjects, EReplicatedObjectChangeReason ChangeReason);
 		void OnPropertiesChanged();
diff --git a/Engine/Plugins/Developer/Concert/ConcertUI/ConcertClientSharedSlate/Source/ConcertClientSharedSlate/Private/Replication/Editor/View/ObjectEditor/SDefaultReplicationStreamEditor.cpp b/Engine/Plugins/Developer/Concert/ConcertUI/ConcertClientSharedSlate/Source/ConcertClientSharedSlate/Private/Replication/Editor/View/ObjectEditor/SDefaultReplicationStreamEditor.cpp
--- a/Engine/Plugins/Developer/Concert/ConcertUI/ConcertClientSharedSlate/Source/ConcertClientSharedSlate/Private/Replication/Editor/View/ObjectEditor/SDefaultReplicationStreamEditor.cpp
+++ b/Engine/Plugins/Developer/Concert/ConcertUI/ConcertClientSharedSlate/Source/ConcertClientSharedSlate/Private/Replication/Editor/View/ObjectEditor/SDefaultReplicationStreamEditor.cpp
@@ -17,8 +17,17 @@ namespace UE::ConcertClientSharedSlate
 	{
 		PropertiesModel = InPropertiesModel;
 		
-		using namespace ReplicationPropertyColumns;
-		const FReplicationPropertyColumn ReplicatesColumn = ReplicatesColumns(SharedThis(this), InPropertiesModel);
+		using namespace ReplicationColumns;
+		using namespace ReplicationColumns::Property;
+		const FReplicationPropertyColumn ReplicatesColumn = ReplicatesColumns(
+			SharedThis(this),
+			InPropertiesModel,
+			TReplicationColumnDelegates<FReplicatedPropertyData>::FIsEnabled::CreateLambda([IsEnabled = InArgs._IsEditingEnabled](const FReplicatedPropertyData&)
+			{
+				return !IsEnabled.IsBound() || IsEnabled.Get();
+			}),
+			InArgs._EditingDisabledToolTipText
+			);
 		TArray<FReplicationPropertyColumn> PropertyColumns = InArgs._AdditionalPropertyColumns;
 		PropertyColumns.Add(ReplicatesColumn);
 		
@@ -38,6 +47,8 @@ namespace UE::ConcertClientSharedSlate
 				[
 					InArgs._LeftOfPropertySearchBar.Widget
 				]
+				.IsEditingEnabled(InArgs._IsEditingEnabled)
+				.EditingDisabledToolTipText(InArgs._EditingDisabledToolTipText)
 		];
 	}
 
@@ -61,7 +72,7 @@ namespace UE::ConcertClientSharedSlate
 		const FReplicatedPropertyData& Right
 		) const
 	{
-		return ReplicationPropertyColumns::SortBySelectionThenByName_PropertyPredicate(
+		return ReplicationColumns::Property::SortBySelectionThenByName_PropertyPredicate(
 			WrappedEditor->GetObjectsBeingPropertyEdited(),
 			*PropertiesModel,
 			Left,
diff --git a/Engine/Plugins/Developer/Concert/ConcertUI/ConcertClientSharedSlate/Source/ConcertClientSharedSlate/Private/Replication/Editor/View/ObjectEditor/SDefaultReplicationStreamEditor.h b/Engine/Plugins/Developer/Concert/ConcertUI/ConcertClientSharedSlate/Source/ConcertClientSharedSlate/Private/Replication/Editor/View/ObjectEditor/SDefaultReplicationStreamEditor.h
--- a/Engine/Plugins/Developer/Concert/ConcertUI/ConcertClientSharedSlate/Source/ConcertClientSharedSlate/Private/Replication/Editor/View/ObjectEditor/SDefaultReplicationStreamEditor.h
+++ b/Engine/Plugins/Developer/Concert/ConcertUI/ConcertClientSharedSlate/Source/ConcertClientSharedSlate/Private/Replication/Editor/View/ObjectEditor/SDefaultReplicationStreamEditor.h
@@ -32,9 +32,9 @@ namespace UE::ConcertClientSharedSlate
 		SLATE_BEGIN_ARGS(SDefaultReplicationStreamEditor)
 		{}
 			/** Additional columns to add to the object view */
-			SLATE_ARGUMENT(TArray<ReplicationObjectColumns::FReplicationObjectColumn>, AdditionalObjectColumns)
+			SLATE_ARGUMENT(TArray<ReplicationColumns::FReplicationTopLevelObjectColumn>, AdditionalObjectColumns)
 			/** Additional columns to add to the property view */
-			SLATE_ARGUMENT(TArray<ReplicationPropertyColumns::FReplicationPropertyColumn>, AdditionalPropertyColumns)
+			SLATE_ARGUMENT(TArray<ReplicationColumns::FReplicationPropertyColumn>, AdditionalPropertyColumns)
 			
 			/** Optional. Placed between root object outliner and property editor. */
 			SLATE_ARGUMENT(TSharedPtr<IReplicationSubobjectView>, SubobjectView)
@@ -49,6 +49,11 @@ namespace UE::ConcertClientSharedSlate
 			SLATE_NAMED_SLOT(FArguments, LeftOfObjectSearchBar)
 			/** Optional widget to add to the left of the property list search bar. */
 			SLATE_NAMED_SLOT(FArguments, LeftOfPropertySearchBar)
+
+			/** Optional. Determines whether all UI for changing the model should be disabled. */
+			SLATE_ATTRIBUTE(bool, IsEditingEnabled)
+			/** Optional. Whenever IsEditingEnabled returns true, this tooltip is displayed for relevant, disabled UI. */
+			SLATE_ATTRIBUTE(FText, EditingDisabledToolTipText)
 		SLATE_END_ARGS()
 
 		void Construct(
diff --git a/Engine/Plugins/Developer/Concert/ConcertUI/ConcertClientSharedSlate/Source/ConcertClientSharedSlate/Private/Replication/Editor/View/ObjectViewer/SReplicationStreamViewer.cpp b/Engine/Plugins/Developer/Concert/ConcertUI/ConcertClientSharedSlate/Source/ConcertClientSharedSlate/Private/Replication/Editor/View/ObjectViewer/SReplicationStreamViewer.cpp
--- a/Engine/Plugins/Developer/Concert/ConcertUI/ConcertClientSharedSlate/Source/ConcertClientSharedSlate/Private/Replication/Editor/View/ObjectViewer/SReplicationStreamViewer.cpp
+++ b/Engine/Plugins/Developer/Concert/ConcertUI/ConcertClientSharedSlate/Source/ConcertClientSharedSlate/Private/Replication/Editor/View/ObjectViewer/SReplicationStreamViewer.cpp
@@ -103,11 +103,6 @@ namespace UE::ConcertClientSharedSlate
 	{
 		SubobjectAndPropertySection->RefreshPropertyData();
 	}
-	
-	const TArray<TSharedPtr<FReplicatedPropertyData>>& SReplicationStreamViewer::GetPropertyRowData() const
-	{
-		return SubobjectAndPropertySection->GetPropertyRowData();
-	}
 
 	TSharedRef<FReplicatedObjectData> SReplicationStreamViewer::AllocateObjectData(FSoftObjectPath ObjectPath)
 	{
@@ -135,11 +130,11 @@ namespace UE::ConcertClientSharedSlate
 
 	TSharedRef<SWidget> SReplicationStreamViewer::CreateActorsSection(const FArguments& InArgs)
 	{
-		TArray<ReplicationObjectColumns::FReplicationObjectColumn> Columns
+		TArray Columns
 		{
-			ReplicationObjectColumns::IconColumn(PropertiesModel.ToSharedRef()),
-			ReplicationObjectColumns::LabelColumn(),
-			ReplicationObjectColumns::TypeColumn(PropertiesModel.ToSharedRef())
+			ReplicationColumns::TopLevel::IconColumn(PropertiesModel.ToSharedRef()),
+			ReplicationColumns::TopLevel::LabelColumn(),
+			ReplicationColumns::TopLevel::TypeColumn(PropertiesModel.ToSharedRef())
 		};
 		Columns.Append(InArgs._AdditionalObjectColumns);
 		
@@ -155,7 +150,7 @@ namespace UE::ConcertClientSharedSlate
 				SubobjectAndPropertySection->SelectRootObjects();
 			})
 			.Columns(Columns)
-			.ExpandableColumnLabel(ReplicationObjectColumns::LabelColumnId)
+			.ExpandableColumnLabel(ReplicationColumns::TopLevel::LabelColumnId)
 			.SelectionMode(ESelectionMode::Multi)
 			.LeftOfSearchBar()
 			[
diff --git a/Engine/Plugins/Developer/Concert/ConcertUI/ConcertClientSharedSlate/Source/ConcertClientSharedSlate/Private/Replication/Editor/View/ObjectViewer/SReplicationStreamViewer.h b/Engine/Plugins/Developer/Concert/ConcertUI/ConcertClientSharedSlate/Source/ConcertClientSharedSlate/Private/Replication/Editor/View/ObjectViewer/SReplicationStreamViewer.h
--- a/Engine/Plugins/Developer/Concert/ConcertUI/ConcertClientSharedSlate/Source/ConcertClientSharedSlate/Private/Replication/Editor/View/ObjectViewer/SReplicationStreamViewer.h
+++ b/Engine/Plugins/Developer/Concert/ConcertUI/ConcertClientSharedSlate/Source/ConcertClientSharedSlate/Private/Replication/Editor/View/ObjectViewer/SReplicationStreamViewer.h
@@ -46,9 +46,9 @@ namespace UE::ConcertClientSharedSlate
 		SLATE_BEGIN_ARGS(SReplicationStreamViewer)
 		{}
 			/** Additional columns to add to the object view */
-			SLATE_ARGUMENT(TArray<ReplicationObjectColumns::FReplicationObjectColumn>, AdditionalObjectColumns)
+			SLATE_ARGUMENT(TArray<ReplicationColumns::FReplicationTopLevelObjectColumn>, AdditionalObjectColumns)
 			/** Additional columns to add to the property view */
-			SLATE_ARGUMENT(TArray<ReplicationPropertyColumns::FReplicationPropertyColumn>, AdditionalPropertyColumns)
+			SLATE_ARGUMENT(TArray<ReplicationColumns::FReplicationPropertyColumn>, AdditionalPropertyColumns)
 
 			/** Optional. Placed between root object outliner and property editor. */
 			SLATE_ARGUMENT(TSharedPtr<IReplicationSubobjectView>, SubobjectView)
@@ -80,9 +80,6 @@ namespace UE::ConcertClientSharedSlate
 		void RefreshSubobjectData();
 		void RefreshPropertyData();
 
-		const TSet<TSharedPtr<FReplicatedObjectData>>& GetObjectRowData() const { return ObjectRowData; }
-		const TArray<TSharedPtr<FReplicatedPropertyData>>& GetPropertyRowData() const;
-
 		/** @return Gets the root objects selected in the outliner; the subobject view chooses which of these objects (or their subobjects) end up in GetSelectedObjectShowingProperties. */
 		TArray<TSharedPtr<FReplicatedObjectData>> GetSelectedOutlinerObjects() const { return ReplicatedObjects->GetSelectedItems(); }
 		
@@ -102,10 +99,9 @@ namespace UE::ConcertClientSharedSlate
 		/**
 		 * These instances can be subclasses of FReplicatedObjectData, e.g. FReplicatedObjectData_Editor.
 		 * Their type can be overridden by subclasses.
-		 * They only have the FReplicatedObjectData type so they can be passed efficiently to SObjectToPropertyView.
 		 * @see GetObjectData
 		 */
-		TSet<TSharedPtr<FReplicatedObjectData>> ObjectRowData;
+		TArray<TSharedPtr<FReplicatedObjectData>> ObjectRowData;
 		/** The instances of ObjectRowData which do not have any parents. This acts as the item source for the tree view. */
 		TArray<TSharedPtr<FReplicatedObjectData>> RootObjectRowData;
 		/** Inverse map of ObjectRowData using FReplicatedObjectData::GetObjectPath as key. Contains all elements of ObjectRowData. */
diff --git a/Engine/Plugins/Developer/Concert/ConcertUI/ConcertClientSharedSlate/Source/ConcertClientSharedSlate/Private/Replication/Editor/View/ObjectViewer/SSubobjectAndPropertySection.cpp b/Engine/Plugins/Developer/Concert/ConcertUI/ConcertClientSharedSlate/Source/ConcertClientSharedSlate/Private/Replication/Editor/View/ObjectViewer/SSubobjectAndPropertySection.cpp
--- a/Engine/Plugins/Developer/Concert/ConcertUI/ConcertClientSharedSlate/Source/ConcertClientSharedSlate/Private/Replication/Editor/View/ObjectViewer/SSubobjectAndPropertySection.cpp
+++ b/Engine/Plugins/Developer/Concert/ConcertUI/ConcertClientSharedSlate/Source/ConcertClientSharedSlate/Private/Replication/Editor/View/ObjectViewer/SSubobjectAndPropertySection.cpp
@@ -38,7 +38,7 @@ namespace UE::ConcertClientSharedSlate
 	{
 		if (SubobjectView)
 		{
-			SubobjectView->SelectRootObjects();
+			SubobjectView->SelectTopLevelObjects();
 		}
 	}
 
@@ -52,7 +52,7 @@ namespace UE::ConcertClientSharedSlate
 		const TArray<TSharedPtr<FReplicatedObjectData>> SelectedObjects = GetSelectedRootObjectsDelegate.Execute();
 		TArray<FSoftObjectPath> SelectedObjectPaths;
 		Algo::Transform(SelectedObjects, SelectedObjectPaths, [](const TSharedPtr<FReplicatedObjectData>& Item){ return Item->GetObjectPath(); });
-		SubobjectView->SetRootObjects(SelectedObjectPaths);
+		SubobjectView->SetTopLevelObjects(SelectedObjectPaths);
 	}
 	
 	void SSubobjectAndPropertySection::RefreshPropertyData()
@@ -147,10 +147,10 @@ namespace UE::ConcertClientSharedSlate
 
 	TSharedRef<SWidget> SSubobjectAndPropertySection::CreatePropertiesView(const FArguments& InArgs)
 	{
-		TArray<ReplicationPropertyColumns::FReplicationPropertyColumn> Columns
+		TArray Columns
 		{
-			ReplicationPropertyColumns::LabelColumn(),
-			ReplicationPropertyColumns::TypeColumn()
+			ReplicationColumns::Property::LabelColumn(),
+			ReplicationColumns::Property::TypeColumn()
 		};
 		Columns.Append(InArgs._AdditionalPropertyColumns);
 		
@@ -165,7 +165,7 @@ namespace UE::ConcertClientSharedSlate
 				.RootItemsSource(&RootPropertyRowData)
 				.OnGetChildren(this, &SSubobjectAndPropertySection::GetPropertyRowChildren)
 				.Columns(Columns)
-				.ExpandableColumnLabel(ReplicationPropertyColumns::LabelColumnId)
+				.ExpandableColumnLabel(ReplicationColumns::Property::LabelColumnId)
 				.SelectionMode(ESelectionMode::Multi)
 				.LeftOfSearchBar()
 				[
diff --git a/Engine/Plugins/Developer/Concert/ConcertUI/ConcertClientSharedSlate/Source/ConcertClientSharedSlate/Private/Replication/Editor/View/ObjectViewer/SSubobjectAndPropertySection.h b/Engine/Plugins/Developer/Concert/ConcertUI/ConcertClientSharedSlate/Source/ConcertClientSharedSlate/Private/Replication/Editor/View/ObjectViewer/SSubobjectAndPropertySection.h
--- a/Engine/Plugins/Developer/Concert/ConcertUI/ConcertClientSharedSlate/Source/ConcertClientSharedSlate/Private/Replication/Editor/View/ObjectViewer/SSubobjectAndPropertySection.h
+++ b/Engine/Plugins/Developer/Concert/ConcertUI/ConcertClientSharedSlate/Source/ConcertClientSharedSlate/Private/Replication/Editor/View/ObjectViewer/SSubobjectAndPropertySection.h
@@ -32,7 +32,7 @@ namespace UE::ConcertClientSharedSlate
 		SLATE_BEGIN_ARGS(SSubobjectAndPropertySection)
 		{}
 			/** Additional columns to add to the property view */
-			SLATE_ARGUMENT(TArray<ReplicationPropertyColumns::FReplicationPropertyColumn>, AdditionalPropertyColumns)
+			SLATE_ARGUMENT(TArray<ReplicationColumns::FReplicationPropertyColumn>, AdditionalPropertyColumns)
 			/** Optional. Placed between root object outliner and property editor. */
 			SLATE_ARGUMENT(TSharedPtr<IReplicationSubobjectView>, SubobjectView)
 		
diff --git a/Engine/Plugins/Developer/Concert/ConcertUI/ConcertClientSharedSlate/Source/ConcertClientSharedSlate/Private/Replication/Editor/View/ObjectViewer/Tree/SReplicationColumnRow.h b/Engine/Plugins/Developer/Concert/ConcertUI/ConcertClientSharedSlate/Source/ConcertClientSharedSlate/Private/Replication/Editor/View/ObjectViewer/Tree/SReplicationColumnRow.h
--- a/Engine/Plugins/Developer/Concert/ConcertUI/ConcertClientSharedSlate/Source/ConcertClientSharedSlate/Private/Replication/Editor/View/ObjectViewer/Tree/SReplicationColumnRow.h
+++ b/Engine/Plugins/Developer/Concert/ConcertUI/ConcertClientSharedSlate/Source/ConcertClientSharedSlate/Private/Replication/Editor/View/ObjectViewer/Tree/SReplicationColumnRow.h
@@ -24,6 +24,7 @@ namespace UE::ConcertClientSharedSlate
 		DECLARE_DELEGATE_RetVal_OneParam(const TColumType*, FGetColumn,
 			const FName& ColumnId
 			);
+		DECLARE_DELEGATE_RetVal_TwoParams(TSharedPtr<SWidget>, FOverrideColumnWidget, const FName& ColumnName, const TListItemType& RowData);
 		
 		SLATE_BEGIN_ARGS(SReplicationColumnRow)
 			: _RowHeight(20.f)
@@ -34,6 +35,12 @@ namespace UE::ConcertClientSharedSlate
 			/** Gets columns info about a certain column */
 			SLATE_EVENT(FGetColumn, ColumnGetter)
 
+			/**
+			 * Optional. If the delegate returns non-null, that widget will be used instead of the one the column would generate.
+			 * This is useful, e.g. if you want to generate a separator widget between items.
+			 */
+			SLATE_EVENT(FOverrideColumnWidget, OverrideColumnWidget)
+
 			/** The data to pass to TReplicationColumn::BuildColumnWidget. */
 			SLATE_ARGUMENT(TSharedPtr<TListItemType>, RowData)
 		
@@ -49,6 +56,7 @@ namespace UE::ConcertClientSharedSlate
 			TSharedRef<STableViewBase> InOwner)
 		{
 			ColumnGetterDelegate = InArgs._ColumnGetter;
+			OverrideColumnWidgetDelegate = InArgs._OverrideColumnWidget;
 			HighlightText = InArgs._HighlightText;
 			RowData = InArgs._RowData;
 			ExpandableColumnLabel = InArgs._ExpandableColumnLabel;
@@ -60,6 +68,14 @@ namespace UE::ConcertClientSharedSlate
 		/** Generates the widget representing this row. */
 		virtual TSharedRef<SWidget> GenerateWidgetForColumn(const FName& ColumnName) override
 		{
+			const TSharedPtr<SWidget> ColumnOverride = OverrideColumnWidgetDelegate.IsBound()
+				? OverrideColumnWidgetDelegate.Execute(ColumnName, *RowData.Get())
+				: nullptr;
+			if (ColumnOverride)
+			{
+				return ColumnOverride.ToSharedRef();
+			}
+			
 			const TColumType* Column = ColumnGetterDelegate.Execute(ColumnName); ensure(Column);
 			if (!Column)
 			{
@@ -98,6 +114,7 @@ namespace UE::ConcertClientSharedSlate
 	private:
 		
 		FGetColumn ColumnGetterDelegate;
+		FOverrideColumnWidget OverrideColumnWidgetDelegate;
 		TSharedPtr<FText> HighlightText;
 		TSharedPtr<TListItemType> RowData;
 		FName ExpandableColumnLabel;
diff --git a/Engine/Plugins/Developer/Concert/ConcertUI/ConcertClientSharedSlate/Source/ConcertClientSharedSlate/Private/Replication/Editor/View/ObjectViewer/Tree/SReplicationTreeView.h b/Engine/Plugins/Developer/Concert/ConcertUI/ConcertClientSharedSlate/Source/ConcertClientSharedSlate/Private/Replication/Editor/View/ObjectViewer/Tree/SReplicationTreeView.h
--- a/Engine/Plugins/Developer/Concert/ConcertUI/ConcertClientSharedSlate/Source/ConcertClientSharedSlate/Private/Replication/Editor/View/ObjectViewer/Tree/SReplicationTreeView.h
+++ b/Engine/Plugins/Developer/Concert/ConcertUI/ConcertClientSharedSlate/Source/ConcertClientSharedSlate/Private/Replication/Editor/View/ObjectViewer/Tree/SReplicationTreeView.h
@@ -26,13 +26,17 @@ namespace UE::ConcertClientSharedSlate
 	{
 	public:
 
+		using TOverrideColumnWidget = typename SReplicationColumnRow<TItemType>::FOverrideColumnWidget;
+
 		DECLARE_DELEGATE_OneParam(FDeleteItems, const TArray<TSharedPtr<TItemType>>& SelectedItems);
 		DECLARE_DELEGATE_TwoParams(FGetItemChildren, TSharedPtr<TItemType> Item, TFunctionRef<void(TSharedPtr<TItemType>)> ProcessChild);
 		DECLARE_DELEGATE(FOnSelectionChanged);
 		DECLARE_DELEGATE_RetVal_OneParam(bool, FCustomFilter, const TSharedPtr<TItemType>& Item);
+		DECLARE_DELEGATE_RetVal_OneParam(bool, FIsSearchableItem, const TSharedPtr<TItemType>& Item);
 
 		SLATE_BEGIN_ARGS(SReplicationTreeView<TItemType>)
-			: _SelectionMode(ESelectionMode::Single)
+			: _HeaderRowVisibility(EVisibility::Visible)
+			, _SelectionMode(ESelectionMode::Single)
 		{}
 			/** The items to display */
 			SLATE_ARGUMENT(TArray<TSharedPtr<TItemType>>*, RootItemsSource)
@@ -51,11 +55,21 @@ namespace UE::ConcertClientSharedSlate
 
 			/** Optional callback to do even more filtering of items. */
 			SLATE_EVENT(FCustomFilter, FilterItem)
+		
+			/**
+			 * Optional. If the delegate returns non-null, that widget will be used instead of the one the column would generate.
+			 * This is useful, e.g. if you want to generate a separator widget between items.
+			 */
+			SLATE_EVENT(TOverrideColumnWidget, OverrideColumnWidget)
+			/** Optional callback for determining whether this item can be searched. */
+			SLATE_EVENT(FIsSearchableItem, IsSearchableItem)
 			
 			/** The columns this list should have */
 			SLATE_ARGUMENT(TArray<TReplicationColumn<TItemType>>, Columns)
 			/** The name of the column that will have the SExpanderArrow for the tree view. */
 			SLATE_ARGUMENT(FName, ExpandableColumnLabel)
+			/** Visibility of the header row */
+			SLATE_ARGUMENT(EVisibility, HeaderRowVisibility)
 		
 			/** How many items are to allowed to be selected */
 			SLATE_ARGUMENT(ESelectionMode::Type, SelectionMode)
@@ -74,6 +88,8 @@ namespace UE::ConcertClientSharedSlate
 			OnGetChildrenDelegate = InArgs._OnGetChildren;
 			OnDeleteItemsDelegate = InArgs._OnDeleteItems;
 			CustomFilterDelegate = InArgs._FilterItem;
+			OverrideColumnWidget = InArgs._OverrideColumnWidget;
+			IsSearchableItemDelegate = InArgs._IsSearchableItem;
 			ExpandableColumnId = InArgs._ExpandableColumnLabel;
 			
 			SearchText = MakeShared<FText>();
@@ -180,6 +196,10 @@ namespace UE::ConcertClientSharedSlate
 		FDeleteItems OnDeleteItemsDelegate;
 		/** Optional delegate for filtering the items even more. */
 		FCustomFilter CustomFilterDelegate;
+		/** Optional delegate for overriding the column widgets. */
+		TOverrideColumnWidget OverrideColumnWidget;
+		/** Optional callback for determining whether this item can be filtered. If false, it will not be shown when searched. */
+		FIsSearchableItem IsSearchableItemDelegate;
 		
 
 		TSharedRef<SWidget> CreateTreeView(const FArguments& InArgs);
@@ -233,7 +253,7 @@ namespace UE::ConcertClientSharedSlate
 		TArray<TReplicationColumn<TItemType>> Columns = InArgs._Columns;
 		Columns.Sort([](const TReplicationColumn<TItemType>& Left, const TReplicationColumn<TItemType>& Right) { return Left.GetColumnSortOrderValue() < Right.GetColumnSortOrderValue(); });
 		
-		HeaderRow = SNew(SHeaderRow);
+		HeaderRow = SNew(SHeaderRow).Visibility(InArgs._HeaderRowVisibility);
 		TSet<FName> DuplicateColumnDetection;
 		for (TReplicationColumn<TItemType>& Column : Columns)
 		{
@@ -268,6 +288,7 @@ namespace UE::ConcertClientSharedSlate
 		return SNew(SReplicationColumnRow<TItemType>, OwnerTable)
 			.HighlightText(SearchText)
 			.ColumnGetter(ColumnGetter)
+			.OverrideColumnWidget(OverrideColumnWidget)
 			.RowData(Item)
 			.ExpandableColumnLabel(ExpandableColumnId);
 	}
@@ -307,6 +328,11 @@ namespace UE::ConcertClientSharedSlate
 	template <typename TItemType>
 	void SReplicationTreeView<TItemType>::PopulateSearchStrings(const TSharedPtr<TItemType>& Item, TArray<FString>& OutSearchStrings)
 	{
+		if (IsSearchableItemDelegate.IsBound() && !IsSearchableItemDelegate.Execute(Item))
+		{
+			return;
+		}
+	
 		for (const SHeaderRow::FColumn& Column : HeaderRow->GetColumns())
 		{
 			const TReplicationColumn<TItemType>& CastColumn = static_cast<const TReplicationColumn<TItemType>&>(Column);
diff --git a/Engine/Plugins/Developer/Concert/ConcertUI/ConcertClientSharedSlate/Source/ConcertClientSharedSlate/Private/Replication/Editor/View/ObjectViewer/Tree/SelectionViewerColumns.cpp b/Engine/Plugins/Developer/Concert/ConcertUI/ConcertClientSharedSlate/Source/ConcertClientSharedSlate/Private/Replication/Editor/View/ObjectViewer/Tree/SelectionViewerColumns.cpp
--- a/Engine/Plugins/Developer/Concert/ConcertUI/ConcertClientSharedSlate/Source/ConcertClientSharedSlate/Private/Replication/Editor/View/ObjectViewer/Tree/SelectionViewerColumns.cpp
+++ b/Engine/Plugins/Developer/Concert/ConcertUI/ConcertClientSharedSlate/Source/ConcertClientSharedSlate/Private/Replication/Editor/View/ObjectViewer/Tree/SelectionViewerColumns.cpp
@@ -2,11 +2,14 @@
 
 #include "SelectionViewerColumns.h"
 
+#include "ClassIconFinder.h"
 #include "Replication/Editor/Model/DisplayUtils.h"
 #include "Replication/Editor/Model/IEditableObjectToPropertiesModel.h"
 #include "Replication/Editor/Model/ReplicatedPropertyData.h"
 #include "Replication/Editor/Model/ReplicatedObjectData.h"
+#include "Replication/Editor/Model/Subobject/ISubobjectModel.h"
 #include "Replication/Editor/View/ObjectEditor/SBaseReplicationStreamEditor.h"
+#include "Replication/Editor/View/ReplicationColumnsUtils.h"
 #include "Replication/PropertyChainUtils.h"
 
 #include "Internationalization/Internationalization.h"
@@ -19,17 +22,17 @@
 
 #define LOCTEXT_NAMESPACE "ReplicationObjectColumns"
 
-namespace UE::ConcertClientSharedSlate::ReplicationObjectColumns
+namespace UE::ConcertClientSharedSlate::ReplicationColumns::TopLevel
 {
 	const FName IconColumnId = TEXT("IconColumn");
 	const FName LabelColumnId = TEXT("LabelColumn");
 	const FName TypeColumnId = TEXT("TypeColumn");
 
-	FReplicationObjectColumn IconColumn(TSharedRef<IObjectToPropertiesModel> Model, const float ColumnWidth)
+	FReplicationTopLevelObjectColumn IconColumn(TSharedRef<IObjectToPropertiesModel> Model, const float ColumnWidth)
 	{
-		return FReplicationObjectColumn(
-			FReplicationObjectColumn::FArguments()
-				.GenerateWidgetColumn_Lambda([Model](const FReplicationObjectColumn::FBuildArgs& Args)
+		return FReplicationTopLevelObjectColumn(
+			FReplicationTopLevelObjectColumn::FArguments()
+				.GenerateWidgetColumn_Lambda([Model](const FReplicationTopLevelObjectColumn::FBuildArgs& Args)
 				{
 					return SNew(SHorizontalBox)
 						+SHorizontalBox::Slot()
@@ -40,18 +43,18 @@ namespace UE::ConcertClientSharedSlate::ReplicationObjectColumns
 							.Image(DisplayUtils::GetObjectIcon(*Model, Args.RowData.GetObjectPath()).GetOptionalIcon())
 						];
 				})
-				.ColumnSortOrder(static_cast<int32>(EReplicationColumnOrder::Icon)),
+				.ColumnSortOrder(static_cast<int32>(ETopLevelColumnOrder::Icon)),
 			SHeaderRow::Column(IconColumnId)
 				.DefaultLabel(FText::GetEmpty())
 				.FixedWidth(ColumnWidth)
 			);
 	}
 	
-	FReplicationObjectColumn LabelColumn()
+	FReplicationTopLevelObjectColumn LabelColumn()
 	{
-		return FReplicationObjectColumn(
-			FReplicationObjectColumn::FArguments()
-				.GenerateWidgetColumn_Lambda([](const FReplicationObjectColumn::FBuildArgs& Args)
+		return FReplicationTopLevelObjectColumn(
+			FReplicationTopLevelObjectColumn::FArguments()
+				.GenerateWidgetColumn_Lambda([](const FReplicationTopLevelObjectColumn::FBuildArgs& Args)
 				{
 					return SNew(STextBlock)
 						.HighlightText(TAttribute<FText>::CreateLambda([HighlightText = Args.HighlightText](){ return *HighlightText; }))
@@ -61,18 +64,18 @@ namespace UE::ConcertClientSharedSlate::ReplicationObjectColumns
 				{
 					InOutSearchStrings.Add(DisplayUtils::GetObjectDisplayText(ObjectData.GetObjectPath()).ToString());
 				})
-				.ColumnSortOrder(static_cast<int32>(EReplicationColumnOrder::Label)),
+				.ColumnSortOrder(static_cast<int32>(ETopLevelColumnOrder::Label)),
 			SHeaderRow::Column(LabelColumnId)
 				.DefaultLabel(LOCTEXT("LabelColumnLabel", "Label"))
 				.FillWidth(1.f)
 			);
 	}
 	
-	FReplicationObjectColumn TypeColumn(TSharedRef<IObjectToPropertiesModel> Model)
+	FReplicationTopLevelObjectColumn TypeColumn(TSharedRef<IObjectToPropertiesModel> Model)
 	{
-		return FReplicationObjectColumn(
-			FReplicationObjectColumn::FArguments()
-				.GenerateWidgetColumn_Lambda([Model](const FReplicationObjectColumn::FBuildArgs& Args)
+		return FReplicationTopLevelObjectColumn(
+			FReplicationTopLevelObjectColumn::FArguments()
+				.GenerateWidgetColumn_Lambda([Model](const FReplicationTopLevelObjectColumn::FBuildArgs& Args)
 				{
 					return SNew(STextBlock)
 						.HighlightText(TAttribute<FText>::CreateLambda([HighlightText = Args.HighlightText](){ return *HighlightText; }))
@@ -82,7 +85,7 @@ namespace UE::ConcertClientSharedSlate::ReplicationObjectColumns
 				{
 					InOutSearchStrings.Add(DisplayUtils::GetObjectTypeText(Model.Get(), ObjectData.GetObjectPath()).ToString());
 				})
-				.ColumnSortOrder(static_cast<int32>(EReplicationColumnOrder::Type)),
+				.ColumnSortOrder(static_cast<int32>(ETopLevelColumnOrder::Type)),
 			SHeaderRow::Column(TypeColumnId)
 				.DefaultLabel(LOCTEXT("TypeColumnLabel", "Type"))
 				.FillWidth(1.f)
@@ -91,10 +94,62 @@ namespace UE::ConcertClientSharedSlate::ReplicationObjectColumns
 }
 
 #undef LOCTEXT_NAMESPACE
+
+namespace UE::ConcertClientSharedSlate::ReplicationColumns::Subobject
+{
+	const FName DisplayColumnId(TEXT("DisplayColumn"));
+	
+	FReplicationSubobjectObjectColumn DisplayColumn(ISubobjectModel& SubobjectModel)
+	{
+		return FReplicationSubobjectObjectColumn(
+			FReplicationSubobjectObjectColumn::FArguments()
+				.GenerateWidgetColumn_Lambda([&SubobjectModel](const FReplicationSubobjectObjectColumn::FBuildArgs& Args)
+				{
+					const FSoftObjectPath ObjectPath = Args.RowData.GetObjectPath();
+					
+					const FSlateBrush* ComponentIcon = FAppStyle::GetBrush("SCS.NativeComponent");
+					const UObject* Object = ObjectPath.ResolveObject();
+					const AActor* AsActor = Cast<AActor>(Object);
+					ComponentIcon = AsActor ? FClassIconFinder::FindIconForActor(AsActor) : ComponentIcon;
+					ComponentIcon = Object ? FSlateIconFinder::FindIconBrushForClass(Object->GetClass(), TEXT("SCS.Component")) : ComponentIcon;
+					
+					return SNew(SHorizontalBox)
+						+SHorizontalBox::Slot()
+						.AutoWidth()
+						.VAlign(VAlign_Center)
+						[
+							SNew(SImage)
+							.Image(ComponentIcon)
+							.ColorAndOpacity(FSlateColor::UseForeground())
+						]
+
+						+SHorizontalBox::Slot()
+							.AutoWidth()
+							.VAlign(VAlign_Center)
+							.Padding(6.f, 0.f, 0.f, 0.f)
+						[
+							SNew(STextBlock)
+							.HighlightText(TAttribute<FText>::CreateLambda([HighlightText = Args.HighlightText](){ return *HighlightText; }))
+							.Text(SubobjectModel.GetSubobjectDisplayName(ObjectPath))
+						];
+				})
+				.PopulateSearchItems_Lambda([&SubobjectModel](const FReplicatedObjectData& ObjectData, TArray<FString>& InOutSearchStrings)
+				{
+					const FSoftObjectPath ObjectPath = ObjectData.GetObjectPath();
+					InOutSearchStrings.Add(SubobjectModel.GetSubobjectDisplayName(ObjectPath).ToString());
+				})
+				.ColumnSortOrder(static_cast<int32>(ESubobjectColumnOrder::DisplayLabel)),
+			SHeaderRow::Column(DisplayColumnId)
+				.FillWidth(1.f)
+			);
+	}
+}
+
 #define LOCTEXT_NAMESPACE "ReplicationPropertyColumns"
 
-namespace UE::ConcertClientSharedSlate::ReplicationPropertyColumns
+namespace UE::ConcertClientSharedSlate::ReplicationColumns::Property
 {
+	const FName ReplicatesColumnId = TEXT("ReplicatedColumn");
 	const FName LabelColumnId = TEXT("LabelColumn");
 	const FName TypeColumnId = TEXT("TypeColumn");
 	
@@ -207,36 +262,53 @@ namespace UE::ConcertClientSharedSlate::ReplicationPropertyColumns
 			}
 		}
 	}
+
+#define MOVE_CAPTURE(a)
 	
 	FReplicationPropertyColumn ReplicatesColumns(
 		TWeakPtr<IReplicationStreamViewer> Viewer,
 		TWeakPtr<IEditableObjectToPropertiesModel> Model,
+		TReplicationColumnDelegates<FReplicatedPropertyData>::FIsEnabled IsEnabledDelegate,
+		TAttribute<FText> DisabledToolTipText,
 		const float ColumnWidth,
 		const int32 Priority
 		)
 	{
-		return MakePropertyCheckboxColumn(
-			FGetColumnCheckboxState::CreateLambda(
-				[Viewer, Model](const FConcertPropertyChain& Property)
+		using FPropertyColumnDelegates = TReplicationColumnDelegates<FReplicatedPropertyData>;
+		return MakeCheckboxColumn<FReplicatedPropertyData>(
+			ReplicatesColumnId,
+			FPropertyColumnDelegates(
+				FPropertyColumnDelegates::FGetColumnCheckboxState::CreateLambda(
+				[Viewer, Model](const FReplicatedPropertyData& Data)
 				{
 					const TSharedPtr<IReplicationStreamViewer> ViewerPin = Viewer.Pin();
 					const TSharedPtr<IEditableObjectToPropertiesModel> ModelPin = Model.Pin();
-					return ensure(ViewerPin && ModelPin) ? Private::OnGetPropertyCheckboxState(Property, *ViewerPin, *ModelPin) : ECheckBoxState::Undetermined;
+					return ensure(ViewerPin && ModelPin) ? Private::OnGetPropertyCheckboxState(Data.GetProperty(), *ViewerPin, *ModelPin) : ECheckBoxState::Undetermined;
 				}),
-			FOnColumnCheckboxChanged::CreateLambda(
-				[Viewer, Model](bool bIsChecked, const FConcertPropertyChain& Property)
+				FPropertyColumnDelegates::FOnColumnCheckboxChanged::CreateLambda(
+				[Viewer, Model](bool bIsChecked, const FReplicatedPropertyData& Data)
 				{
 					const TSharedPtr<IReplicationStreamViewer> ViewerPin = Viewer.Pin();
 					const TSharedPtr<IEditableObjectToPropertiesModel> ModelPin = Model.Pin();
 					if (ensure(ViewerPin && ModelPin))
 					{
-						Private::OnPropertyCheckboxChanged(bIsChecked, Property, *ViewerPin, *ModelPin);
+						Private::OnPropertyCheckboxChanged(bIsChecked, Data.GetProperty(), *ViewerPin, *ModelPin);
 					}
 				}),
+				FPropertyColumnDelegates::FGetToolTipText::CreateLambda(
+				[IsEnabledDelegate, DisabledToolTipText = MoveTemp(DisabledToolTipText)](const FReplicatedPropertyData& Data)
+				{
+					const bool bIsDisabled = IsEnabledDelegate.IsBound() && !IsEnabledDelegate.Execute(Data);
+					const bool bCanCall = DisabledToolTipText.IsBound() || DisabledToolTipText.IsSet();
+					return bIsDisabled
+						? bCanCall ? DisabledToolTipText.Get() : FText::GetEmpty()
+						: LOCTEXT("Replicates.ToolTip", "Select whether this property should be replicated");
+				}),
+				IsEnabledDelegate
+				),
 			FText::GetEmpty(),
-			LOCTEXT("Replicates.ToolTip", "Select whether this property should be replicated"),
-			ColumnWidth, 
-			Priority
+			Priority,
+			ColumnWidth
 		);
 	}
 
diff --git a/Engine/Plugins/Developer/Concert/ConcertUI/ConcertClientSharedSlate/Source/ConcertClientSharedSlate/Private/Replication/Editor/View/ObjectViewer/Tree/SelectionViewerColumns.h b/Engine/Plugins/Developer/Concert/ConcertUI/ConcertClientSharedSlate/Source/ConcertClientSharedSlate/Private/Replication/Editor/View/ObjectViewer/Tree/SelectionViewerColumns.h
--- a/Engine/Plugins/Developer/Concert/ConcertUI/ConcertClientSharedSlate/Source/ConcertClientSharedSlate/Private/Replication/Editor/View/ObjectViewer/Tree/SelectionViewerColumns.h
+++ b/Engine/Plugins/Developer/Concert/ConcertUI/ConcertClientSharedSlate/Source/ConcertClientSharedSlate/Private/Replication/Editor/View/ObjectViewer/Tree/SelectionViewerColumns.h
@@ -3,39 +3,89 @@
 #pragma once
 
 #include "Replication/Editor/View/PredefinedReplicationColumns.h"
+#include "Replication/Editor/View/ReplicationColumnsUtils.h"
 #include "Templates/SharedPointer.h"
 
 namespace UE::ConcertClientSharedSlate
 {
+	class IEditableObjectToPropertiesModel;
+	class IReplicationStreamViewer;
 	class IObjectToPropertiesModel;
+	class ISubobjectModel;
 }
 
-namespace UE::ConcertClientSharedSlate::ReplicationObjectColumns
+namespace UE::ConcertClientSharedSlate::ReplicationColumns::TopLevel
 {
 	extern const FName IconColumnId;
 	extern const FName LabelColumnId;
 	extern const FName TypeColumnId;
 
-	FReplicationObjectColumn IconColumn(TSharedRef<IObjectToPropertiesModel> Model, const float ColumnWidth = 20.f);
-	FReplicationObjectColumn LabelColumn();
-	FReplicationObjectColumn TypeColumn(TSharedRef<IObjectToPropertiesModel> Model);
+	enum class ETopLevelColumnOrder : int32
+	{
+		/** Displays the class icon */
+		Icon = 10,
+		/** Label of the object */
+		Label = 20,
+		/** Class of the object */
+		Type = 30,
+	};
+
+	FReplicationTopLevelObjectColumn IconColumn(TSharedRef<IObjectToPropertiesModel> Model, const float ColumnWidth = 20.f);
+	FReplicationTopLevelObjectColumn LabelColumn();
+	FReplicationTopLevelObjectColumn TypeColumn(TSharedRef<IObjectToPropertiesModel> Model);
+}
+
+namespace UE::ConcertClientSharedSlate::ReplicationColumns::Subobject
+{
+	extern const FName DisplayColumnId;
+	
+	enum class ESubobjectColumnOrder : int32
+	{
+		/** Displays subobject name and class icon */
+		DisplayLabel = 10,
+	};
+
+	/**
+	 * Displays class icon and name of subobject
+	 * @param SubobjectModel Used to get display names. Outlives the column.
+	 */
+	FReplicationSubobjectObjectColumn DisplayColumn(ISubobjectModel& SubobjectModel);
 }
 
-namespace UE::ConcertClientSharedSlate::ReplicationPropertyColumns
+namespace UE::ConcertClientSharedSlate::ReplicationColumns::Property
 {
 	extern const FName ReplicatesColumnId;
 	extern const FName LabelColumnId;
 	extern const FName TypeColumnId;
 	
+	enum class EReplicationPropertyColumnOrder : int32
+	{
+		/** The checkbox in SDefaultReplicationStreamEditor determining whether the property is in the selection*/
+		ReplicatesCheckbox = 0,
+		/** Label of the property */
+		Label = 10,
+		/** Type of the property */
+		Type = 20
+	};
+	
 	FReplicationPropertyColumn LabelColumn();
 	FReplicationPropertyColumn TypeColumn();
 	/**
 	 * A checkbox that is placed at the beginning of every property.
 	 * Checking & unchecking adds & removes the property to the selected objects' property mapping, respectively.
+	 *
+	 * @param Viewer Used to determine which objects are currently selected (and being edited by the checkbox)
+	 * @param Model Used to actually edit the object to properties mapping
+	 * @param IsEnabledDelegate Determines whether the checkbox is enabled
+	 * @param DisabledToolTipText Tooltip to display when IsEnabledDelegate returns false
+	 * @param ColumnWidth Width to use for the column
+	 * @param Priority Determines position of this columns relative to the others
 	 */
 	CONCERTCLIENTSHAREDSLATE_API FReplicationPropertyColumn ReplicatesColumns(
 		TWeakPtr<IReplicationStreamViewer> Viewer,
 		TWeakPtr<IEditableObjectToPropertiesModel> Model,
+		TReplicationColumnDelegates<FReplicatedPropertyData>::FIsEnabled IsEnabledDelegate = {},
+		TAttribute<FText> DisabledToolTipText = {},
 		const float ColumnWidth = 20.f,
 		const int32 Priority = static_cast<int32>(EReplicationPropertyColumnOrder::ReplicatesCheckbox)
 		);
diff --git a/Engine/Plugins/Developer/Concert/ConcertUI/ConcertClientSharedSlate/Source/ConcertClientSharedSlate/Private/Replication/Editor/View/PredefinedReplicationColumns.cpp b/Engine/Plugins/Developer/Concert/ConcertUI/ConcertClientSharedSlate/Source/ConcertClientSharedSlate/Private/Replication/Editor/View/PredefinedReplicationColumns.cpp
deleted file mode 100644
diff --git a/Engine/Plugins/Developer/Concert/ConcertUI/ConcertClientSharedSlate/Source/ConcertClientSharedSlate/Private/Replication/Editor/View/SubobjectView/ReplicatedSubobjectData.h b/Engine/Plugins/Developer/Concert/ConcertUI/ConcertClientSharedSlate/Source/ConcertClientSharedSlate/Private/Replication/Editor/View/SubobjectView/ReplicatedSubobjectData.h
new file mode 100644
--- /dev/null
+++ b/Engine/Plugins/Developer/Concert/ConcertUI/ConcertClientSharedSlate/Source/ConcertClientSharedSlate/Private/Replication/Editor/View/SubobjectView/ReplicatedSubobjectData.h
@@ -0,0 +1,30 @@
+// Copyright Epic Games, Inc. All Rights Reserved.
+
+#pragma once
+
+#include "Replication/Editor/Model/ReplicatedObjectData.h"
+
+namespace UE::ConcertClientSharedSlate
+{
+	/** Type of items in SSubobjectView's tree view */
+	class FReplicatedSubobjectData
+	{
+	public:
+
+		FReplicatedSubobjectData() = default;
+		explicit FReplicatedSubobjectData(FSoftObjectPath ObjectPath)
+			: bIsSeparator(false)
+			, ObjectData(MoveTemp(ObjectPath))
+		{}
+
+		bool IsSeparator() const { return bIsSeparator; }
+		const TOptional<FReplicatedObjectData>& GetObjectData() const { return ObjectData; }
+		
+	private:
+
+		/** Whether this is a separator (i.e. whether it separates categories). */
+		bool bIsSeparator = true;
+		/** Set if bIsSeparator == false. */
+		TOptional<FReplicatedObjectData> ObjectData;
+	};
+}
\ No newline at end of file
diff --git a/Engine/Plugins/Developer/Concert/ConcertUI/ConcertClientSharedSlate/Source/ConcertClientSharedSlate/Private/Replication/Editor/View/SubobjectView/SConcertReplicationSubobjectEditor.cpp b/Engine/Plugins/Developer/Concert/ConcertUI/ConcertClientSharedSlate/Source/ConcertClientSharedSlate/Private/Replication/Editor/View/SubobjectView/SConcertReplicationSubobjectEditor.cpp
deleted file mode 100644
diff --git a/Engine/Plugins/Developer/Concert/ConcertUI/ConcertClientSharedSlate/Source/ConcertClientSharedSlate/Private/Replication/Editor/View/SubobjectView/SConcertReplicationSubobjectEditor.h b/Engine/Plugins/Developer/Concert/ConcertUI/ConcertClientSharedSlate/Source/ConcertClientSharedSlate/Private/Replication/Editor/View/SubobjectView/SConcertReplicationSubobjectEditor.h
deleted file mode 100644
diff --git a/Engine/Plugins/Developer/Concert/ConcertUI/ConcertClientSharedSlate/Source/ConcertClientSharedSlate/Private/Replication/Editor/View/SubobjectView/SSubobjectView.cpp b/Engine/Plugins/Developer/Concert/ConcertUI/ConcertClientSharedSlate/Source/ConcertClientSharedSlate/Private/Replication/Editor/View/SubobjectView/SSubobjectView.cpp
new file mode 100644
--- /dev/null
+++ b/Engine/Plugins/Developer/Concert/ConcertUI/ConcertClientSharedSlate/Source/ConcertClientSharedSlate/Private/Replication/Editor/View/SubobjectView/SSubobjectView.cpp
@@ -0,0 +1,227 @@
+// Copyright Epic Games, Inc. All Rights Reserved.
+
+#include "SSubobjectView.h"
+
+#include "Misc/EBreakBehavior.h"
+#include "ReplicatedSubobjectData.h"
+#include "Replication/Editor/Model/Subobject/ISubobjectModel.h"
+#include "Replication/Editor/View/ObjectViewer/Tree/SelectionViewerColumns.h"
+
+#include "Containers/Queue.h"
+#include "Styling/AppStyle.h"
+#include "Widgets/Layout/SBox.h"
+#include "Widgets/Layout/SSeparator.h"
+
+namespace UE::ConcertClientSharedSlate
+{
+	void SSubobjectView::Construct(const FArguments& InArgs, TSharedRef<ISubobjectModel> InSubobjectModel)
+	{
+		SubobjectModel = MoveTemp(InSubobjectModel);
+		
+		const auto TransformOp = [](const FReplicatedSubobjectData Data) -> FReplicatedObjectData
+		{
+			return ensureMsgf(!Data.IsSeparator(), TEXT("The widgets are supposed to be set up so that separators are not passed down to API users."))
+				? *Data.GetObjectData() : FReplicatedObjectData{{}};
+		};
+		TArray<TReplicationColumn<FReplicatedSubobjectData>> Columns
+		{
+			ReplicationColumns::Subobject::DisplayColumn(*SubobjectModel).TransformColumn<FReplicatedSubobjectData>(TransformOp)
+		};
+		for (const ReplicationColumns::FReplicationSubobjectObjectColumn& Column : InArgs._AdditionalColumns)
+		{
+			Columns.Emplace(Column.TransformColumn<FReplicatedSubobjectData>(TransformOp));
+		}
+		
+		ChildSlot
+		[
+			SAssignNew(SubobjectTreeView, SReplicationTreeView<FReplicatedSubobjectData>)
+			.RootItemsSource(&RootSubobjectData)
+			.OnGetChildren(this, &SSubobjectView::GetObjectRowChildren)
+			.OnSelectionChanged_Lambda([this]()
+			{
+				OnSelectionChangedDelegate.Broadcast();
+			})
+			.OverrideColumnWidget(this, &SSubobjectView::GetOptionalColumnOverrideWidget)
+			.IsSearchableItem_Lambda([](const TSharedPtr<FReplicatedSubobjectData>& Item){ return !Item->IsSeparator(); })
+			.Columns(Columns)
+			.ExpandableColumnLabel(ReplicationColumns::Subobject::DisplayColumnId)
+			.HeaderRowVisibility(EVisibility::Collapsed)
+			.SelectionMode(ESelectionMode::Multi)
+		];
+
+		SubobjectModel->OnHierarchyChanged().AddSP(this, &SSubobjectView::RefreshView);
+		RefreshView();
+	}
+	
+	void SSubobjectView::SetTopLevelObjects(const TArray<FSoftObjectPath>& InTopLevelObjects)
+	{
+		// TODO DP: Handle multiple top level objects. Create a new widget class that wraps SSubobjectView and displays a message.
+		if (InTopLevelObjects.IsEmpty())
+		{
+			SubobjectModel->SetTopLevelObject(nullptr);
+		}
+		else
+		{
+			// This triggers OnHierarchyChanged which calls RefreshView 
+			SubobjectModel->SetTopLevelObject(InTopLevelObjects[0]);
+			
+			// There might be nothing selected
+			if (!AllSubobjectData.IsEmpty())
+			{
+				SubobjectTreeView->SetExpandedItems(AllSubobjectData, true);
+			}
+		}
+	}
+
+	void SSubobjectView::SelectTopLevelObjects()
+	{
+		if (SubobjectModel->GetTopLevelObject().IsNull())
+		{
+			return;
+		}
+		
+		const TSharedPtr<FReplicatedSubobjectData>* TopLevelItem = RootSubobjectData.FindByPredicate([this](const TSharedPtr<FReplicatedSubobjectData>& RowData)
+		{
+			return !RowData->IsSeparator() && RowData->GetObjectData()->GetObjectPath() == SubobjectModel->GetTopLevelObject();
+		});
+		if (ensureMsgf(TopLevelItem, TEXT("Top level item is supposed to be contained")))
+		{
+			SubobjectTreeView->SetSelectedItems({ *TopLevelItem }, true);
+		}
+	}
+
+	TArray<FSoftObjectPath> SSubobjectView::GetSelectedObjects() const
+	{
+		TArray<FSoftObjectPath> Result;
+		Algo::TransformIf(SubobjectTreeView->GetSelectedItems(), Result,
+			[](const TSharedPtr<FReplicatedSubobjectData>& RowData)
+			{
+				return !RowData->IsSeparator();
+			},
+			[](const TSharedPtr<FReplicatedSubobjectData>& RowData)
+			{
+				return RowData->GetObjectData()->GetObjectPath();
+			});
+		return Result;
+	}
+
+	void SSubobjectView::RefreshView()
+	{
+		// We want to re-use TSharedPtr items so selection states are retained.
+		// PathToObjectDataCache will keep the items alive until it is replaced below.
+		RootSubobjectData.Empty();
+		AllSubobjectData.Empty();
+
+		if (SubobjectModel->GetTopLevelObject().IsNull())
+		{
+			SubobjectTreeView->OnItemsChanged();
+			return;
+		}
+		
+		TMap<FSoftObjectPath, TSharedPtr<FReplicatedSubobjectData>> NewPathToObjectDataCache;
+		BuildRootObjectData(NewPathToObjectDataCache);
+		BuildChildObjectData(NewPathToObjectDataCache);
+
+		// Only refresh if something actually changed
+		if (!PathToObjectDataCache.OrderIndependentCompareEqual(NewPathToObjectDataCache))
+		{
+			// This will destroy the old items
+			PathToObjectDataCache = MoveTemp(NewPathToObjectDataCache);
+			SubobjectTreeView->OnItemsChanged();
+		}
+	}
+
+	void SSubobjectView::BuildRootObjectData(TMap<FSoftObjectPath, TSharedPtr<FReplicatedSubobjectData>> NewPathToObjectDataCache)
+	{
+		// Top level item is always first
+		const FSoftObjectPath TopLevelObject = SubobjectModel->GetTopLevelObject();
+		const TSharedPtr<FReplicatedSubobjectData>* ExistingTopLevelItem = PathToObjectDataCache.Find(TopLevelObject);
+		const TSharedPtr<FReplicatedSubobjectData> TopLevelItem = ExistingTopLevelItem ? *ExistingTopLevelItem : MakeShared<FReplicatedSubobjectData>(TopLevelObject);
+		RootSubobjectData.Add(TopLevelItem);
+		AllSubobjectData.Add(TopLevelItem);
+		NewPathToObjectDataCache.Add(TopLevelObject, TopLevelItem);
+		
+		for (const FName Category : SubobjectModel->GetCategories())
+		{
+			bool bAddedSeparator = false;
+			SubobjectModel->ForEachRootSubobject(Category, [this, &NewPathToObjectDataCache, &bAddedSeparator](const FSoftObjectPath& Subobject) 
+			{
+				// Separators are between categories
+				if (!bAddedSeparator)
+				{
+					RootSubobjectData.Emplace(MakeShared<FReplicatedSubobjectData>());
+					bAddedSeparator = true;
+				}
+
+				const TSharedPtr<FReplicatedSubobjectData>* Existing = PathToObjectDataCache.Find(Subobject);
+				const TSharedPtr<FReplicatedSubobjectData> Item = Existing ? *Existing : MakeShared<FReplicatedSubobjectData>(Subobject);
+				RootSubobjectData.Add(Item);
+				AllSubobjectData.Add(Item);
+				NewPathToObjectDataCache.Add(Subobject, Item);
+				return EBreakBehavior::Continue;
+			});
+		}
+	}
+
+	void SSubobjectView::BuildChildObjectData(TMap<FSoftObjectPath, TSharedPtr<FReplicatedSubobjectData>>& NewPathToObjectDataCache)
+	{
+		TQueue<FSoftObjectPath> PendingObjects;
+		for (const TSharedPtr<FReplicatedSubobjectData>& RowData : RootSubobjectData)
+		{
+			if (!RowData->IsSeparator() && RowData->GetObjectData()->GetObjectPath() != SubobjectModel->GetTopLevelObject())
+			{
+				PendingObjects.Enqueue(RowData->GetObjectData()->GetObjectPath());
+			}
+		}
+
+		// Subobjects form a tree hierarchy so no need for infinite loop detection
+		FSoftObjectPath CurrentPath;
+		while (PendingObjects.Dequeue(CurrentPath))
+		{
+			const TSharedPtr<FReplicatedSubobjectData>* Existing = PathToObjectDataCache.Find(CurrentPath);
+			TSharedPtr<FReplicatedSubobjectData> Data = Existing ? *Existing : MakeShared<FReplicatedSubobjectData>(CurrentPath);
+			NewPathToObjectDataCache.Add(CurrentPath, Data);
+			AllSubobjectData.Add(Data);
+
+			SubobjectModel->ForEachDirectChildSubobject(CurrentPath, [&PendingObjects](const FSoftObjectPath& ChildPath)
+			{
+				PendingObjects.Enqueue(ChildPath);
+				return EBreakBehavior::Continue;
+			});
+		}
+	}
+
+	void SSubobjectView::GetObjectRowChildren(TSharedPtr<FReplicatedSubobjectData> Item, TFunctionRef<void(TSharedPtr<FReplicatedSubobjectData>)> Callback) const
+	{
+		if (Item->IsSeparator()
+			|| Item->GetObjectData()->GetObjectPath() == SubobjectModel->GetTopLevelObject())
+		{
+			return;
+		}
+		
+		SubobjectModel->ForEachDirectChildSubobject(Item->GetObjectData()->GetObjectPath(), [this, &Callback](const FSoftObjectPath& Object)
+		{
+			if (const TSharedPtr<FReplicatedSubobjectData>* FoundItem = PathToObjectDataCache.Find(Object)
+				; ensureMsgf(FoundItem, TEXT("No item created for %s"), *Object.ToString()))
+			{
+				Callback(*FoundItem);
+			}
+			return EBreakBehavior::Continue;
+		});
+	}
+
+	TSharedPtr<SWidget> SSubobjectView::GetOptionalColumnOverrideWidget(const FName& Name, const FReplicatedSubobjectData& RowData) const
+	{
+		if (RowData.IsSeparator())
+		{
+			return SNew(SBox)
+				.Padding(1.f)
+				[
+					SNew(SSeparator)
+					.SeparatorImage(FAppStyle::Get().GetBrush("Menu.Separator"))
+					.Thickness(1.f)
+				 ];
+		}
+		return nullptr;
+	}
+}
diff --git a/Engine/Plugins/Developer/Concert/ConcertUI/ConcertClientSharedSlate/Source/ConcertClientSharedSlate/Private/Replication/Editor/View/SubobjectView/SSubobjectView.h b/Engine/Plugins/Developer/Concert/ConcertUI/ConcertClientSharedSlate/Source/ConcertClientSharedSlate/Private/Replication/Editor/View/SubobjectView/SSubobjectView.h
new file mode 100644
--- /dev/null
+++ b/Engine/Plugins/Developer/Concert/ConcertUI/ConcertClientSharedSlate/Source/ConcertClientSharedSlate/Private/Replication/Editor/View/SubobjectView/SSubobjectView.h
@@ -0,0 +1,68 @@
+// Copyright Epic Games, Inc. All Rights Reserved.
+
+#pragma once
+
+#include "Replication/Editor/Model/ReplicatedObjectData.h"
+#include "Replication/Editor/View/IReplicationSubobjectView.h"
+#include "Replication/Editor/View/PredefinedReplicationColumns.h"
+#include "Replication/Editor/View/ObjectViewer/Tree/SReplicationTreeView.h"
+
+namespace UE::ConcertClientSharedSlate
+{
+	class ISubobjectModel;
+	class FReplicatedSubobjectData;
+	class SConcertReplicationSubobjectEditor;
+
+	/** Adapts SSubobjectEditor to the IReplicationSubobjectView interface. */
+	class SSubobjectView : public IReplicationSubobjectView
+	{
+	public:
+
+		SLATE_BEGIN_ARGS(SSubobjectView)
+		{}
+			/** Additional columns that should be displayed for the subobjects. */
+			SLATE_ARGUMENT(TArray<ReplicationColumns::FReplicationSubobjectObjectColumn>, AdditionalColumns)
+		SLATE_END_ARGS()
+
+		void Construct(const FArguments& InArgs, TSharedRef<ISubobjectModel> InSubobjectModel);
+
+		//~ Begin IReplicationSubobjectView Interface
+		virtual void SetTopLevelObjects(const TArray<FSoftObjectPath>& RootObjects) override;
+		virtual void SelectTopLevelObjects() override;
+		virtual TArray<FSoftObjectPath> GetSelectedObjects() const override;
+		virtual FOnSelectionChanged& OnSelectionChanged() override { return OnSelectionChangedDelegate; }
+		//~ End IReplicationSubobjectView Interface
+
+	private:
+
+		/** Gets the subobjects to display */
+		TSharedPtr<ISubobjectModel> SubobjectModel;
+
+		/** Reuses the SSubobjectEditor for displaying subobjects. */
+		TSharedPtr<SReplicationTreeView<FReplicatedSubobjectData>> SubobjectTreeView;
+
+		/** Contains all subobject data. */
+		TArray<TSharedPtr<FReplicatedSubobjectData>> AllSubobjectData;
+		/**
+		 * Contains only the root elements of the tree, which are:
+		 * - Item corresponding to top level object
+		 * - Separators
+		 * - Those objects returned by ISubobjectModel::ForEachRootSubobject
+		 */
+		TArray<TSharedPtr<FReplicatedSubobjectData>> RootSubobjectData;
+		/** Inverse map of ObjectRowData using FReplicatedObjectData::GetObjectPath as key. Contains all elements of ObjectRowData. */
+		TMap<FSoftObjectPath, TSharedPtr<FReplicatedSubobjectData>> PathToObjectDataCache;
+
+		/** Executes when the selection is changed. */
+		FOnSelectionChanged OnSelectionChangedDelegate;
+
+		void RefreshView();
+		void BuildRootObjectData(TMap<FSoftObjectPath, TSharedPtr<FReplicatedSubobjectData>> NewPathToObjectDataCache);
+		void BuildChildObjectData(TMap<FSoftObjectPath, TSharedPtr<FReplicatedSubobjectData>>& NewPathToObjectDataCache);
+		
+		void GetObjectRowChildren(TSharedPtr<FReplicatedSubobjectData> Item, TFunctionRef<void(TSharedPtr<FReplicatedSubobjectData>)> Callback) const;
+		/** Returns a separator bar if the item is a separator */
+		TSharedPtr<SWidget> GetOptionalColumnOverrideWidget(const FName& Name, const FReplicatedSubobjectData& RowData) const;
+	};
+}
+
diff --git a/Engine/Plugins/Developer/Concert/ConcertUI/ConcertClientSharedSlate/Source/ConcertClientSharedSlate/Private/Replication/Editor/View/SubobjectView/SUnrealEditorSubobjectView.cpp b/Engine/Plugins/Developer/Concert/ConcertUI/ConcertClientSharedSlate/Source/ConcertClientSharedSlate/Private/Replication/Editor/View/SubobjectView/SUnrealEditorSubobjectView.cpp
deleted file mode 100644
diff --git a/Engine/Plugins/Developer/Concert/ConcertUI/ConcertClientSharedSlate/Source/ConcertClientSharedSlate/Private/Replication/Editor/View/SubobjectView/SUnrealEditorSubobjectView.h b/Engine/Plugins/Developer/Concert/ConcertUI/ConcertClientSharedSlate/Source/ConcertClientSharedSlate/Private/Replication/Editor/View/SubobjectView/SUnrealEditorSubobjectView.h
deleted file mode 100644
diff --git a/Engine/Plugins/Developer/Concert/ConcertUI/ConcertClientSharedSlate/Source/ConcertClientSharedSlate/Private/Replication/ReplicationWidgetFactories.cpp b/Engine/Plugins/Developer/Concert/ConcertUI/ConcertClientSharedSlate/Source/ConcertClientSharedSlate/Private/Replication/ReplicationWidgetFactories.cpp
--- a/Engine/Plugins/Developer/Concert/ConcertUI/ConcertClientSharedSlate/Source/ConcertClientSharedSlate/Private/Replication/ReplicationWidgetFactories.cpp
+++ b/Engine/Plugins/Developer/Concert/ConcertUI/ConcertClientSharedSlate/Source/ConcertClientSharedSlate/Private/Replication/ReplicationWidgetFactories.cpp
@@ -3,15 +3,17 @@
 #include "Replication/ReplicationWidgetFactories.h"
 
 #include "Editor/Model/TransactionalPropertySelectionModel.h"
+#include "Editor/Model/Subobject/ComponentHierarchySubobjectModel.h"
 #include "Editor/View/ObjectEditor/SDefaultReplicationStreamEditor.h"
 #include "Editor/View/ObjectEditor/SBaseReplicationStreamEditor.h"
-#include "Editor/View/SubobjectView/SUnrealEditorSubobjectView.h"
+#include "Editor/View/SubobjectView/SSubobjectView.h"
 
 namespace UE::ConcertClientSharedSlate
 {
-	TSharedRef<IReplicationSubobjectView> CreateUnrealEditorSubobjectView()
+	TSharedRef<IReplicationSubobjectView> CreateUnrealEditorSubobjectView(FCreateSubobjectViewParams Params)
 	{
-		return SNew(SUnrealEditorSubobjectView);
+		return SNew(SSubobjectView, MakeShared<FComponentHierarchySubobjectModel>())
+			.AdditionalColumns(Params.AdditionalColumns);
 	}
 	
 	TSharedRef<IReplicationStreamEditor> CreateEditor(FCreateEditorParams Params)
@@ -23,13 +25,15 @@ namespace UE::ConcertClientSharedSlate
 			.OnExtendObjectsContextMenu(Params.OnExtendObjectsContextMenu)
 			.SortPropertyRowPredicate(Params.SortPropertyRowPredicate)
 			.LeftOfObjectSearchBar() [ Params.LeftOfObjectSearchBar.Widget ]
-			.LeftOfPropertySearchBar() [ Params.LeftOfPropertySearchBar.Widget ];
+			.LeftOfPropertySearchBar() [ Params.LeftOfPropertySearchBar.Widget ]
+			.IsEditingEnabled(Params.IsEditingEnabled)
+			.EditingDisabledToolTipText(Params.EditingDisabledToolTipText);
 	}
 
 	TSharedRef<IReplicationStreamEditor> CreateDefaultStreamEditor(FCreateEditorParams Params)
 	{
 		// Default editor relies on a subobject view
-		Params.SubobjectView = Params.SubobjectView ? Params.SubobjectView : CreateUnrealEditorSubobjectView();
+		Params.SubobjectView = Params.SubobjectView ? Params.SubobjectView : CreateUnrealEditorSubobjectView(FCreateSubobjectViewParams{});
 		
 		return SNew(SDefaultReplicationStreamEditor, Params.DataModel, Params.ObjectSource, Params.PropertySource)
 			.AdditionalObjectColumns(Params.AdditionalObjectColumns)
@@ -38,7 +42,9 @@ namespace UE::ConcertClientSharedSlate
 			.OnExtendObjectsContextMenu(Params.OnExtendObjectsContextMenu)
 			.SortPropertyRowPredicate(Params.SortPropertyRowPredicate)
 			.LeftOfObjectSearchBar() [ Params.LeftOfObjectSearchBar.Widget ]
-			.LeftOfPropertySearchBar() [ Params.LeftOfPropertySearchBar.Widget ];
+			.LeftOfPropertySearchBar() [ Params.LeftOfPropertySearchBar.Widget ]
+			.IsEditingEnabled(Params.IsEditingEnabled)
+			.EditingDisabledToolTipText(Params.EditingDisabledToolTipText);
 	}
 
 	TSharedRef<IEditableObjectToPropertiesModel> CreatePropertySelectionModel(
diff --git a/Engine/Plugins/Developer/Concert/ConcertUI/ConcertClientSharedSlate/Source/ConcertClientSharedSlate/Public/Replication/Editor/Model/IObjectToPropertiesModel.h b/Engine/Plugins/Developer/Concert/ConcertUI/ConcertClientSharedSlate/Source/ConcertClientSharedSlate/Public/Replication/Editor/Model/IObjectToPropertiesModel.h
--- a/Engine/Plugins/Developer/Concert/ConcertUI/ConcertClientSharedSlate/Source/ConcertClientSharedSlate/Public/Replication/Editor/Model/IObjectToPropertiesModel.h
+++ b/Engine/Plugins/Developer/Concert/ConcertUI/ConcertClientSharedSlate/Source/ConcertClientSharedSlate/Public/Replication/Editor/Model/IObjectToPropertiesModel.h
@@ -89,6 +89,18 @@ namespace UE::ConcertClientSharedSlate
 			}, OptionalParent);
 			return Result;
 		}
+
+		/** Iterates every subobject of Parent. */
+		void ForEachSubobject(const FSoftObjectPath& Parent, TFunctionRef<EBreakBehavior(const FSoftObjectPath& Child)> Callback) const
+		{
+			const FString ParentPathString = Parent.ToString();
+			ForEachReplicatedObject([&Parent, &ParentPathString, &Callback](const FSoftObjectPath& Object)
+			{
+				return Object != Parent && Object.ToString().Contains(ParentPathString)
+					? Callback(Object)
+					: EBreakBehavior::Continue;
+			});
+		}
 		
 		virtual ~IObjectToPropertiesModel() = default;
 	};
diff --git a/Engine/Plugins/Developer/Concert/ConcertUI/ConcertClientSharedSlate/Source/ConcertClientSharedSlate/Public/Replication/Editor/View/IReplicationSubobjectView.h b/Engine/Plugins/Developer/Concert/ConcertUI/ConcertClientSharedSlate/Source/ConcertClientSharedSlate/Public/Replication/Editor/View/IReplicationSubobjectView.h
--- a/Engine/Plugins/Developer/Concert/ConcertUI/ConcertClientSharedSlate/Source/ConcertClientSharedSlate/Public/Replication/Editor/View/IReplicationSubobjectView.h
+++ b/Engine/Plugins/Developer/Concert/ConcertUI/ConcertClientSharedSlate/Source/ConcertClientSharedSlate/Public/Replication/Editor/View/IReplicationSubobjectView.h
@@ -26,12 +26,12 @@ namespace UE::ConcertClientSharedSlate
 	public:
 
 		/** Sets the root objects for which the subobject list should be shown. */
-		virtual void SetRootObjects(const TArray<FSoftObjectPath>& RootObjects) = 0;
+		virtual void SetTopLevelObjects(const TArray<FSoftObjectPath>& RootObjects) = 0;
 		/** Clears the hierarchy so it displays no subobjects. */
-		void ClearRootObjects() { SetRootObjects({}); }
+		void ClearRootObjects() { SetTopLevelObjects({}); }
 
 		/** Selects the root object in the UI. Executes OnSelectionChanged. */
-		virtual void SelectRootObjects() = 0;
+		virtual void SelectTopLevelObjects() = 0;
 		
 		/**
 		 * Gets the list of selected objects.
diff --git a/Engine/Plugins/Developer/Concert/ConcertUI/ConcertClientSharedSlate/Source/ConcertClientSharedSlate/Public/Replication/Editor/View/PredefinedReplicationColumns.h b/Engine/Plugins/Developer/Concert/ConcertUI/ConcertClientSharedSlate/Source/ConcertClientSharedSlate/Public/Replication/Editor/View/PredefinedReplicationColumns.h
--- a/Engine/Plugins/Developer/Concert/ConcertUI/ConcertClientSharedSlate/Source/ConcertClientSharedSlate/Public/Replication/Editor/View/PredefinedReplicationColumns.h
+++ b/Engine/Plugins/Developer/Concert/ConcertUI/ConcertClientSharedSlate/Source/ConcertClientSharedSlate/Public/Replication/Editor/View/PredefinedReplicationColumns.h
@@ -2,58 +2,15 @@
 
 #pragma once
 
-#include "Containers/ArrayView.h"
 #include "Replication/Editor/View/ReplicationColumn.h"
 #include "Replication/Editor/Model/ReplicatedObjectData.h"
 #include "Replication/Editor/Model/ReplicatedPropertyData.h"
-#include "Templates/SharedPointer.h"
 
-struct FConcertPropertyChain;
-
-namespace UE::ConcertClientSharedSlate
-{
-	class IEditableObjectToPropertiesModel;
-	class IReplicationStreamViewer;
-	class IObjectToPropertiesModel;
-}
-
-namespace UE::ConcertClientSharedSlate::ReplicationObjectColumns
-{
-	using FReplicationObjectColumn = TReplicationColumn<FReplicatedObjectData>;
-	enum class EReplicationColumnOrder : int32
-	{
-		/** Displays the class icon */
-		Icon = 10,
-		/** Label of the object */
-		Label = 20,
-		/** Class of the object */
-		Type = 30,
-	};
-}
-
-namespace UE::ConcertClientSharedSlate::ReplicationPropertyColumns
+namespace UE::ConcertClientSharedSlate::ReplicationColumns
 {
+	/** Rows displayed in the outliner */
+	using FReplicationTopLevelObjectColumn = TReplicationColumn<FReplicatedObjectData>;
+	/** The type of rows exposed by the subobject view as API. */
+	using FReplicationSubobjectObjectColumn = TReplicationColumn<FReplicatedObjectData>;
 	using FReplicationPropertyColumn = TReplicationColumn<FReplicatedPropertyData>;
-	enum class EReplicationPropertyColumnOrder : int32
-	{
-		/** The checkbox in SDefaultReplicationStreamEditor determining whether the property is in the selection*/
-		ReplicatesCheckbox = 0,
-		/** Label of the property */
-		Label = 10,
-		/** Type of the property */
-		Type = 20
-	};
-	
-	DECLARE_DELEGATE_RetVal_OneParam(ECheckBoxState, FGetColumnCheckboxState, const FConcertPropertyChain& /*Property*/);
-	DECLARE_DELEGATE_TwoParams(FOnColumnCheckboxChanged, bool /*bIsChecked*/, const FConcertPropertyChain& /*Property*/);
-
-	/** Util for creating a checkbox in a property row. The delegates will be given the path to the property. */
-	CONCERTCLIENTSHAREDSLATE_API FReplicationPropertyColumn MakePropertyCheckboxColumn(
-		FGetColumnCheckboxState GetPropertyCheckboxStateDelegate,
-		FOnColumnCheckboxChanged OnPropertyBoxToggledDelegate,
-		FText DefaultLabel = FText::GetEmpty(),
-		FText ToolTipText = FText::GetEmpty(),
-		const float ColumnWidth = 20.f,
-		const int32 Priority = 0
-		);
 }
diff --git a/Engine/Plugins/Developer/Concert/ConcertUI/ConcertClientSharedSlate/Source/ConcertClientSharedSlate/Public/Replication/Editor/View/ReplicationColumn.h b/Engine/Plugins/Developer/Concert/ConcertUI/ConcertClientSharedSlate/Source/ConcertClientSharedSlate/Public/Replication/Editor/View/ReplicationColumn.h
--- a/Engine/Plugins/Developer/Concert/ConcertUI/ConcertClientSharedSlate/Source/ConcertClientSharedSlate/Public/Replication/Editor/View/ReplicationColumn.h
+++ b/Engine/Plugins/Developer/Concert/ConcertUI/ConcertClientSharedSlate/Source/ConcertClientSharedSlate/Public/Replication/Editor/View/ReplicationColumn.h
@@ -38,12 +38,19 @@ namespace UE::ConcertClientSharedSlate
 			SLATE_ARGUMENT(int32, ColumnSortOrder)
 		SLATE_END_ARGS()
 
-		TReplicationColumn(const TReplicationColumn::FArguments& InArgs, const FColumn::FArguments& InColumnArgs)
+		TReplicationColumn(const FArguments& InArgs, const FColumn::FArguments& InColumnArgs)
 			: FColumn(InColumnArgs)
 			, GenerateColumnWidgetCallback(InArgs._GenerateWidgetColumn)
 			, PopulateSearchStringCallback(InArgs._PopulateSearchItems)
 			, ColumnSortOrderValue(InArgs._ColumnSortOrder)
 		{}
+		
+		TReplicationColumn(const FArguments& InArgs, const FColumn& Column)
+			: FColumn(Column)
+			, GenerateColumnWidgetCallback(InArgs._GenerateWidgetColumn)
+			, PopulateSearchStringCallback(InArgs._PopulateSearchItems)
+			, ColumnSortOrderValue(InArgs._ColumnSortOrder)
+		{}
 
 		TSharedRef<SWidget> BuildColumnWidget(const FBuildArgs& InArgs) const
 		{
@@ -61,6 +68,34 @@ namespace UE::ConcertClientSharedSlate
 		}
 		
 		int32 GetColumnSortOrderValue() const { return ColumnSortOrderValue; }
+
+		/**
+		 * Wraps a TReplicationColumn<TListItemType> in a TReplicationColumn<TOtherColumnType> which transforms
+		 * TOtherColumnType to TListItemType.
+		 *
+		 * This is useful e.g. if TOtherColumnType inherits from TListItemType.
+		 * The default TTransformOp argument handles this situation.
+		 */
+		template<typename TOtherColumnType, typename TTransformOp>
+		TReplicationColumn<TOtherColumnType> TransformColumn(TTransformOp TransformOperation = [](const TOtherColumnType& RowData) -> TListItemType { return RowData; }) const
+		{
+			using TReturnColumnType = TReplicationColumn<TOtherColumnType>;
+			FGenerateColumnWidget GenerateWidget = GenerateColumnWidgetCallback;
+			FPopulateSearchString PopulateSearchString = PopulateSearchStringCallback;
+			return TReturnColumnType(
+				typename TReturnColumnType::FArguments()
+					.GenerateWidgetColumn_Lambda([TransformOperation, GenerateWidget](const typename TReturnColumnType::FBuildArgs& BuildArgs)
+					{
+						return GenerateWidget.Execute({ BuildArgs.HighlightText, TransformOperation(BuildArgs.RowData)});
+					})
+					.PopulateSearchItems_Lambda([TransformOperation, PopulateSearchString](const TOtherColumnType& InOtherRowData, TArray<FString>& InOutSearchStrings)
+					{
+						PopulateSearchString.Execute(TransformOperation(InOtherRowData), InOutSearchStrings);
+					})
+					.ColumnSortOrder(ColumnSortOrderValue),
+				*this
+			);
+		}
 		
 	private:
 		
diff --git a/Engine/Plugins/Developer/Concert/ConcertUI/ConcertClientSharedSlate/Source/ConcertClientSharedSlate/Public/Replication/Editor/View/ReplicationColumnsUtils.h b/Engine/Plugins/Developer/Concert/ConcertUI/ConcertClientSharedSlate/Source/ConcertClientSharedSlate/Public/Replication/Editor/View/ReplicationColumnsUtils.h
new file mode 100644
--- /dev/null
+++ b/Engine/Plugins/Developer/Concert/ConcertUI/ConcertClientSharedSlate/Source/ConcertClientSharedSlate/Public/Replication/Editor/View/ReplicationColumnsUtils.h
@@ -0,0 +1,103 @@
+// Copyright Epic Games, Inc. All Rights Reserved.
+
+#pragma once
+
+#include "Containers/ArrayView.h"
+#include "Replication/Editor/View/ReplicationColumn.h"
+#include "Widgets/Input/SCheckBox.h"
+#include "Templates/SharedPointer.h"
+
+struct FConcertPropertyChain;
+
+namespace UE::ConcertClientSharedSlate
+{
+	class IEditableObjectToPropertiesModel;
+	class IReplicationStreamViewer;
+	class IObjectToPropertiesModel;
+}
+
+namespace UE::ConcertClientSharedSlate
+{
+	template<typename TListItemType>
+	struct TReplicationColumnDelegates
+	{
+		DECLARE_DELEGATE_RetVal_OneParam(ECheckBoxState, FGetColumnCheckboxState, const TListItemType& /*ObjectData*/);
+		DECLARE_DELEGATE_TwoParams(FOnColumnCheckboxChanged, bool /*bIsChecked*/, const TListItemType& /*ObjectData*/);
+		DECLARE_DELEGATE_RetVal_OneParam(bool, FIsEnabled, const TListItemType& /*ObjectData*/);
+		DECLARE_DELEGATE_RetVal_OneParam(FText, FGetToolTipText, const TListItemType& /*ObjectData*/);
+
+		TReplicationColumnDelegates(
+			FGetColumnCheckboxState InGetCheckboxStateDelegate,
+			FOnColumnCheckboxChanged InOnCheckboxChangedDelegate,
+			FGetToolTipText InGetToolTipTextDelegate = {},
+			FIsEnabled InIsEnabledDelegate = {}
+			)
+			: GetCheckboxStateDelegate(MoveTemp(InGetCheckboxStateDelegate))
+			, OnCheckboxChangedDelegate(MoveTemp(InOnCheckboxChangedDelegate))
+			, GetToolTipTextDelegate(MoveTemp(InGetToolTipTextDelegate))
+			, IsEnabledDelegate(MoveTemp(InIsEnabledDelegate))
+		{}
+		
+		TReplicationColumnDelegates(
+			FGetColumnCheckboxState InGetCheckboxStateDelegate,
+			FOnColumnCheckboxChanged InOnCheckboxChangedDelegate,
+			TAttribute<FText> InToolTipTextAttribute,
+			FIsEnabled InIsEnabledDelegate = {}
+			)
+			: GetCheckboxStateDelegate(MoveTemp(InGetCheckboxStateDelegate))
+			, OnCheckboxChangedDelegate(MoveTemp(InOnCheckboxChangedDelegate))
+			, GetToolTipTextDelegate(FGetToolTipText::CreateLambda([GetToolTipText = MoveTemp(InToolTipTextAttribute)](const TListItemType&){ return GetToolTipText.Get(); }))
+			, IsEnabledDelegate(MoveTemp(InIsEnabledDelegate))
+		{}
+
+		FGetColumnCheckboxState GetCheckboxStateDelegate;
+		FOnColumnCheckboxChanged OnCheckboxChangedDelegate;
+
+		/** Optional. Defaults to FText::GetEmpty() */
+		FGetToolTipText GetToolTipTextDelegate;
+		/** Optional. Default to true. */
+		FIsEnabled IsEnabledDelegate;
+	};
+
+	/** Util for making a checkbox column */
+	template<typename TListItemType>
+	TReplicationColumn<TListItemType> MakeCheckboxColumn(
+		FName ColumnId,
+		TReplicationColumnDelegates<TListItemType> Delegates,
+		FText DefaultLabel,
+		const int32 Priority,
+		const float ColumnWidth = 20.f
+		)
+	{
+		check(Delegates.GetCheckboxStateDelegate.IsBound() && Delegates.OnCheckboxChangedDelegate.IsBound());
+		return TReplicationColumn<TListItemType>(
+			typename TReplicationColumn<TListItemType>::FArguments()
+				.GenerateWidgetColumn_Lambda(
+					[Delegates](const typename TReplicationColumn<TListItemType>::FBuildArgs& Args)
+					{
+						return SNew(SCheckBox)
+							.ToolTipText_Lambda([GetToolTipDelegate = Delegates.GetToolTipTextDelegate, RowData = Args.RowData]()
+							{
+								return GetToolTipDelegate.IsBound() ? GetToolTipDelegate.Execute(RowData) : FText::GetEmpty();
+							})
+							.IsEnabled_Lambda([IsEnabledDelegate = Delegates.IsEnabledDelegate, RowData = Args.RowData]()
+							{
+								return !IsEnabledDelegate.IsBound() || IsEnabledDelegate.Execute(RowData);
+							})
+							.IsChecked_Lambda([IsCheckedDelegate = Delegates.GetCheckboxStateDelegate, RowData = Args.RowData]()
+							{
+								return IsCheckedDelegate.Execute(RowData);
+							})
+							.OnCheckStateChanged_Lambda([OnChangedDelegate = Delegates.OnCheckboxChangedDelegate, RowData = Args.RowData](ECheckBoxState NewState)
+							{
+								const bool bIsChecked = NewState == ECheckBoxState::Checked;
+								OnChangedDelegate.Execute(bIsChecked, RowData);
+							});
+					})
+				.ColumnSortOrder(Priority),
+			SHeaderRow::Column(ColumnId)
+				.DefaultLabel(DefaultLabel)
+				.FixedWidth(ColumnWidth)
+			);
+	}
+}
diff --git a/Engine/Plugins/Developer/Concert/ConcertUI/ConcertClientSharedSlate/Source/ConcertClientSharedSlate/Public/Replication/ReplicationWidgetFactories.h b/Engine/Plugins/Developer/Concert/ConcertUI/ConcertClientSharedSlate/Source/ConcertClientSharedSlate/Public/Replication/ReplicationWidgetFactories.h
--- a/Engine/Plugins/Developer/Concert/ConcertUI/ConcertClientSharedSlate/Source/ConcertClientSharedSlate/Public/Replication/ReplicationWidgetFactories.h
+++ b/Engine/Plugins/Developer/Concert/ConcertUI/ConcertClientSharedSlate/Source/ConcertClientSharedSlate/Public/Replication/ReplicationWidgetFactories.h
@@ -2,37 +2,26 @@
 
 #pragma once
 
-#include "ReplicationWidgetDelegates.h"
 #include "Editor/View/PredefinedReplicationColumns.h"
+#include "ReplicationWidgetDelegates.h"
+
 #include "Misc/Attribute.h"
 #include "Templates/SharedPointer.h"
 
-namespace UE::ConcertClientSharedSlate
-{
-	class IObjectToPropertiesModel;
-}
-
-namespace UE::ConcertClientSharedSlate
-{
-	class IReplicationStreamViewer;
-}
-
 class UObject;
 class SWidget;
 
 struct FConcertReplicationEditorSettings;
 struct FObjectReplicationMap;
 
-namespace UE::ConcertClientSharedSlate
-{
-	class IReplicationSubobjectView;
-}
-
 namespace UE::ConcertClientSharedSlate
 {
 	class IEditableObjectToPropertiesModel;
-	class IReplicationStreamEditor;
+	class IObjectToPropertiesModel;
 	class IObjectSelectionSourceModel;
+	class IReplicationStreamEditor;
+	class IReplicationStreamViewer;
+	class IReplicationSubobjectView;
 	class IPropertySelectionSourceModel;
 	
 	/**
@@ -52,8 +41,15 @@ namespace UE::ConcertClientSharedSlate
 		TAttribute<const FConcertReplicationEditorSettings*> OptionalReplicationSettingsAttribute = {}
 		);
 	
+	struct FCreateSubobjectViewParams
+	{
+		/** Additional columns that should be displayed for the subobjects. */
+		TArray<ReplicationColumns::FReplicationSubobjectObjectColumn> AdditionalColumns;
+	};
 	/** Creates a stream editor with a subobject view that looks like the SSubobjectEditor. */
-	CONCERTCLIENTSHAREDSLATE_API TSharedRef<IReplicationSubobjectView> CreateUnrealEditorSubobjectView();
+	CONCERTCLIENTSHAREDSLATE_API TSharedRef<IReplicationSubobjectView> CreateUnrealEditorSubobjectView(
+		FCreateSubobjectViewParams Params
+		);
 
 	/** Params for creating an IReplicationStreamEditor */
 	struct FCreateEditorParams
@@ -89,14 +85,19 @@ namespace UE::ConcertClientSharedSlate
 		FSortPropertyPredicate SortPropertyRowPredicate;
 		
 		/** Additional columns to add to the object view */
-		TArray<ReplicationObjectColumns::FReplicationObjectColumn> AdditionalObjectColumns;
+		TArray<ReplicationColumns::FReplicationTopLevelObjectColumn> AdditionalObjectColumns;
 		/** Additional columns to add to the property view */
-		TArray<ReplicationPropertyColumns::FReplicationPropertyColumn> AdditionalPropertyColumns;
+		TArray<ReplicationColumns::FReplicationPropertyColumn> AdditionalPropertyColumns;
 		
 		/** Optional widget to add to the left of the object list search bar. */
 		TAlwaysValidWidget LeftOfObjectSearchBar;
 		/** Optional widget to add to the left of the property list search bar. */
 		TAlwaysValidWidget LeftOfPropertySearchBar;
+
+		/** Optional. Determines whether all UI for changing the model should be disabled. */
+		TAttribute<bool> IsEditingEnabled;
+		/** Optional. Whenever IsEditingEnabled returns true, this tooltip is displayed for relevant, disabled UI. */
+		TAttribute<FText> EditingDisabledToolTipText;
 	};
 
 	/**
diff --git a/Engine/Plugins/Developer/Concert/ConcertUI/ConcertSharedSlate/Source/ConcertSharedSlate/Public/Model/Item/SourceModelBuilders.h b/Engine/Plugins/Developer/Concert/ConcertUI/ConcertSharedSlate/Source/ConcertSharedSlate/Public/Model/Item/SourceModelBuilders.h
--- a/Engine/Plugins/Developer/Concert/ConcertUI/ConcertSharedSlate/Source/ConcertSharedSlate/Public/Model/Item/SourceModelBuilders.h
+++ b/Engine/Plugins/Developer/Concert/ConcertUI/ConcertSharedSlate/Source/ConcertSharedSlate/Public/Model/Item/SourceModelBuilders.h
@@ -44,11 +44,25 @@ namespace UE::ConcertSharedSlate
 			/** Optional. Used primarily to filter out objects that are already selected. */
 			FIsItemSelected IsItemSelected;
 
-			FItemPickerArgs(FOnItemsSelected OnObjectsSelected, FGetItemDisplayString GetItemDisplayString, FGetItemIcon GetItemIcon = {}, FIsItemSelected IsItemSelected = {})
+			/** Optional. Whether the UI is enabled. */
+			TAttribute<bool> IsEnabledAttribute;
+			/** Optional. If IsEnabled returns true, this tooltip is displayed on the relevant UI. */
+			TAttribute<FText> DisabledToolTipAttribute;
+
+			FItemPickerArgs(
+				FOnItemsSelected OnObjectsSelected,
+				FGetItemDisplayString GetItemDisplayString,
+				FGetItemIcon GetItemIcon = {},
+				FIsItemSelected IsItemSelected = {},
+				TAttribute<bool> IsEnabledAttribute = {},
+				TAttribute<FText> DisabledToolTipAttribute = {}
+				)
 				: OnItemsSelected(MoveTemp(OnObjectsSelected))
 				, GetItemDisplayString(MoveTemp(GetItemDisplayString))
 				, GetItemIcon(MoveTemp(GetItemIcon))
 				, IsItemSelected(MoveTemp(IsItemSelected))
+				, IsEnabledAttribute(MoveTemp(IsEnabledAttribute))
+				, DisabledToolTipAttribute(MoveTemp(DisabledToolTipAttribute))
 			{
 				check(this->OnItemsSelected.IsBound() && this->GetItemDisplayString.IsBound());
 			}
@@ -97,8 +111,16 @@ namespace UE::ConcertSharedSlate
 
 		return SNew(SPositiveActionButton)
 			.Text(Category.DisplayInfo.Label)
-			.ToolTipText(Category.DisplayInfo.ToolTip)
+			.ToolTipText_Lambda([ToolTipText = Category.DisplayInfo.ToolTip, DisabledToolTipAttribute = Args.DisabledToolTipAttribute, IsEnabled = Args.IsEnabledAttribute]()
+			{
+				const bool bIsEnabled = !(IsEnabled.IsBound() || IsEnabled.IsSet()) || IsEnabled.Get();
+				const bool bHasDisabledText = DisabledToolTipAttribute.IsBound() || DisabledToolTipAttribute.IsSet();
+				return bIsEnabled
+					? ToolTipText
+					: bHasDisabledText ? DisabledToolTipAttribute.Get() : FText::GetEmpty();
+			})
 			.Icon(Category.DisplayInfo.Icon.IsSet() ? Category.DisplayInfo.Icon.GetOptionalIcon() : FAppStyle::Get().GetBrush("Icons.Plus"))
+			.IsEnabled_Lambda([IsEnabled = Args.IsEnabledAttribute](){ return !(IsEnabled.IsBound() || IsEnabled.IsSet()) || IsEnabled.Get(); })
 			.OnGetMenuContent_Lambda([Options = Category.Options, SubCategories = Category.SubCategories, Args]()
 			{
 				return BuildMenu(Options, SubCategories, Args);
@@ -116,8 +138,16 @@ namespace UE::ConcertSharedSlate
 		case ESourceType::ShowAsList:
 			return SNew(SPositiveActionButton)
 				.Text(DisplayInfo.Label)
-				.ToolTipText(DisplayInfo.ToolTip)
+				.ToolTipText_Lambda([ToolTipText = DisplayInfo.ToolTip, DisabledToolTipAttribute = Args.DisabledToolTipAttribute, IsEnabled = Args.IsEnabledAttribute]()
+				{
+					const bool bIsEnabled = !(IsEnabled.IsBound() || IsEnabled.IsSet()) || IsEnabled.Get();
+					const bool bHasDisabledText = DisabledToolTipAttribute.IsBound() || DisabledToolTipAttribute.IsSet();
+					return bIsEnabled
+						? ToolTipText
+						: bHasDisabledText ? DisabledToolTipAttribute.Get() : FText::GetEmpty();
+				})
 				.Icon(IconBrush)
+				.IsEnabled_Lambda([IsEnabled = Args.IsEnabledAttribute](){ return !(IsEnabled.IsBound() || IsEnabled.IsSet()) || IsEnabled.Get(); })
 				.OnGetMenuContent_Lambda([Source = MoveTemp(Source), Args = MoveTemp(Args)]()
 				{
 					FMenuBuilder MenuBuilder(true, nullptr);
@@ -127,8 +157,16 @@ namespace UE::ConcertSharedSlate
 		case ESourceType::AddOnClick: 
 			return SNew(SPositiveActionButton)
 				.Text(DisplayInfo.Label)
-				.ToolTipText(DisplayInfo.ToolTip)
+				.ToolTipText_Lambda([ToolTipText = DisplayInfo.ToolTip, DisabledToolTipAttribute = Args.DisabledToolTipAttribute, IsEnabled = Args.IsEnabledAttribute]()
+				{
+					const bool bIsEnabled = !(IsEnabled.IsBound() || IsEnabled.IsSet()) || IsEnabled.Get();
+					const bool bHasDisabledText = DisabledToolTipAttribute.IsBound() || DisabledToolTipAttribute.IsSet();
+					return bIsEnabled
+						? ToolTipText
+						: bHasDisabledText ? DisabledToolTipAttribute.Get() : FText::GetEmpty();
+				})
 				.Icon(IconBrush)
+				.IsEnabled_Lambda([IsEnabled = Args.IsEnabledAttribute](){ return !(IsEnabled.IsBound() || IsEnabled.IsSet()) || IsEnabled.Get(); })
 				.OnClicked_Lambda([Source = MoveTemp(Source), AddObjectsDelegate = Args.OnItemsSelected]()
 				{
 					AddObjectsDelegate.Execute(Source->GetSelectableItems());
