Fix duplicate loca keys

diff --git a/Engine/Plugins/Developer/Concert/ConcertApp/MultiUserClient/Source/MultiUserClient/Private/Widgets/ActiveSession/Replication/Client/SSubmissionControls.cpp b/Engine/Plugins/Developer/Concert/ConcertApp/MultiUserClient/Source/MultiUserClient/Private/Widgets/ActiveSession/Replication/Client/SSubmissionControls.cpp
--- a/Engine/Plugins/Developer/Concert/ConcertApp/MultiUserClient/Source/MultiUserClient/Private/Widgets/ActiveSession/Replication/Client/SSubmissionControls.cpp
+++ b/Engine/Plugins/Developer/Concert/ConcertApp/MultiUserClient/Source/MultiUserClient/Private/Widgets/ActiveSession/Replication/Client/SSubmissionControls.cpp
@@ -160,7 +160,7 @@ namespace UE::MultiUserClient
 				if (SubmitChangesResult.ErrorCode == EStreamSubmissionErrorCode::Cancelled)
 				{
 					Notification->SetCompletionState(SNotificationItem::CS_Fail);
-					Notification->SetText(LOCTEXT("Uploading.Failure", "Updating Stream cancelled."));
+					Notification->SetText(LOCTEXT("Uploading.Failure.Cancelled", "Updating Stream cancelled."));
 					return;
 				}
 
@@ -182,8 +182,8 @@ namespace UE::MultiUserClient
 				else
 				{
 					Notification->SetCompletionState(SNotificationItem::CS_Fail);
-					Notification->SetText(LOCTEXT("Uploading.Failure", "Changes rejected."));
-					Notification->SetSubText(LOCTEXT("Uploading.Failure.SubText", "See log for details."));
+					Notification->SetText(LOCTEXT("Uploading.Failure.Rejected", "Changes rejected."));
+					Notification->SetSubText(LOCTEXT("Uploading.Failure.Rejected.SubText", "See log for details."));
 
 					if (SubmissionInfo.IsSet())
 					{
diff --git a/Engine/Plugins/Developer/Concert/ConcertApp/MultiUserClient/Source/MultiUserClient/Private/Widgets/ActiveSession/Replication/Client/Single/StreamEditorColumns.cpp b/Engine/Plugins/Developer/Concert/ConcertApp/MultiUserClient/Source/MultiUserClient/Private/Widgets/ActiveSession/Replication/Client/Single/StreamEditorColumns.cpp
--- a/Engine/Plugins/Developer/Concert/ConcertApp/MultiUserClient/Source/MultiUserClient/Private/Widgets/ActiveSession/Replication/Client/Single/StreamEditorColumns.cpp
+++ b/Engine/Plugins/Developer/Concert/ConcertApp/MultiUserClient/Source/MultiUserClient/Private/Widgets/ActiveSession/Replication/Client/Single/StreamEditorColumns.cpp
@@ -133,8 +133,8 @@ namespace UE::MultiUserClient::StreamEditorColumns
 								switch (CheckBoxState)
 								{
 								case ECheckBoxState::Unchecked: return LOCTEXT("TopLevel.ChangeAuthority.ToolTip.NotSupported.Unchecked", "None of the subobjects are being replicated.");
-								case ECheckBoxState::Checked: return LOCTEXT("TopLevel.ChangeAuthority.ToolTip.NotSupported.Unchecked", "All of the subobjects are being replicated.");
-								case ECheckBoxState::Undetermined: return LOCTEXT("TopLevel.ChangeAuthority.ToolTip.NotSupported.Unchecked", "Some of the subobjects are being replicated.");
+								case ECheckBoxState::Checked: return LOCTEXT("TopLevel.ChangeAuthority.ToolTip.NotSupported.Checked", "All of the subobjects are being replicated.");
+								case ECheckBoxState::Undetermined: return LOCTEXT("TopLevel.ChangeAuthority.ToolTip.NotSupported.Undetermined", "Some of the subobjects are being replicated.");
 								default: checkNoEntry(); return FText::GetEmpty();
 								}
 							}();
@@ -182,7 +182,7 @@ namespace UE::MultiUserClient::StreamEditorColumns
 						const bool bHasAuthority = ChangeTracker.GetAuthorityStateAfterApplied(ObjectData.GetObjectPath());
 						const FText NotSupportedText = FText::Format(
 							LOCTEXT("Subobject.ChangeAuthority.ToolTip.NotSupportedFmt", "{0}\nEditing remote clients is not implemented. You can only edit the local client."),
-							bHasAuthority ? LOCTEXT("Subobject.ChangeAuthority.Replicating.True", "This object is being replicated.") : LOCTEXT("Subobject.ChangeAuthority.Replicating.True", "This object is not currently being replicated.")
+							bHasAuthority ? LOCTEXT("Subobject.ChangeAuthority.Replicating.True", "This object is being replicated.") : LOCTEXT("Subobject.ChangeAuthority.Replicating.False", "This object is not currently being replicated.")
 							);
 						if (SubmissionWorkflow.GetUploadability() == EChangeUploadability::NotImplemented)
 						{
