Avoid dangling memory when remote client array is resized

diff --git a/Engine/Plugins/Developer/Concert/ConcertApp/MultiUserClient/Source/MultiUserClient/Private/Replication/Client/ReplicationClientManager.cpp b/Engine/Plugins/Developer/Concert/ConcertApp/MultiUserClient/Source/MultiUserClient/Private/Replication/Client/ReplicationClientManager.cpp
--- a/Engine/Plugins/Developer/Concert/ConcertApp/MultiUserClient/Source/MultiUserClient/Private/Replication/Client/ReplicationClientManager.cpp
+++ b/Engine/Plugins/Developer/Concert/ConcertApp/MultiUserClient/Source/MultiUserClient/Private/Replication/Client/ReplicationClientManager.cpp
@@ -44,12 +44,23 @@ namespace UE::MultiUserClient
 		}
 	}
 
+	TArray<TNonNullPtr<FRemoteReplicationClient>> FReplicationClientManager::GetRemoteClients() const
+	{
+		TArray<TNonNullPtr<FRemoteReplicationClient>> Result;
+		Algo::Transform(RemoteClients, Result, [](const TUniquePtr<FRemoteReplicationClient>& Client) -> TNonNullPtr<FRemoteReplicationClient>
+		{
+			return Client.Get();
+		});
+		return Result;
+	}
+
 	const FRemoteReplicationClient* FReplicationClientManager::FindRemoteClient(const FGuid& EndpointId) const
 	{
-		return RemoteClients.FindByPredicate([&EndpointId](const FRemoteReplicationClient& Client)
+		const TUniquePtr<FRemoteReplicationClient>* Client = RemoteClients.FindByPredicate([&EndpointId](const TUniquePtr<FRemoteReplicationClient>& Client)
 		{
-			return Client.GetRemoteEndpointId() == EndpointId;
+			return Client->GetRemoteEndpointId() == EndpointId;
 		});
+		return Client ? Client->Get() : nullptr;
 	}
 
 	void FReplicationClientManager::AddReferencedObjects(FReferenceCollector& Collector)
@@ -70,18 +81,21 @@ namespace UE::MultiUserClient
 		case EConcertClientStatus::Disconnected:
 			{
 				const int32 Index = RemoteClients.IndexOfByPredicate(
-					[&ClientEndpointId](const FRemoteReplicationClient& Client)
+					[&ClientEndpointId](const TUniquePtr<FRemoteReplicationClient>& Client)
 					{
-						return Client.GetRemoteEndpointId() == ClientEndpointId;
+						return Client->GetRemoteEndpointId() == ClientEndpointId;
 					});
 				if (!ensure(RemoteClients.IsValidIndex(Index)))
 				{
 					return;
 				}
-				
-				const FRemoteReplicationClient& Client = RemoteClients[Index];
-				SessionContent->RemoveClient(*Client.GetClientContent());
-				RemoteClients.RemoveAtSwap(Index);
+
+				{
+					const TUniquePtr<FRemoteReplicationClient> Client = MoveTemp(RemoteClients[Index]);
+					SessionContent->RemoveClient(*Client->GetClientContent());
+					RemoteClients.RemoveAtSwap(Index);
+				}
+				// We want to broadcast after the client has been fully cleaned up
 				OnRemoteClientsChangedDelegate.Broadcast();
 			}
 			break;
@@ -94,7 +108,9 @@ namespace UE::MultiUserClient
 	
 	void FReplicationClientManager::CreateRemoteClient(const FGuid& ClientEndpointId, bool bBroadcastDelegate)
 	{
-		RemoteClients.Emplace(*SessionContent->AddClient(), ClientEndpointId, QueryService);
+		RemoteClients.Emplace(
+			MakeUnique<FRemoteReplicationClient>(*SessionContent->AddClient(), ClientEndpointId, QueryService)
+			);
 
 		if (bBroadcastDelegate)
 		{
diff --git a/Engine/Plugins/Developer/Concert/ConcertApp/MultiUserClient/Source/MultiUserClient/Private/Replication/Client/ReplicationClientManager.h b/Engine/Plugins/Developer/Concert/ConcertApp/MultiUserClient/Source/MultiUserClient/Private/Replication/Client/ReplicationClientManager.h
--- a/Engine/Plugins/Developer/Concert/ConcertApp/MultiUserClient/Source/MultiUserClient/Private/Replication/Client/ReplicationClientManager.h
+++ b/Engine/Plugins/Developer/Concert/ConcertApp/MultiUserClient/Source/MultiUserClient/Private/Replication/Client/ReplicationClientManager.h
@@ -10,12 +10,13 @@
 
 #include "UObject/GCObject.h"
 #include "Templates/UnrealTemplate.h"
-
-enum class EConcertClientStatus : uint8;
+#include "Templates/NonNullPointer.h"
 
 class IConcertClientSession;
 class IConcertSyncClient;
 
+enum class EConcertClientStatus : uint8;
+
 struct FConcertSessionClientInfo;
 
 namespace UE::ConcertClientSharedSlate
@@ -48,7 +49,7 @@ namespace UE::MultiUserClient
 
 		const FLocalReplicationClient& GetLocalClient() const { return LocalClient; }
 		FLocalReplicationClient& GetLocalClient() { return LocalClient; }
-		const TArray<FRemoteReplicationClient>& GetRemoteClients() const { return RemoteClients; }
+		TArray<TNonNullPtr<FRemoteReplicationClient>> GetRemoteClients() const;
 
 		/** Util for finding a remote client by its EndpointId. */
 		const FRemoteReplicationClient* FindRemoteClient(const FGuid& EndpointId) const;
@@ -88,8 +89,11 @@ namespace UE::MultiUserClient
 		
 		/** Manages the local client */
 		FLocalReplicationClient LocalClient;
-		/** Manages remote clients. Updated when client connects or disconnects to the active session. */
-		TArray<FRemoteReplicationClient> RemoteClients;
+		/**
+		 * Manages remote clients. Updated when client connects or disconnects to the active session.
+		 * UI keeps references to systems inside the client so it is TSharedRef in case TArray is reallocated.
+		 */
+		TArray<TUniquePtr<FRemoteReplicationClient>> RemoteClients;
 
 		/** Called when RemoteClients changes. */
 		FRemoteClientsChanged OnRemoteClientsChangedDelegate;
diff --git a/Engine/Plugins/Developer/Concert/ConcertApp/MultiUserClient/Source/MultiUserClient/Private/Widgets/ActiveSession/Replication/Joined/SReplicationJoinedView.cpp b/Engine/Plugins/Developer/Concert/ConcertApp/MultiUserClient/Source/MultiUserClient/Private/Widgets/ActiveSession/Replication/Joined/SReplicationJoinedView.cpp
--- a/Engine/Plugins/Developer/Concert/ConcertApp/MultiUserClient/Source/MultiUserClient/Private/Widgets/ActiveSession/Replication/Joined/SReplicationJoinedView.cpp
+++ b/Engine/Plugins/Developer/Concert/ConcertApp/MultiUserClient/Source/MultiUserClient/Private/Widgets/ActiveSession/Replication/Joined/SReplicationJoinedView.cpp
@@ -95,9 +95,9 @@ namespace UE::MultiUserClient
 	void SReplicationJoinedView::RebuildClientViewSwitcherChildren(const TArray<TSharedRef<SWidget>> OldClientWidgets)
 	{
 		TMap<FGuid, int32> OldRemoteClientToWidgetSwitcherIndex = MoveTemp(RemoteClientToWidgetSwitcherIndex);
-        for (const FRemoteReplicationClient& RemoteClient : ReplicationManager->GetClientManager()->GetRemoteClients())
+        for (const TNonNullPtr<FRemoteReplicationClient>& RemoteClient : ReplicationManager->GetClientManager()->GetRemoteClients())
         {
-        	const FGuid& EndpointId = RemoteClient.GetRemoteEndpointId();
+        	const FGuid& EndpointId = RemoteClient->GetRemoteEndpointId();
         	
         	const int32* ExistingClientWidgetIndex = OldRemoteClientToWidgetSwitcherIndex.Find(EndpointId);
         	if (ExistingClientWidgetIndex)
@@ -112,7 +112,7 @@ namespace UE::MultiUserClient
         		ClientViewSwitcher->AddSlot()
         		[
         			SNew(SReplicationClientView)
-        			.GetReplicationClient_Lambda([this, EndpointId = RemoteClient.GetRemoteEndpointId()]()
+        			.GetReplicationClient_Lambda([this, EndpointId = RemoteClient->GetRemoteEndpointId()]()
         			{
         				// It is unsafe to simply capture RemoteClient because the containing TArray may reallocate its location
         				return ReplicationManager->GetClientManager()->FindRemoteClient(EndpointId);
@@ -162,9 +162,9 @@ namespace UE::MultiUserClient
 			.Client(Client->GetConcertClient())
 			.SelectableClients_Lambda([this]()
 			{
-				const TArray<FRemoteReplicationClient>& RemoteClients = ReplicationManager->GetClientManager()->GetRemoteClients();
+				const TArray<TNonNullPtr<FRemoteReplicationClient>> RemoteClients = ReplicationManager->GetClientManager()->GetRemoteClients();
 				TArray<FGuid> Result;
-				Algo::Transform(RemoteClients, Result, [](const FRemoteReplicationClient& InClient){ return InClient.GetRemoteEndpointId(); });
+				Algo::Transform(RemoteClients, Result, [](const TNonNullPtr<FRemoteReplicationClient>& InClient){ return InClient->GetRemoteEndpointId(); });
 
 				const TSharedPtr<IConcertClientSession> CurrentSession = Client->GetConcertClient()->GetCurrentSession();
 				Result.Sort([&CurrentSession](const FGuid& Left, const FGuid& Right)
