Add way to configure automatically added properties

diff --git a/Engine/Plugins/Developer/Concert/ConcertApp/MultiUserClient/Config/BaseMultiUserClient.ini b/Engine/Plugins/Developer/Concert/ConcertApp/MultiUserClient/Config/BaseMultiUserClient.ini
--- a/Engine/Plugins/Developer/Concert/ConcertApp/MultiUserClient/Config/BaseMultiUserClient.ini
+++ b/Engine/Plugins/Developer/Concert/ConcertApp/MultiUserClient/Config/BaseMultiUserClient.ini
@@ -1,3 +1,2 @@
-[/Script/MultiUserReplicationEditor.MultiUserReplicationSettings]
-+DefaultPropertySelection=(("/Script/Engine.SceneComponent", (DefaultSelectedProperties=("RelativeLocation","RelativeLocation.X","RelativeLocation.Y","RelativeLocation.Z","RelativeRotation","RelativeRotation.X","RelativeRotation.Y","RelativeRotation.Z","RelativeScale","RelativeScale.X","RelativeScale.Y","RelativeScale.Z"))))
-+DefaultSubobjectSelection=(("/Script/Engine.Actor", (IncludeClasses=("/Script/Engine.StaticMeshComponent","/Script/Engine.CameraComponent"))))
\ No newline at end of file
+[/Script/MultiUserClient.MultiUserReplicationSettings]
+ReplicationEditorSettings=(DefaultPropertySelection=(("/Script/Engine.SceneComponent", (DefaultSelectedProperties=("RelativeLocation","RelativeLocation.X","RelativeLocation.Y","RelativeLocation.Z","RelativeRotation"))),("/Script/CinematicCamera.CineCameraComponent", ()),("/Script/Engine.CameraComponent", ())),DefaultSubobjectSelection=(("/Script/Engine.Actor", (IncludeClasses=("/Script/Engine.StaticMeshComponent","/Script/Engine.CameraComponent","/Script/CinematicCamera.CineCameraComponent")))))
diff --git a/Engine/Plugins/Developer/Concert/ConcertApp/MultiUserClient/Source/MultiUserClient/Private/MultiUserClientModule.cpp b/Engine/Plugins/Developer/Concert/ConcertApp/MultiUserClient/Source/MultiUserClient/Private/MultiUserClientModule.cpp
--- a/Engine/Plugins/Developer/Concert/ConcertApp/MultiUserClient/Source/MultiUserClient/Private/MultiUserClientModule.cpp
+++ b/Engine/Plugins/Developer/Concert/ConcertApp/MultiUserClient/Source/MultiUserClient/Private/MultiUserClientModule.cpp
@@ -50,6 +50,7 @@
 #include "DetailLayoutBuilder.h"
 #include "DetailCategoryBuilder.h"
 #include "DetailWidgetRow.h"
+#include "MultiUserReplicationSettings.h"
 
 
 #include "Widgets/SConcertBrowser.h"
@@ -1277,7 +1278,8 @@ private:
 			ISettingsSectionPtr SettingsSection = SettingsModule->RegisterSettings("Project", "Plugins", "Concert",
 				LOCTEXT("ConcertFrontendSettingsName", "Multi-User Editing"),
 				LOCTEXT("ConcertFrontendSettingsDescription", "Configure the Multi-User settings."),
-				GetMutableDefault<UConcertClientConfig>());
+				GetMutableDefault<UConcertClientConfig>()
+				);
 
 			if (SettingsSection.IsValid())
 			{
@@ -1287,7 +1289,14 @@ private:
 			SettingsModule->RegisterSettings("Project", "Plugins", "Concert Sync",
 				LOCTEXT("ConcertFrontendSyncSettingsName", "Multi-User Transactions"),
 				LOCTEXT("ConcertFrontendSyncSettingsDescription", "Configure the Multi-User Transactions settings."),
-				GetMutableDefault<UConcertSyncConfig>());
+				GetMutableDefault<UConcertSyncConfig>()
+				);
+
+			SettingsModule->RegisterSettings("Project", "Plugins", "Multi User",
+				LOCTEXT("ConcertFrontendSyncSettingsName", "Multi-User Replication"),
+				LOCTEXT("ConcertFrontendSyncSettingsDescription", "Configure the Multi-User Replication settings."),
+				UMultiUserReplicationSettings::Get()
+				);
 
 			// Register a special layout to validate user input when configuring the settings.
 			FPropertyEditorModule& PropertyModule = FModuleManager::LoadModuleChecked<FPropertyEditorModule>("PropertyEditor");
diff --git a/Engine/Plugins/Developer/Concert/ConcertApp/MultiUserClient/Source/MultiUserClient/Private/MultiUserReplicationSettings.h b/Engine/Plugins/Developer/Concert/ConcertApp/MultiUserClient/Source/MultiUserClient/Private/MultiUserReplicationSettings.h
new file mode 100644
--- /dev/null
+++ b/Engine/Plugins/Developer/Concert/ConcertApp/MultiUserClient/Source/MultiUserClient/Private/MultiUserReplicationSettings.h
@@ -0,0 +1,20 @@
+// Copyright Epic Games, Inc. All Rights Reserved.
+
+#pragma once
+
+#include "UObject/Object.h"
+#include "Replication/Settings/ConcertReplicationEditorSettings.h"
+#include "MultiUserReplicationSettings.generated.h"
+
+UCLASS(config=MultiUserClient)
+class UMultiUserReplicationSettings : public UObject
+{
+	GENERATED_BODY()
+public:
+
+	static UMultiUserReplicationSettings* Get() { return GetMutableDefault<UMultiUserReplicationSettings>(); }
+
+	/** Settings that affect the replication editor */
+	UPROPERTY(Config, Category = "Replication Settings", EditAnywhere)
+	FConcertReplicationEditorSettings ReplicationEditorSettings;
+};
\ No newline at end of file
diff --git a/Engine/Plugins/Developer/Concert/ConcertApp/MultiUserClient/Source/MultiUserClient/Private/Replication/Client/ReplicationClient.cpp b/Engine/Plugins/Developer/Concert/ConcertApp/MultiUserClient/Source/MultiUserClient/Private/Replication/Client/ReplicationClient.cpp
--- a/Engine/Plugins/Developer/Concert/ConcertApp/MultiUserClient/Source/MultiUserClient/Private/Replication/Client/ReplicationClient.cpp
+++ b/Engine/Plugins/Developer/Concert/ConcertApp/MultiUserClient/Source/MultiUserClient/Private/Replication/Client/ReplicationClient.cpp
@@ -2,6 +2,7 @@
 
 #include "ReplicationClient.h"
 
+#include "MultiUserReplicationSettings.h"
 #include "Assets/MultiUserReplicationClientPreset.h"
 #include "Replication/Editor/Model/IEditableObjectToPropertiesModel.h"
 #include "Replication/ReplicationWidgetFactories.h"
@@ -18,7 +19,11 @@ namespace UE::MultiUserClient
 		: ClientContentStorage(&InSessionContent)
 		, StreamSynchronizer(MoveTemp(InStreamSynchronizer))
 		, AuthoritySynchronizer(MoveTemp(InAuthoritySynchronizer))
-		, LocalClientEditModel(ConcertClientSharedSlate::CreatePropertySelectionModel(*ClientContentStorage->Stream, ClientContentStorage->Stream->MakeReplicationMapGetterAttribute()))
+		, LocalClientEditModel(ConcertClientSharedSlate::CreatePropertySelectionModel(
+			*ClientContentStorage->Stream,
+			ClientContentStorage->Stream->MakeReplicationMapGetterAttribute(),
+			TAttribute<const FConcertReplicationEditorSettings*>::CreateLambda([](){ return &UMultiUserReplicationSettings::Get()->ReplicationEditorSettings; })
+			))
 		, LocalClientStreamDiffer(
 			GetStreamSynchronizer(),
 			ClientContentStorage->Stream->MakeReplicationMapGetterAttribute(),
