Fix duplicate loca keys

diff --git a/Engine/Plugins/Developer/Concert/ConcertApp/MultiUserClient/Source/MultiUserClient/Private/MultiUserClientModule.cpp b/Engine/Plugins/Developer/Concert/ConcertApp/MultiUserClient/Source/MultiUserClient/Private/MultiUserClientModule.cpp
--- a/Engine/Plugins/Developer/Concert/ConcertApp/MultiUserClient/Source/MultiUserClient/Private/MultiUserClientModule.cpp
+++ b/Engine/Plugins/Developer/Concert/ConcertApp/MultiUserClient/Source/MultiUserClient/Private/MultiUserClientModule.cpp
@@ -1293,8 +1293,8 @@ private:
 				);
 
 			SettingsModule->RegisterSettings("Project", "Plugins", "Multi User",
-				LOCTEXT("ConcertFrontendSyncSettingsName", "Multi-User Replication"),
-				LOCTEXT("ConcertFrontendSyncSettingsDescription", "Configure the Multi-User Replication settings."),
+				LOCTEXT("MultiUserReplicationSettingsName", "Multi-User Replication"),
+				LOCTEXT("MultiUserReplicationSettingsDescription", "Configure the Multi-User Replication settings."),
 				UMultiUserReplicationSettings::Get()
 				);
 
