Handling timeouts:
- Add way to detect timeouts
- Handle timeouts appropriately in the UI
- Add console variables for simulating timeouts client side (for UI testing)

diff --git a/Engine/Plugins/Developer/Concert/ConcertApp/MultiUserClient/Source/MultiUserClient/Private/Replication/Submission/Data/AuhtoritySubmission.cpp b/Engine/Plugins/Developer/Concert/ConcertApp/MultiUserClient/Source/MultiUserClient/Private/Replication/Submission/Data/AuhtoritySubmission.cpp
--- a/Engine/Plugins/Developer/Concert/ConcertApp/MultiUserClient/Source/MultiUserClient/Private/Replication/Submission/Data/AuhtoritySubmission.cpp
+++ b/Engine/Plugins/Developer/Concert/ConcertApp/MultiUserClient/Source/MultiUserClient/Private/Replication/Submission/Data/AuhtoritySubmission.cpp
@@ -6,15 +6,15 @@
 
 namespace UE::MultiUserClient
 {
-	FText LexToText(EAuthoritySubmissionErrorCode ErrorCode)
+	FText LexToText(EAuthoritySubmissionResponseErrorCode ErrorCode)
 	{
 		switch (ErrorCode)
 		{
-		case EAuthoritySubmissionErrorCode::Success: return LOCTEXT("Success", "Success");
-		case EAuthoritySubmissionErrorCode::NoChange: return LOCTEXT("NoChange", "No Change");
-		case EAuthoritySubmissionErrorCode::Timeout: return LOCTEXT("Timeout", "Timeout");
-		case EAuthoritySubmissionErrorCode::CancelledDueToStreamUpdate: return LOCTEXT("CancelledDueToStreamUpdate", "Failed because stream dependency could not be updated");
-		case EAuthoritySubmissionErrorCode::Cancelled: return LOCTEXT("Cancelled", "Cancelled");
+		case EAuthoritySubmissionResponseErrorCode::Success: return LOCTEXT("Success", "Success");
+		case EAuthoritySubmissionResponseErrorCode::NoChange: return LOCTEXT("NoChange", "No Change");
+		case EAuthoritySubmissionResponseErrorCode::Timeout: return LOCTEXT("Timeout", "Timeout");
+		case EAuthoritySubmissionResponseErrorCode::CancelledDueToStreamUpdate: return LOCTEXT("CancelledDueToStreamUpdate", "Failed because stream dependency could not be updated");
+		case EAuthoritySubmissionResponseErrorCode::Cancelled: return LOCTEXT("Cancelled", "Cancelled");
 		default:
 			checkNoEntry();
 			return LOCTEXT("Unknown", "Unknown");;
diff --git a/Engine/Plugins/Developer/Concert/ConcertApp/MultiUserClient/Source/MultiUserClient/Private/Replication/Submission/Data/AuthoritySubmission.h b/Engine/Plugins/Developer/Concert/ConcertApp/MultiUserClient/Source/MultiUserClient/Private/Replication/Submission/Data/AuthoritySubmission.h
--- a/Engine/Plugins/Developer/Concert/ConcertApp/MultiUserClient/Source/MultiUserClient/Private/Replication/Submission/Data/AuthoritySubmission.h
+++ b/Engine/Plugins/Developer/Concert/ConcertApp/MultiUserClient/Source/MultiUserClient/Private/Replication/Submission/Data/AuthoritySubmission.h
@@ -8,7 +8,7 @@ namespace UE::MultiUserClient
 {
 	class ISubmissionOperation;
 	
-	enum class EAuthoritySubmissionErrorCode
+	enum class EAuthoritySubmissionRequestErrorCode
 	{
 		/** The operation completed as expected */
 		Success,
@@ -16,9 +16,6 @@ namespace UE::MultiUserClient
 		/** No change was sent to the server because there were no local changes. */
 		NoChange, 
 
-		/** The request timed out */
-		Timeout,
-
 		/** Cancelled because the required stream update could not be made. */
 		CancelledDueToStreamUpdate,
 		
@@ -26,14 +23,29 @@ namespace UE::MultiUserClient
 		Cancelled
 	};
 
-	FText LexToText(EAuthoritySubmissionErrorCode ErrorCode);
+	enum class EAuthoritySubmissionResponseErrorCode
+    {
+    	/** The operation completed as expected */
+    	Success,
+
+    	/** No change was sent to the server because there were no local changes. */
+    	NoChange, 
+
+    	/** The request timed out */
+    	Timeout,
+
+    	/** Cancelled because the required stream update could not be made. */
+    	CancelledDueToStreamUpdate,
+    	
+    	/** The request was cancelled (e.g. because client disconnected while operation was in progress) */
+    	Cancelled
+    };
+
+	FText LexToText(EAuthoritySubmissionResponseErrorCode ErrorCode);
 
 	struct FSubmitAuthorityChangesRequest
 	{
-		EAuthoritySubmissionErrorCode ErrorCode;
-		
-		/** The operation this was executed as part of. You cannot keep any reference to this. */
-		ISubmissionOperation& OperationContext;
+		EAuthoritySubmissionRequestErrorCode ErrorCode;
 		
 		/** Only valid if ErrorCode == EAuthoritySubmissionErrorCode::Success */
 		TOptional<ConcertSyncClient::Replication::FAuthorityChangeRequest> Request;
@@ -41,10 +53,7 @@ namespace UE::MultiUserClient
 	
 	struct FSubmitAuthorityChangesResponse
 	{
-		EAuthoritySubmissionErrorCode ErrorCode;
-		
-		/** The operation this was executed as part of. You cannot keep any reference to this. */
-		ISubmissionOperation& OperationContext;
+		EAuthoritySubmissionResponseErrorCode ErrorCode;
 		
 		/** Only valid if ErrorCode == EAuthoritySubmissionErrorCode::Success */
 		TOptional<ConcertSyncClient::Replication::FAuthorityChangeResponse> Response;
diff --git a/Engine/Plugins/Developer/Concert/ConcertApp/MultiUserClient/Source/MultiUserClient/Private/Replication/Submission/Data/StreamSubmission.h b/Engine/Plugins/Developer/Concert/ConcertApp/MultiUserClient/Source/MultiUserClient/Private/Replication/Submission/Data/StreamSubmission.h
--- a/Engine/Plugins/Developer/Concert/ConcertApp/MultiUserClient/Source/MultiUserClient/Private/Replication/Submission/Data/StreamSubmission.h
+++ b/Engine/Plugins/Developer/Concert/ConcertApp/MultiUserClient/Source/MultiUserClient/Private/Replication/Submission/Data/StreamSubmission.h
@@ -16,6 +16,9 @@ namespace UE::MultiUserClient
 		/** No change was sent to the server because there were no local changes. */
 		NoChange,
 
+		/** The request timed out */
+		Timeout,
+
 		/** The request was cancelled (e.g. because client disconnected while operation was in progress) */
 		Cancelled
 	};
@@ -33,9 +36,6 @@ namespace UE::MultiUserClient
 		/** Error code of the submission. Determines whether Response is valid. */
 		EStreamSubmissionErrorCode ErrorCode;
 
-		/** The operation this was executed as part of. You cannot keep any reference to this. */
-		ISubmissionOperation& OperationContext; 
-
 		/** Valid if ErrorCode != ESubmitChangesErrorCode::AlreadyInProgress. */
 		TOptional<FCompletedChangeSubmission> SubmissionInfo;
 	};
diff --git a/Engine/Plugins/Developer/Concert/ConcertApp/MultiUserClient/Source/MultiUserClient/Private/Replication/Submission/ISubmissionWorkflow.h b/Engine/Plugins/Developer/Concert/ConcertApp/MultiUserClient/Source/MultiUserClient/Private/Replication/Submission/ISubmissionWorkflow.h
--- a/Engine/Plugins/Developer/Concert/ConcertApp/MultiUserClient/Source/MultiUserClient/Private/Replication/Submission/ISubmissionWorkflow.h
+++ b/Engine/Plugins/Developer/Concert/ConcertApp/MultiUserClient/Source/MultiUserClient/Private/Replication/Submission/ISubmissionWorkflow.h
@@ -5,6 +5,9 @@
 #include "Data/EChangeRevertability.h"
 #include "Data/EChangeUploadability.h"
 
+#include "Templates/Function.h"
+#include "Templates/SharedPointer.h"
+
 namespace UE::MultiUserClient
 {
 	class ISubmissionOperation;
@@ -23,11 +26,16 @@ namespace UE::MultiUserClient
 	public:
 
 		/**
-		 * Synchronizes the server with the locally made changes streams and authority.
-		 * @return Operation object which emits special events. You cannot keep any reference to it. Null if there is already an operation in progress.
+		 * Synchronizes the server with the locally made changes to streams and authority.
+		 *
+		 * This function creates an operation object which emits special events.
+		 * If there is already an operation in progress, this function fails.
 		 * @see CanSubmit
+		 *
+		 * @note The operation might start and instantly stop before SubmitChanges finishes (e.g. a network request fails to be created instantly).
+		 * @return The operation object if the operation was started
 		 */
-		virtual ISubmissionOperation* SubmitChanges() = 0;
+		virtual TSharedPtr<ISubmissionOperation> SubmitChanges() = 0;
 
 		/** Reverts all local changes */
 		virtual void RevertChanges() = 0;
diff --git a/Engine/Plugins/Developer/Concert/ConcertApp/MultiUserClient/Source/MultiUserClient/Private/Replication/Submission/SingleClientSubmissionOperation.h b/Engine/Plugins/Developer/Concert/ConcertApp/MultiUserClient/Source/MultiUserClient/Private/Replication/Submission/SingleClientSubmissionOperation.h
--- a/Engine/Plugins/Developer/Concert/ConcertApp/MultiUserClient/Source/MultiUserClient/Private/Replication/Submission/SingleClientSubmissionOperation.h
+++ b/Engine/Plugins/Developer/Concert/ConcertApp/MultiUserClient/Source/MultiUserClient/Private/Replication/Submission/SingleClientSubmissionOperation.h
@@ -24,15 +24,15 @@ namespace UE::MultiUserClient
 		{
 			if (!bStreamPromiseWasSet)
 			{
-				StreamChangesPromise.EmplaceValue(FSubmitStreamChangesResponse{ EStreamSubmissionErrorCode::Cancelled, *this });
+				StreamChangesPromise.EmplaceValue(FSubmitStreamChangesResponse{ EStreamSubmissionErrorCode::Cancelled });
 			}
 			if (!bAuthorityRequestPromiseWasSet)
 			{
-				AuthorityChangeRequestPromise.EmplaceValue(FSubmitAuthorityChangesRequest{ EAuthoritySubmissionErrorCode::Cancelled, *this });
+				AuthorityChangeRequestPromise.EmplaceValue(FSubmitAuthorityChangesRequest{ EAuthoritySubmissionRequestErrorCode::Cancelled });
 			}
 			if (!bAuthorityResponsePromiseWasSet)
 			{
-				AuthorityChangeResponsePromise.EmplaceValue(FSubmitAuthorityChangesResponse{ EAuthoritySubmissionErrorCode::Cancelled, *this });
+				AuthorityChangeResponsePromise.EmplaceValue(FSubmitAuthorityChangesResponse{ EAuthoritySubmissionResponseErrorCode::Cancelled });
 			}
 		}
 
diff --git a/Engine/Plugins/Developer/Concert/ConcertApp/MultiUserClient/Source/MultiUserClient/Private/Replication/Submission/SubmissionWorkflow_LocalClient.cpp b/Engine/Plugins/Developer/Concert/ConcertApp/MultiUserClient/Source/MultiUserClient/Private/Replication/Submission/SubmissionWorkflow_LocalClient.cpp
--- a/Engine/Plugins/Developer/Concert/ConcertApp/MultiUserClient/Source/MultiUserClient/Private/Replication/Submission/SubmissionWorkflow_LocalClient.cpp
+++ b/Engine/Plugins/Developer/Concert/ConcertApp/MultiUserClient/Source/MultiUserClient/Private/Replication/Submission/SubmissionWorkflow_LocalClient.cpp
@@ -52,7 +52,7 @@ namespace UE::MultiUserClient
 		return bHasChanges ? EChangeRevertability::Revertable : EChangeRevertability::NoChanges;
 	}
 
-	ISubmissionOperation* FSubmissionWorkflow_LocalClient::SubmitChanges()
+	TSharedPtr<ISubmissionOperation> FSubmissionWorkflow_LocalClient::SubmitChanges()
 	{
 		using namespace ConcertSyncClient::Replication;
 		IConcertClientReplicationManager* ReplicationManager = Client->GetReplicationManager();
@@ -68,11 +68,13 @@ namespace UE::MultiUserClient
 		const FStreamChangelist& Changelist = StreamChangeTracker.GetCachedDeltaChange();
 		const bool bIsChangelistEmpty = Changelist.ObjectsToPut.IsEmpty() && Changelist.ObjectsToRemove.IsEmpty();
 		const bool bModifyStreams = !bIsChangelistEmpty;
-		InProgressOperation.Emplace(bModifyStreams);
+		
+		const TSharedRef<FSingleClientSubmissionOperation> Operation = MakeShared<FSingleClientSubmissionOperation>(bModifyStreams);
+		InProgressOperation.Emplace(Operation);
 		
 		if (bIsChangelistEmpty)
 		{
-			InProgressOperation->EmplaceStreamPromise(FSubmitStreamChangesResponse{ EStreamSubmissionErrorCode::NoChange, *InProgressOperation });
+			Operation->EmplaceStreamPromise(FSubmitStreamChangesResponse{ EStreamSubmissionErrorCode::NoChange });
 			SendAuthorityChangeRequest(MoveTemp(AuthorityChangeRequest));
 		}
 		else
@@ -84,7 +86,7 @@ namespace UE::MultiUserClient
 			ReplicationManager->ChangeStream(Request)
 				.Next([this, DestructionDetection = LifetimeToken->AsWeak(), Request, AuthorityChangeRequest = MoveTemp(AuthorityChangeRequest)](FChangeStreamResponse&& Response)
 				{
-					const FSubmitStreamChangesResponse SubmissionResult { EStreamSubmissionErrorCode::Success, *InProgressOperation, { FCompletedChangeSubmission{Request, Response } } };
+					const FSubmitStreamChangesResponse SubmissionResult { EStreamSubmissionErrorCode::Success, { FCompletedChangeSubmission{Request, Response } } };
 					// The request might execute after we're destroyed, e.g. by leaving session while request is on the way.
 					// In that case, the Concert session triggers the OnSessionConnectionChanged which destroys us. Only after that, all the requests are timed out.
 					if (DestructionDetection.IsValid())
@@ -95,8 +97,9 @@ namespace UE::MultiUserClient
 					return SubmissionResult;
 				});
 		}
-		
-		return InProgressOperation.GetPtrOrNull();
+
+		// Note that InProgressOperation might already be unset due to ChangeStream failing instantly
+		return Operation;
 	}
 
 	void FSubmissionWorkflow_LocalClient::OnStreamChangeCompleted(
@@ -105,21 +108,25 @@ namespace UE::MultiUserClient
 		ConcertSyncClient::Replication::FAuthorityChangeRequest AuthorityChangeRequest
 		)
 	{
-		InProgressOperation->EmplaceStreamPromise(FSubmitStreamChangesResponse{
-			EStreamSubmissionErrorCode::Success,
-			*InProgressOperation,
+		const TSharedRef<FSingleClientSubmissionOperation> Operation = *InProgressOperation;
+		
+		const EStreamSubmissionErrorCode ErrorCode = ChangeStreamResponse.ErrorCode == EReplicationResponseErrorCode::Handled
+			? EStreamSubmissionErrorCode::Success
+			: EStreamSubmissionErrorCode::Timeout;
+		Operation->EmplaceStreamPromise(FSubmitStreamChangesResponse{
+			ErrorCode,
 			{ FCompletedChangeSubmission{ Request, ChangeStreamResponse } }
 		});
 		
-		if (!ChangeStreamResponse.IsSuccess())
+		if (ChangeStreamResponse.IsSuccess())
 		{
-			InProgressOperation->EmplaceAuthorityRequestPromise(FSubmitAuthorityChangesRequest{ EAuthoritySubmissionErrorCode::CancelledDueToStreamUpdate, *InProgressOperation });
-			InProgressOperation->EmplaceAuthorityResponsePromise(FSubmitAuthorityChangesResponse{ EAuthoritySubmissionErrorCode::CancelledDueToStreamUpdate, *InProgressOperation });
-			InProgressOperation.Reset();
+			SendAuthorityChangeRequest(MoveTemp(AuthorityChangeRequest));
 		}
 		else
 		{
-			SendAuthorityChangeRequest(MoveTemp(AuthorityChangeRequest));
+			Operation->EmplaceAuthorityRequestPromise(FSubmitAuthorityChangesRequest{ EAuthoritySubmissionRequestErrorCode::CancelledDueToStreamUpdate });
+			Operation->EmplaceAuthorityResponsePromise(FSubmitAuthorityChangesResponse{ EAuthoritySubmissionResponseErrorCode::CancelledDueToStreamUpdate });
+			InProgressOperation.Reset();
 		}
 	}
 
@@ -133,22 +140,26 @@ namespace UE::MultiUserClient
 			return;
 		}
 
+		const TSharedRef<FSingleClientSubmissionOperation> Operation = *InProgressOperation;
 		const bool bHasNoChanges = AuthorityChangeRequest.ReleaseAuthority.IsEmpty() && AuthorityChangeRequest.TakeAuthority.IsEmpty();
 		if (bHasNoChanges)
 		{
-			InProgressOperation->EmplaceAuthorityRequestPromise(FSubmitAuthorityChangesRequest{ EAuthoritySubmissionErrorCode::NoChange, *InProgressOperation });
-			InProgressOperation->EmplaceAuthorityResponsePromise(FSubmitAuthorityChangesResponse{ EAuthoritySubmissionErrorCode::NoChange, *InProgressOperation });
+			Operation->EmplaceAuthorityRequestPromise(FSubmitAuthorityChangesRequest{ EAuthoritySubmissionRequestErrorCode::NoChange });
+			Operation->EmplaceAuthorityResponsePromise(FSubmitAuthorityChangesResponse{ EAuthoritySubmissionResponseErrorCode::NoChange });
 			InProgressOperation.Reset();
 			return;
 		}
 		
-		InProgressOperation->EmplaceAuthorityRequestPromise(FSubmitAuthorityChangesRequest{ EAuthoritySubmissionErrorCode::Success, *InProgressOperation, AuthorityChangeRequest });
+		Operation->EmplaceAuthorityRequestPromise(FSubmitAuthorityChangesRequest{ EAuthoritySubmissionRequestErrorCode::Success, AuthorityChangeRequest });
 		ReplicationManager->RequestAuthorityChange(MoveTemp(AuthorityChangeRequest))
 			.Next([this, DestructionDetection = LifetimeToken->AsWeak()](FAuthorityChangeResponse&& Response)
 			{
 				if (DestructionDetection.IsValid())
 				{
-					InProgressOperation->EmplaceAuthorityResponsePromise(FSubmitAuthorityChangesResponse{ EAuthoritySubmissionErrorCode::Success, *InProgressOperation, MoveTemp(Response) });
+					const EAuthoritySubmissionResponseErrorCode ErrorCode = Response.ErrorCode == EReplicationResponseErrorCode::Handled
+						? EAuthoritySubmissionResponseErrorCode::Success
+						: EAuthoritySubmissionResponseErrorCode::Timeout;
+					InProgressOperation->Get().EmplaceAuthorityResponsePromise(FSubmitAuthorityChangesResponse{ ErrorCode, MoveTemp(Response) });
 					InProgressOperation.Reset();
 				}
 			});
diff --git a/Engine/Plugins/Developer/Concert/ConcertApp/MultiUserClient/Source/MultiUserClient/Private/Replication/Submission/SubmissionWorkflow_LocalClient.h b/Engine/Plugins/Developer/Concert/ConcertApp/MultiUserClient/Source/MultiUserClient/Private/Replication/Submission/SubmissionWorkflow_LocalClient.h
--- a/Engine/Plugins/Developer/Concert/ConcertApp/MultiUserClient/Source/MultiUserClient/Private/Replication/Submission/SubmissionWorkflow_LocalClient.h
+++ b/Engine/Plugins/Developer/Concert/ConcertApp/MultiUserClient/Source/MultiUserClient/Private/Replication/Submission/SubmissionWorkflow_LocalClient.h
@@ -29,7 +29,7 @@ namespace UE::MultiUserClient
 			);
 		
 		//~ Begin ISubmissionWorkflow Interface
-		virtual ISubmissionOperation* SubmitChanges() override;
+		virtual TSharedPtr<ISubmissionOperation> SubmitChanges() override;
 		virtual void RevertChanges() override;
 		virtual EChangeUploadability GetUploadability() const override;
 		virtual EChangeRevertability GetRevertability() const override;
@@ -55,7 +55,7 @@ namespace UE::MultiUserClient
 		 * Set for as long as there is a SubmitChanges operation in progress.
 		 * Automatically cancels pending promises when destroyed.
 		 */
-		TOptional<FSingleClientSubmissionOperation> InProgressOperation;
+		TOptional<TSharedRef<FSingleClientSubmissionOperation>> InProgressOperation;
 
 		FGuid GetLocalClientStreamId() const { return StreamSynchronizer.GetStreamId(); }
 
diff --git a/Engine/Plugins/Developer/Concert/ConcertApp/MultiUserClient/Source/MultiUserClient/Private/Replication/Submission/SubmissionWorkflow_RemoteClient.cpp b/Engine/Plugins/Developer/Concert/ConcertApp/MultiUserClient/Source/MultiUserClient/Private/Replication/Submission/SubmissionWorkflow_RemoteClient.cpp
--- a/Engine/Plugins/Developer/Concert/ConcertApp/MultiUserClient/Source/MultiUserClient/Private/Replication/Submission/SubmissionWorkflow_RemoteClient.cpp
+++ b/Engine/Plugins/Developer/Concert/ConcertApp/MultiUserClient/Source/MultiUserClient/Private/Replication/Submission/SubmissionWorkflow_RemoteClient.cpp
@@ -4,7 +4,7 @@
 
 namespace UE::MultiUserClient
 {
-	ISubmissionOperation* FSubmissionWorkflow_RemoteClient::SubmitChanges()
+	TSharedPtr<ISubmissionOperation> FSubmissionWorkflow_RemoteClient::SubmitChanges()
 	{
 		// TODO DP UE-180657: Changing remote client's stream is not implemented for now
 		// TODO DP UE-198088 changing remote client authority
diff --git a/Engine/Plugins/Developer/Concert/ConcertApp/MultiUserClient/Source/MultiUserClient/Private/Replication/Submission/SubmissionWorkflow_RemoteClient.h b/Engine/Plugins/Developer/Concert/ConcertApp/MultiUserClient/Source/MultiUserClient/Private/Replication/Submission/SubmissionWorkflow_RemoteClient.h
--- a/Engine/Plugins/Developer/Concert/ConcertApp/MultiUserClient/Source/MultiUserClient/Private/Replication/Submission/SubmissionWorkflow_RemoteClient.h
+++ b/Engine/Plugins/Developer/Concert/ConcertApp/MultiUserClient/Source/MultiUserClient/Private/Replication/Submission/SubmissionWorkflow_RemoteClient.h
@@ -16,7 +16,7 @@ namespace UE::MultiUserClient
 	public:
 
 		//~ Begin ISubmissionWorkflow Interface
-		virtual ISubmissionOperation* SubmitChanges() override;
+		virtual TSharedPtr<ISubmissionOperation> SubmitChanges() override;
 		virtual void RevertChanges() override;
 		virtual EChangeUploadability GetUploadability() const override;
 		virtual EChangeRevertability GetRevertability() const override;
diff --git a/Engine/Plugins/Developer/Concert/ConcertApp/MultiUserClient/Source/MultiUserClient/Private/Widgets/ActiveSession/Replication/Client/SSubmissionControls.cpp b/Engine/Plugins/Developer/Concert/ConcertApp/MultiUserClient/Source/MultiUserClient/Private/Widgets/ActiveSession/Replication/Client/SSubmissionControls.cpp
--- a/Engine/Plugins/Developer/Concert/ConcertApp/MultiUserClient/Source/MultiUserClient/Private/Widgets/ActiveSession/Replication/Client/SSubmissionControls.cpp
+++ b/Engine/Plugins/Developer/Concert/ConcertApp/MultiUserClient/Source/MultiUserClient/Private/Widgets/ActiveSession/Replication/Client/SSubmissionControls.cpp
@@ -27,6 +27,21 @@
 
 namespace UE::MultiUserClient
 {
+	namespace Private
+	{
+		static void ExecuteOnGameThreadIfSafe(TFunction<void()> Callback)
+		{
+			if (IsInGameThread())
+			{
+				Callback();
+			}
+			else if (!IsEngineExitRequested())
+			{
+				ExecuteOnGameThread(TEXT("SSubmissionControls"), MoveTemp(Callback));
+			}
+		}
+	}
+	
 	void SSubmissionControls::Construct(const FArguments& InArgs)
 	{
 		SubmissionWorkflowAttribute = InArgs._SubmissionWorkflow;
@@ -127,20 +142,21 @@ namespace UE::MultiUserClient
 	FReply SSubmissionControls::OnUploadButtonClicked() const
 	{
 		ISubmissionWorkflow& Workflow = GetSubmissionWorkflow();
-		// SubmitChanges implicitly checks whether Workflow.CanSubmit()
-		if (ISubmissionOperation* SubmissionOperation = Workflow.SubmitChanges())
+		// SubmitChanges implicitly checks whether Workflow.CanSubmit() and calls our lambda if true
+		if (const TSharedPtr<ISubmissionOperation> SubmissionOperation = Workflow.SubmitChanges())
 		{
-			SetupNotificationsFor(*SubmissionOperation);
+			SetupNotificationsFor(SubmissionOperation.ToSharedRef());
 		}
 		return FReply::Handled();
 	}
 
-	void SSubmissionControls::SetupNotificationsFor(ISubmissionOperation& SubmissionOperation) const
+	void SSubmissionControls::SetupNotificationsFor(TSharedRef<ISubmissionOperation> SubmissionOperation) const
 	{
-		HandleStreamUpdatedNotification(SubmissionOperation);
-		HandleAuthorityChangeRequestedNotification(SubmissionOperation);
+		HandleStreamUpdatedNotification(*SubmissionOperation);
+		HandleAuthorityChangeRequestedNotification(MoveTemp(SubmissionOperation));
 	}
 
+
 	void SSubmissionControls::HandleStreamUpdatedNotification(ISubmissionOperation& SubmissionOperation) const
 	{
 		if (!SubmissionOperation.IsModifyingStreams())
@@ -152,83 +168,107 @@ namespace UE::MultiUserClient
 		NotificationInfo.bUseThrobber = true;
 		NotificationInfo.bUseSuccessFailIcons = true;
 		NotificationInfo.bFireAndForget = false;
+		NotificationInfo.ExpireDuration = 4.f;
 		
 		const TSharedPtr<SNotificationItem> Notification = FSlateNotificationManager::Get().AddNotification(NotificationInfo);
 		SubmissionOperation.OnStreamChangesSubmittedFuture()
 			.Next([Notification](FSubmitStreamChangesResponse&& SubmitChangesResult)
 			{
-				if (SubmitChangesResult.ErrorCode == EStreamSubmissionErrorCode::Cancelled)
+				// Not game thread e.g. if the message bus thread destroys the remote endpoint at timeout
+				Private::ExecuteOnGameThreadIfSafe([Notification, SubmitChangesResult]()
 				{
-					Notification->SetCompletionState(SNotificationItem::CS_Fail);
-					Notification->SetText(LOCTEXT("Uploading.Failure.Cancelled", "Updating Stream cancelled."));
-					return;
-				}
+					OnStreamChangesSubmitted(Notification, SubmitChangesResult.ErrorCode, SubmitChangesResult.SubmissionInfo);
+				});
+			});
+	}
 
-				const TOptional<FCompletedChangeSubmission>& SubmissionInfo = SubmitChangesResult.SubmissionInfo;
-				if (SubmissionInfo.IsSet() && SubmissionInfo->Response.IsSuccess())
-				{
-					Notification->SetCompletionState(SNotificationItem::CS_Success);
-					Notification->SetText(LOCTEXT("Uploading.Success", "Changes accepted."));
+	void SSubmissionControls::OnStreamChangesSubmitted(const TSharedPtr<SNotificationItem>& Notification, EStreamSubmissionErrorCode ErrorCode, TOptional<FCompletedChangeSubmission> SubmissionInfo)
+	{
+		ON_SCOPE_EXIT{ Notification->ExpireAndFadeout(); };
+				
+		if (ErrorCode == EStreamSubmissionErrorCode::Cancelled)
+		{
+			Notification->SetCompletionState(SNotificationItem::CS_Fail);
+			Notification->SetText(LOCTEXT("Uploading.Failure.UpdateCancelled", "Stream update cancelled."));
+			return;
+		}
+		if (ErrorCode == EStreamSubmissionErrorCode::Timeout)
+		{
+			Notification->SetCompletionState(SNotificationItem::CS_Fail);
+			Notification->SetText(LOCTEXT("Uploading.Failure.Timedout", "Stream update timed out."));
+			return;
+		}
 
-					const ConcertSyncClient::Replication::FChangeStreamRequest& Request = SubmissionInfo->Request;
-					// Multi-User adds only one stream so we only need to look at index 0
-					const int32 NumAdditions = Request.StreamsToAdd.IsEmpty() ? 0 : Request.StreamsToAdd[0].BaseDescription.ReplicationMap.ReplicatedObjects.Num();
-					const int32 NumPuts = Request.ObjectsToPut.Num() + NumAdditions;
-					Notification->SetSubText(FText::Format(LOCTEXT("Uploading.Success.SubText", "{0} puts, {1} removals"),
-						NumPuts,
-						Request.ObjectsToRemove.Num()
-						));
-				}
-				else
-				{
-					Notification->SetCompletionState(SNotificationItem::CS_Fail);
-					Notification->SetText(LOCTEXT("Uploading.Failure.Rejected", "Changes rejected."));
-					Notification->SetSubText(LOCTEXT("Uploading.Failure.Rejected.SubText", "See log for details."));
+		if (SubmissionInfo.IsSet() && SubmissionInfo->Response.IsSuccess())
+		{
+			Notification->SetCompletionState(SNotificationItem::CS_Success);
+			Notification->SetText(LOCTEXT("Uploading.Success", "Changes accepted."));
 
-					if (SubmissionInfo.IsSet())
-					{
-						FStringOutputDevice Errors;
-						SubmissionInfo->Response.LogErrors(Errors);
-						UE_LOG(LogConcert, Error, TEXT("Errors uploading stream changes:\n%s"), *Errors);
-					}
-					else
-					{
-						UE_LOG(LogConcert, Error, TEXT("Failed to submit changes. ESubmitChangesErrorCode: %d"), SubmitChangesResult.ErrorCode);
-					}
-				}
+			const ConcertSyncClient::Replication::FChangeStreamRequest& Request = SubmissionInfo->Request;
+			// Multi-User adds only one stream so we only need to look at index 0
+			const int32 NumAdditions = Request.StreamsToAdd.IsEmpty() ? 0 : Request.StreamsToAdd[0].BaseDescription.ReplicationMap.ReplicatedObjects.Num();
+			const int32 NumPuts = Request.ObjectsToPut.Num() + NumAdditions;
+			Notification->SetSubText(FText::Format(LOCTEXT("Uploading.Success.SubText", "{0} puts, {1} removals"),
+				NumPuts,
+				Request.ObjectsToRemove.Num()
+				));
+		}
+		else
+		{
+			Notification->SetCompletionState(SNotificationItem::CS_Fail);
+			Notification->SetText(LOCTEXT("Uploading.Failure.Rejected", "Changes rejected."));
+			Notification->SetSubText(LOCTEXT("Uploading.Failure.Rejected.SubText", "See log for details."));
 
-				Notification->ExpireAndFadeout();
-           });
+			if (SubmissionInfo.IsSet())
+			{
+				FStringOutputDevice Errors;
+				SubmissionInfo->Response.LogErrors(Errors);
+				UE_LOG(LogConcert, Error, TEXT("Errors uploading stream changes:\n%s"), *Errors);
+			}
+			else
+			{
+				UE_LOG(LogConcert, Error, TEXT("Failed to submit changes. ESubmitChangesErrorCode: %d"), ErrorCode);
+			}
+		}
 	}
 
-	void SSubmissionControls::HandleAuthorityChangeRequestedNotification(ISubmissionOperation& SubmissionOperation) const
+	void SSubmissionControls::HandleAuthorityChangeRequestedNotification(TSharedRef<ISubmissionOperation> SubmissionOperation) const
 	{
-		SubmissionOperation.OnAuthorityChangeRequestedFuture()
-			.Next([WeakThis = TWeakPtr<const SSubmissionControls>(SharedThis(this))](FSubmitAuthorityChangesRequest&& RequestMetaData)
+		SubmissionOperation->OnAuthorityChangeRequestedFuture()
+			.Next([SubmissionOperation = MoveTemp(SubmissionOperation)](FSubmitAuthorityChangesRequest&& RequestMetaData) mutable
 			{
-				if (RequestMetaData.ErrorCode != EAuthoritySubmissionErrorCode::Success)
+				if (RequestMetaData.ErrorCode != EAuthoritySubmissionRequestErrorCode::Success)
 				{
 					return;
 				}
 
-				FNotificationInfo NotificationInfo(LOCTEXT("Authority.InProgress.Text", "Requesting authority."));
-				NotificationInfo.bUseThrobber = true;
-				NotificationInfo.bUseSuccessFailIcons = true;
-				NotificationInfo.bFireAndForget = false;
-				const TSharedPtr<SNotificationItem> AuthorityChangeNotification = FSlateNotificationManager::Get().AddNotification(NotificationInfo);
+				// Not game thread e.g. if the message bus thread destroys the remote endpoint at timeout
+				Private::ExecuteOnGameThreadIfSafe([RequestMetaData = MoveTemp(RequestMetaData), SubmissionOperation = MoveTemp(SubmissionOperation)]() mutable
+				{
+					FNotificationInfo NotificationInfo(LOCTEXT("Authority.InProgress.Text", "Requesting authority."));
+					NotificationInfo.bUseThrobber = true;
+					NotificationInfo.bUseSuccessFailIcons = true;
+					NotificationInfo.bFireAndForget = false;
+					NotificationInfo.ExpireDuration = 4.f;
+					TSharedPtr<SNotificationItem> AuthorityChangeNotification = FSlateNotificationManager::Get().AddNotification(NotificationInfo);
 
-				const ConcertSyncClient::Replication::FAuthorityChangeRequest& Request = *RequestMetaData.Request ;
-				AuthorityChangeNotification->SetSubText(
-					FText::Format(LOCTEXT("Authority.InProgress.SubTextFmt", "Taking {0}|plural(one=object,other\nReleasing {1}|plural(one=object,other=objects)"),
-						Request.TakeAuthority.Num(),
-						Request.ReleaseAuthority.Num()
-					));
-				
-				RequestMetaData.OperationContext.OnAuthorityChangeResponseReceivedFuture()
-					.Next([Request, AuthorityChangeNotification](FSubmitAuthorityChangesResponse&& Response)
-					{
-						HandleAuthorityChangedResponseNotification(Request, Response, AuthorityChangeNotification.ToSharedRef());
-					});
+					ConcertSyncClient::Replication::FAuthorityChangeRequest& Request = *RequestMetaData.Request;
+					AuthorityChangeNotification->SetSubText(
+						FText::Format(LOCTEXT("Authority.InProgress.SubTextFmt", "Taking {0}|plural(one=object,other\nReleasing {1}|plural(one=object,other=objects)"),
+							Request.TakeAuthority.Num(),
+							Request.ReleaseAuthority.Num()
+						));
+					
+					SubmissionOperation->OnAuthorityChangeResponseReceivedFuture()
+						.Next([Request = MoveTemp(Request), AuthorityChangeNotification = MoveTemp(AuthorityChangeNotification)](FSubmitAuthorityChangesResponse&& Response) mutable
+						{
+							// Not game thread e.g. if the message bus thread destroys the remote endpoint at timeout
+							Private::ExecuteOnGameThreadIfSafe([Request = MoveTemp(Request), Response = MoveTemp(Response), AuthorityChangeNotification = MoveTemp(AuthorityChangeNotification)]()
+							{
+								HandleAuthorityChangedResponseNotification(Request, Response, AuthorityChangeNotification.ToSharedRef());
+							});
+						});
+				});
 			});
 	}
 
@@ -238,13 +278,13 @@ namespace UE::MultiUserClient
 		const TSharedRef<SNotificationItem>& NotificationItem
 		)
 	{
-		if (ResponseMetaData.ErrorCode == EAuthoritySubmissionErrorCode::NoChange)
+		if (ResponseMetaData.ErrorCode == EAuthoritySubmissionResponseErrorCode::NoChange)
 		{
 			return;
 		}
 		
 		ON_SCOPE_EXIT{ NotificationItem->ExpireAndFadeout(); };
-		if (ResponseMetaData.ErrorCode != EAuthoritySubmissionErrorCode::Success)
+		if (ResponseMetaData.ErrorCode != EAuthoritySubmissionResponseErrorCode::Success)
 		{
 			NotificationItem->SetCompletionState(SNotificationItem::CS_Fail);
 			NotificationItem->SetText(LOCTEXT("Authority.Failed.Text", "Authority request failed."));
diff --git a/Engine/Plugins/Developer/Concert/ConcertApp/MultiUserClient/Source/MultiUserClient/Private/Widgets/ActiveSession/Replication/Client/SSubmissionControls.h b/Engine/Plugins/Developer/Concert/ConcertApp/MultiUserClient/Source/MultiUserClient/Private/Widgets/ActiveSession/Replication/Client/SSubmissionControls.h
--- a/Engine/Plugins/Developer/Concert/ConcertApp/MultiUserClient/Source/MultiUserClient/Private/Widgets/ActiveSession/Replication/Client/SSubmissionControls.h
+++ b/Engine/Plugins/Developer/Concert/ConcertApp/MultiUserClient/Source/MultiUserClient/Private/Widgets/ActiveSession/Replication/Client/SSubmissionControls.h
@@ -46,9 +46,16 @@ namespace UE::MultiUserClient
 		FReply OnUploadButtonClicked() const;
 
 		// Handle submission events
-		void SetupNotificationsFor(ISubmissionOperation& SubmissionOperation) const;
+		void SetupNotificationsFor(TSharedRef<ISubmissionOperation> SubmissionOperation) const;
+
+		/** Shows a notification that the stream is being updated */
 		void HandleStreamUpdatedNotification(ISubmissionOperation& SubmissionOperation) const;
-		void HandleAuthorityChangeRequestedNotification(ISubmissionOperation& SubmissionOperation) const;
+		/** Closes the stream update notification */
+		static void OnStreamChangesSubmitted(const TSharedPtr<SNotificationItem>& Notification, EStreamSubmissionErrorCode ErrorCode, TOptional<FCompletedChangeSubmission> SubmissionInfo);
+
+		/** Shows a notification that the authority is being updated. */
+		void HandleAuthorityChangeRequestedNotification(TSharedRef<ISubmissionOperation> SubmissionOperation) const;
+		/** Hides the authority hiding notification. */
 		static void HandleAuthorityChangedResponseNotification(
 			const ConcertSyncClient::Replication::FAuthorityChangeRequest& Request,
 			const FSubmitAuthorityChangesResponse& ResponseMetaData,
diff --git a/Engine/Plugins/Developer/Concert/ConcertSync/ConcertSyncClient/Source/ConcertSyncClient/Private/Replication/IConcertClientReplicationManager.cpp b/Engine/Plugins/Developer/Concert/ConcertSync/ConcertSyncClient/Source/ConcertSyncClient/Private/Replication/IConcertClientReplicationManager.cpp
--- a/Engine/Plugins/Developer/Concert/ConcertSync/ConcertSyncClient/Source/ConcertSyncClient/Private/Replication/IConcertClientReplicationManager.cpp
+++ b/Engine/Plugins/Developer/Concert/ConcertSync/ConcertSyncClient/Source/ConcertSyncClient/Private/Replication/IConcertClientReplicationManager.cpp
@@ -51,7 +51,7 @@ TFuture<UE::ConcertSyncClient::Replication::FAuthorityChangeResponse> IConcertCl
 		UE_LOG(LogConcert, Error, TEXT("Attempted to take authority while not connected!"));
 		TMap<FSoftObjectPath, FConcertStreamArray> Result;
 		Algo::Transform(Objects, Result, [](const FSoftObjectPath& Path){ return Path; });
-		return MakeFulfilledPromise<FAuthorityChangeResponse>(FAuthorityChangeResponse{{ MoveTemp(Result) }}).GetFuture();
+		return MakeFulfilledPromise<FAuthorityChangeResponse>(FAuthorityChangeResponse{{ EReplicationResponseErrorCode::Handled ,MoveTemp(Result) }}).GetFuture();
 	}
 
 	FAuthorityChangeRequest Request;
@@ -82,7 +82,7 @@ TFuture<UE::ConcertSyncClient::Replication::FAuthorityChangeResponse> IConcertCl
 		UE_LOG(LogConcert, Error, TEXT("Attempted to take authority while not connected!"));
 		TMap<FSoftObjectPath, FConcertStreamArray> Result;
 		Algo::Transform(Objects, Result, [](const FSoftObjectPath& Path){ return Path; });
-		return MakeFulfilledPromise<FAuthorityChangeResponse>(FAuthorityChangeResponse{{ MoveTemp(Result) }}).GetFuture();
+		return MakeFulfilledPromise<FAuthorityChangeResponse>(FAuthorityChangeResponse{{ EReplicationResponseErrorCode::Handled, MoveTemp(Result) }}).GetFuture();
 	}
 	
 	FAuthorityChangeRequest Request;
diff --git a/Engine/Plugins/Developer/Concert/ConcertSync/ConcertSyncClient/Source/ConcertSyncClient/Private/Replication/Manager/ReplicationManagerState_Connected.cpp b/Engine/Plugins/Developer/Concert/ConcertSync/ConcertSyncClient/Source/ConcertSyncClient/Private/Replication/Manager/ReplicationManagerState_Connected.cpp
--- a/Engine/Plugins/Developer/Concert/ConcertSync/ConcertSyncClient/Source/ConcertSyncClient/Private/Replication/Manager/ReplicationManagerState_Connected.cpp
+++ b/Engine/Plugins/Developer/Concert/ConcertSync/ConcertSyncClient/Source/ConcertSyncClient/Private/Replication/Manager/ReplicationManagerState_Connected.cpp
@@ -21,6 +21,22 @@
 
 namespace UE::ConcertSyncClient::Replication
 {
+	TAutoConsoleVariable<bool> CVarSimulateAuthorityTimeouts(
+		TEXT("Concert.Replication.SimulateAuthorityTimeouts"),
+		false,
+		TEXT("Whether the client should pretend that authority requests timed out instead of sending to the server.")
+		);
+	TAutoConsoleVariable<bool> CVarSimulateQueryTimeouts(
+		TEXT("Concert.Replication.SimulateQueryTimeouts"),
+		false,
+		TEXT("Whether the client should pretend that query requests timed out instead of sending to the server.")
+		);
+	TAutoConsoleVariable<bool> CVarSimulateStreamChangeTimeouts(
+		TEXT("Concert.Replication.SimulateStreamChangeTimeouts"),
+		false,
+		TEXT("Whether the client should pretend that stream change requests timed out instead of sending to the server.")
+		);
+	
 	FReplicationManagerState_Connected::FReplicationManagerState_Connected(
 		TSharedRef<IConcertClientSession> LiveSession,
 		IConcertClientReplicationBridge* ReplicationBridge,
@@ -82,6 +98,11 @@ namespace UE::ConcertSyncClient::Replication
 
 	TFuture<FAuthorityChangeResponse> FReplicationManagerState_Connected::RequestAuthorityChange(FAuthorityChangeRequest Args)
 	{
+		if (CVarSimulateAuthorityTimeouts.GetValueOnGameThread())
+		{
+			return MakeFulfilledPromise<FAuthorityChangeResponse>(FAuthorityChangeResponse{ EReplicationResponseErrorCode::Timeout }).GetFuture();
+		}
+		
 		// Stop replicating removed objects right now: the server will remove authority after processing this request.
 		// At that point, it will log errors for receiving replication data from a client without authority.
 		HandleReleasingReplicatedObjects(Args);
@@ -89,10 +110,16 @@ namespace UE::ConcertSyncClient::Replication
 		return LiveSession->SendCustomRequest<FConcertReplication_ChangeAuthority_Request, FConcertReplication_ChangeAuthority_Response>(Args, LiveSession->GetSessionServerEndpointId())
 			.Next([WeakThis = TWeakPtr<FReplicationManagerState_Connected>(SharedThis(this)), Args](FConcertReplication_ChangeAuthority_Response&& Response) mutable
 			{
-				if (const TSharedPtr<FReplicationManagerState_Connected> ThisPin = WeakThis.Pin())
+				if (const TSharedPtr<FReplicationManagerState_Connected> ThisPin = WeakThis.Pin()
+					; ThisPin && Response.ErrorCode == EReplicationResponseErrorCode::Handled)
 				{
 					ThisPin->UpdateReplicatedObjectsAfterAuthorityChange(MoveTemp(Args), Response);
 				}
+				else if (ThisPin && Response.ErrorCode == EReplicationResponseErrorCode::Timeout)
+				{
+					// HandleReleasingReplicatedObjects caused Args.ReleaseAuthority to stop being replicated. Revert.
+					ThisPin->RevertReleasingReplicatedObjects(Args);
+				}
 
 				return FAuthorityChangeResponse { MoveTemp(Response) };
 			});
@@ -100,6 +127,11 @@ namespace UE::ConcertSyncClient::Replication
 
 	TFuture<FClientQueryResponse> FReplicationManagerState_Connected::QueryClientInfo(FClientQueryRequest Args)
 	{
+		if (CVarSimulateQueryTimeouts.GetValueOnGameThread())
+		{
+			return MakeFulfilledPromise<FClientQueryResponse>(FClientQueryResponse{ EReplicationResponseErrorCode::Timeout }).GetFuture();
+		}
+		
 		if (EnumHasAllFlags(Args.QueryFlags, EConcertQueryClientStreamFlags::SkipAuthority | EConcertQueryClientStreamFlags::SkipStreamInfo))
 		{
 			UE_LOG(LogConcert, Warning, TEXT("Request QueryClientInfo is pointless because SkipAuthority and SkipStreamInfo are both set. Returning immediately..."));
@@ -115,6 +147,11 @@ namespace UE::ConcertSyncClient::Replication
 
 	TFuture<FChangeStreamResponse> FReplicationManagerState_Connected::ChangeStream(FChangeStreamRequest Args)
 	{
+		if (CVarSimulateStreamChangeTimeouts.GetValueOnGameThread())
+		{
+			return MakeFulfilledPromise<FChangeStreamResponse>(FChangeStreamResponse{ EReplicationResponseErrorCode::Timeout }).GetFuture();
+		}
+		
 		// Stop replicating removed objects right now: the server will remove authority after processing this request.
 		// At that point, it will log errors for receiving replication data from a client without authority.
 		HandleRemovingReplicatedObjects(Args);
@@ -127,6 +164,11 @@ namespace UE::ConcertSyncClient::Replication
 				{
 					ThisPin->UpdateReplicatedObjectsAfterStreamChange(Args, Response);
 				}
+				else if (ThisPin && Response.ErrorCode == EReplicationResponseErrorCode::Timeout)
+				{
+					// HandleRemovingReplicatedObjects caused Request.ObjectsToRemove to stop being replicated. Revert.
+					ThisPin->RevertRemovingReplicatedObjects(Args);
+				}
 				
 				return FChangeStreamResponse { MoveTemp(Response) };
 			});
@@ -231,18 +273,36 @@ namespace UE::ConcertSyncClient::Replication
 		}
 	}
 
-	void FReplicationManagerState_Connected::HandleRemovingReplicatedObjects(const FChangeStreamRequest& Request) const
+	namespace Private
 	{
-		TMap<FSoftObjectPath, TArray<FGuid>> BundledRemovedObjects;
-		for (const FObjectInStreamID& RemovedObject : Request.ObjectsToRemove)
+		static void ForEachObjectRemovedFromStreams(const FChangeStreamRequest& Request, TFunctionRef<void(const FSoftObjectPath& ObjectPath, const TArray<FGuid>& Streams)> Callback)
 		{
-			BundledRemovedObjects.FindOrAdd(RemovedObject.Object).Add(RemovedObject.StreamId);
+			TMap<FSoftObjectPath, TArray<FGuid>> BundledRemovedObjects;
+			for (const FObjectInStreamID& RemovedObject : Request.ObjectsToRemove)
+			{
+				BundledRemovedObjects.FindOrAdd(RemovedObject.Object).Add(RemovedObject.StreamId);
+			}
+			for (const TPair<FSoftObjectPath, TArray<FGuid>>& RemovedObjectInfo : BundledRemovedObjects)
+			{
+				Callback(RemovedObjectInfo.Key, RemovedObjectInfo.Value);
+			}
 		}
-		for (const TPair<FSoftObjectPath, TArray<FGuid>>& RemovedObjectInfo : BundledRemovedObjects)
+	}
+
+	void FReplicationManagerState_Connected::HandleRemovingReplicatedObjects(const FChangeStreamRequest& Request) const
+	{
+		Private::ForEachObjectRemovedFromStreams(Request, [this](const FSoftObjectPath& Object, const TArray<FGuid>& RemovedStreams)
 		{
-			// Removing an object from a stream implies removing its authority so stop replicating it
-			ReplicationDataSource->RemoveReplicatedObjectStreams(RemovedObjectInfo.Key, RemovedObjectInfo.Value);
-		}
+			ReplicationDataSource->RemoveReplicatedObjectStreams(Object, RemovedStreams);
+		});
+	}
+
+	void FReplicationManagerState_Connected::RevertRemovingReplicatedObjects(const FChangeStreamRequest& Request) const
+	{
+		Private::ForEachObjectRemovedFromStreams(Request, [this](const FSoftObjectPath& Object, const TArray<FGuid>& RemovedStreams)
+		{
+			ReplicationDataSource->AddReplicatedObjectStreams(Object, RemovedStreams);
+		});
 	}
 
 	void FReplicationManagerState_Connected::UpdateReplicatedObjectsAfterAuthorityChange(FAuthorityChangeRequest&& Request, const FConcertReplication_ChangeAuthority_Response& Response) const
@@ -282,5 +342,13 @@ namespace UE::ConcertSyncClient::Replication
 			ReplicationDataSource->RemoveReplicatedObjectStreams(ReleaseAuthority.Key, ReleaseAuthority.Value.StreamIds);
 		}
 	}
+
+	void FReplicationManagerState_Connected::RevertReleasingReplicatedObjects(const FAuthorityChangeRequest& Request) const
+	{
+		for (const TPair<FSoftObjectPath, FConcertStreamArray>& ReleaseAuthority : Request.ReleaseAuthority)
+		{
+			ReplicationDataSource->AddReplicatedObjectStreams(ReleaseAuthority.Key, ReleaseAuthority.Value.StreamIds);
+		}
+	}
 }
 
diff --git a/Engine/Plugins/Developer/Concert/ConcertSync/ConcertSyncClient/Source/ConcertSyncClient/Private/Replication/Manager/ReplicationManagerState_Connected.h b/Engine/Plugins/Developer/Concert/ConcertSync/ConcertSyncClient/Source/ConcertSyncClient/Private/Replication/Manager/ReplicationManagerState_Connected.h
--- a/Engine/Plugins/Developer/Concert/ConcertSync/ConcertSyncClient/Source/ConcertSyncClient/Private/Replication/Manager/ReplicationManagerState_Connected.h
+++ b/Engine/Plugins/Developer/Concert/ConcertSync/ConcertSyncClient/Source/ConcertSyncClient/Private/Replication/Manager/ReplicationManagerState_Connected.h
@@ -121,6 +121,7 @@ namespace UE::ConcertSyncClient::Replication
 		/** Updates replicated objects affected by the change request. */
 		void UpdateReplicatedObjectsAfterStreamChange(const FChangeStreamRequest& Request, const FConcertReplication_ChangeStream_Response& Response);
 		void HandleRemovingReplicatedObjects(const FChangeStreamRequest& Request) const;
+		void RevertRemovingReplicatedObjects(const FChangeStreamRequest& Request) const;
 
 		/**
 		 * Updates the objects which should be replicated after changing authority.
@@ -130,5 +131,6 @@ namespace UE::ConcertSyncClient::Replication
 		 */
 		void UpdateReplicatedObjectsAfterAuthorityChange(FAuthorityChangeRequest&& Request, const FConcertReplication_ChangeAuthority_Response& Response) const;
 		void HandleReleasingReplicatedObjects(const FAuthorityChangeRequest& Request) const;
+		void RevertReleasingReplicatedObjects(const FAuthorityChangeRequest& Request) const;
 	};
 }
diff --git a/Engine/Plugins/Developer/Concert/ConcertSync/ConcertSyncClient/Source/ConcertSyncClient/Private/Replication/Manager/ReplicationManagerState_Disconnected.cpp b/Engine/Plugins/Developer/Concert/ConcertSync/ConcertSyncClient/Source/ConcertSyncClient/Private/Replication/Manager/ReplicationManagerState_Disconnected.cpp
--- a/Engine/Plugins/Developer/Concert/ConcertSync/ConcertSyncClient/Source/ConcertSyncClient/Private/Replication/Manager/ReplicationManagerState_Disconnected.cpp
+++ b/Engine/Plugins/Developer/Concert/ConcertSync/ConcertSyncClient/Source/ConcertSyncClient/Private/Replication/Manager/ReplicationManagerState_Disconnected.cpp
@@ -5,8 +5,16 @@
 #include "ConcertLogGlobal.h"
 #include "ReplicationManagerState_Handshaking.h"
 
+#include "HAL/IConsoleManager.h"
+
 namespace UE::ConcertSyncClient::Replication
 {
+	TAutoConsoleVariable<bool> CVarSimulateJoinTimeouts(
+		TEXT("Concert.Replication.SimulateJoinTimeouts"),
+		false,
+		TEXT("Whether the client should pretend that join requests timed out instead of sending to the server.")
+		);
+	
 	FReplicationManagerState_Disconnected::FReplicationManagerState_Disconnected(
 		TSharedRef<IConcertClientSession> LiveSession,
 		IConcertClientReplicationBridge* ReplicationBridge,
@@ -19,6 +27,11 @@ namespace UE::ConcertSyncClient::Replication
 
 	TFuture<FJoinReplicatedSessionResult> FReplicationManagerState_Disconnected::JoinReplicationSession(FJoinReplicatedSessionArgs Args)
 	{
+		if (CVarSimulateJoinTimeouts.GetValueOnGameThread())
+		{
+			return MakeFulfilledPromise<FJoinReplicatedSessionResult>(FJoinReplicatedSessionResult{ EJoinReplicationErrorCode::NetworkError }).GetFuture();
+		}
+		
 		TPromise<FJoinReplicatedSessionResult> JoinPromise;
 		TFuture<FJoinReplicatedSessionResult> JoinFuture = JoinPromise.GetFuture();
 		ChangeState(MakeShared<FReplicationManagerState_Handshaking>(MoveTemp(Args), MoveTemp(JoinPromise), LiveSession, ReplicationBridge, GetOwner()));
diff --git a/Engine/Plugins/Developer/Concert/ConcertSync/ConcertSyncClient/Source/ConcertSyncClient/Private/Replication/Manager/ReplicationManagerState_Handshaking.cpp b/Engine/Plugins/Developer/Concert/ConcertSync/ConcertSyncClient/Source/ConcertSyncClient/Private/Replication/Manager/ReplicationManagerState_Handshaking.cpp
--- a/Engine/Plugins/Developer/Concert/ConcertSync/ConcertSyncClient/Source/ConcertSyncClient/Private/Replication/Manager/ReplicationManagerState_Handshaking.cpp
+++ b/Engine/Plugins/Developer/Concert/ConcertSync/ConcertSyncClient/Source/ConcertSyncClient/Private/Replication/Manager/ReplicationManagerState_Handshaking.cpp
@@ -61,7 +61,7 @@ namespace UE::ConcertSyncClient::Replication
 					return;
 				}
 
-				if (Response.ErrorCode == EJoinReplicationErrorCode::Success)
+				if (Response.JoinErrorCode == EJoinReplicationErrorCode::Success)
 				{
 					This->ChangeState(
 						MakeShared<FReplicationManagerState_Connected>(This->LiveSession, This->ReplicationBridge, MoveTemp(This->RequestArgs.Streams), This->GetOwner())
@@ -73,7 +73,7 @@ namespace UE::ConcertSyncClient::Replication
 				}
 
 				This->bFulfilledPromise = true;
-				This->JoinSessionPromise.EmplaceValue(Response.ErrorCode, Response.DetailedErrorMessage);
+				This->JoinSessionPromise.EmplaceValue(Response.JoinErrorCode, Response.DetailedErrorMessage);
 			});
 	}
 
diff --git a/Engine/Plugins/Developer/Concert/ConcertSync/ConcertSyncClient/Source/ConcertSyncClient/Public/Replication/IConcertClientReplicationManager.h b/Engine/Plugins/Developer/Concert/ConcertSync/ConcertSyncClient/Source/ConcertSyncClient/Public/Replication/IConcertClientReplicationManager.h
--- a/Engine/Plugins/Developer/Concert/ConcertSync/ConcertSyncClient/Source/ConcertSyncClient/Public/Replication/IConcertClientReplicationManager.h
+++ b/Engine/Plugins/Developer/Concert/ConcertSync/ConcertSyncClient/Source/ConcertSyncClient/Public/Replication/IConcertClientReplicationManager.h
@@ -60,6 +60,7 @@ public:
 	/**
 	 * Joins a replication session.
 	 * Subsequent calls to JoinReplicationSession will fail until either the resulting TFuture returns or LeaveReplicationSession is called.
+	 * @note This future can finish on any thread (e.g. when message endpoint times out); usually it finishes game thread.  
 	 */
 	virtual TFuture<UE::ConcertSyncClient::Replication::FJoinReplicatedSessionResult> JoinReplicationSession(UE::ConcertSyncClient::Replication::FJoinReplicatedSessionArgs Args) = 0;
 	/** Leaves the current replication session. */
@@ -78,6 +79,7 @@ public:
 	 * Iterates the streams the client has registered with the server.
 	 * It only makes sense to call this function the manager has joined a replication session.
 	 * @return Whether this manager is connected to a session (Iterated) or not (NoRegisteredStreams).
+	 * @note This future can finish on any thread (e.g. when message endpoint times out); usually it finishes game thread.  
 	 */
 	virtual EStreamEnumerationResult ForEachRegisteredStream(TFunctionRef<EBreakBehavior(const FReplicationStreamDescription& Stream)> Callback) const = 0;
 	/** @return Whether this manager is has any registered streams (basically whether ForEachRegisteredStream returns EStreamEnumerationResult::Iterated). */
@@ -87,6 +89,7 @@ public:
 	
 	/**
 	 * Requests from the server to change the authority over some objects.
+	 * @note This future can finish on any thread (e.g. when message endpoint times out); usually it finishes game thread.  
 	 */
 	virtual TFuture<UE::ConcertSyncClient::Replication::FAuthorityChangeResponse> RequestAuthorityChange(UE::ConcertSyncClient::Replication::FAuthorityChangeRequest Args) = 0;
 	/** Util function that will request authority for all streams for the given objects. */
@@ -105,10 +108,16 @@ public:
 	/** @return All owned objects and associated owning streams. */
 	TMap<FSoftObjectPath, TSet<FGuid>> GetClientOwnedObjects() const;
 
-	/** Requests replication info about other clients, including the streams registered and which objects they have authority over (i.e. are sending). */
+	/**
+	 * Requests replication info about other clients, including the streams registered and which objects they have authority over (i.e. are sending).
+	 * @note This future can finish on any thread (e.g. when message endpoint times out); usually it finishes game thread.  
+	 */
 	virtual TFuture<UE::ConcertSyncClient::Replication::FClientQueryResponse> QueryClientInfo(UE::ConcertSyncClient::Replication::FClientQueryRequest Args) = 0;
 
-	/** Requests to change the client's registered streams */
+	/**
+	 * Requests to change the client's registered stream.
+	 * @note This future can finish on any thread (e.g. when message endpoint times out); usually it finishes game thread.  
+	 */
 	virtual TFuture<UE::ConcertSyncClient::Replication::FChangeStreamResponse> ChangeStream(UE::ConcertSyncClient::Replication::FChangeStreamRequest Args) = 0;
 
 	DECLARE_MULTICAST_DELEGATE_TwoParams(FOnPreStreamsChanged,
diff --git a/Engine/Plugins/Developer/Concert/ConcertSync/ConcertSyncCore/Source/ConcertSyncCore/Public/Replication/Messages/ChangeAuthority.h b/Engine/Plugins/Developer/Concert/ConcertSync/ConcertSyncCore/Source/ConcertSyncCore/Public/Replication/Messages/ChangeAuthority.h
--- a/Engine/Plugins/Developer/Concert/ConcertSync/ConcertSyncCore/Source/ConcertSyncCore/Public/Replication/Messages/ChangeAuthority.h
+++ b/Engine/Plugins/Developer/Concert/ConcertSync/ConcertSyncCore/Source/ConcertSyncCore/Public/Replication/Messages/ChangeAuthority.h
@@ -2,7 +2,7 @@
 
 #pragma once
 
-#include "Replication/Data/ObjectIds.h"
+#include "EReplicationResponseErrorCode.h"
 #include "ChangeAuthority.generated.h"
 
 USTRUCT()
@@ -51,6 +51,10 @@ struct FConcertReplication_ChangeAuthority_Response
 {
 	GENERATED_BODY()
 
+	/** Concert's custom requests are default constructed when they timeout. Server always sets this to Handled when processed. */
+	UPROPERTY()
+	EReplicationResponseErrorCode ErrorCode = EReplicationResponseErrorCode::Timeout;
+	
 	/**
 	 * Objects streams the client did not receive authority over.
 	 * The client is implied to have authority over all other request objects.
diff --git a/Engine/Plugins/Developer/Concert/ConcertSync/ConcertSyncCore/Source/ConcertSyncCore/Public/Replication/Messages/ChangeStream.h b/Engine/Plugins/Developer/Concert/ConcertSync/ConcertSyncCore/Source/ConcertSyncCore/Public/Replication/Messages/ChangeStream.h
--- a/Engine/Plugins/Developer/Concert/ConcertSync/ConcertSyncCore/Source/ConcertSyncCore/Public/Replication/Messages/ChangeStream.h
+++ b/Engine/Plugins/Developer/Concert/ConcertSync/ConcertSyncCore/Source/ConcertSyncCore/Public/Replication/Messages/ChangeStream.h
@@ -2,6 +2,7 @@
 
 #pragma once
 
+#include "EReplicationResponseErrorCode.h"
 #include "Replication/Data/ObjectIds.h"
 #include "Replication/Data/ConcertPropertySelection.h"
 #include "Replication/Data/ReplicationStreamDescription.h"
@@ -152,6 +153,10 @@ struct FConcertReplication_ChangeStream_Response
 {
 	GENERATED_BODY()
 
+	/** Concert's custom requests are default constructed when they timeout. Server always sets this to Handled when processed. */
+	UPROPERTY()
+	EReplicationResponseErrorCode ErrorCode = EReplicationResponseErrorCode::Timeout;
+	
 	/**
 	 * Reports dynamic authority errors with ObjectsToPut:
 	 * Changing an object over which a client already has authority can yield unresolvable conflicts for which the entire FConcertChangeStream_Request is rejected.
@@ -178,7 +183,7 @@ struct FConcertReplication_ChangeStream_Response
 	UPROPERTY()
 	TSet<FGuid> FailedStreamCreation;
 
-	bool IsSuccess() const { return AuthorityConflicts.IsEmpty() && ObjectsToPutSemanticErrors.IsEmpty() && FailedStreamCreation.IsEmpty(); }
+	bool IsSuccess() const { return ErrorCode == EReplicationResponseErrorCode::Handled && AuthorityConflicts.IsEmpty() && ObjectsToPutSemanticErrors.IsEmpty() && FailedStreamCreation.IsEmpty(); }
 	bool IsFailure() const { return !IsSuccess(); }
 
 	bool WasObjectPutSuccessful(const FObjectInStreamID& Object) const { return !AuthorityConflicts.Contains(Object) && !ObjectsToPutSemanticErrors.Contains(Object); }
diff --git a/Engine/Plugins/Developer/Concert/ConcertSync/ConcertSyncCore/Source/ConcertSyncCore/Public/Replication/Messages/ClientQuery.h b/Engine/Plugins/Developer/Concert/ConcertSync/ConcertSyncCore/Source/ConcertSyncCore/Public/Replication/Messages/ClientQuery.h
--- a/Engine/Plugins/Developer/Concert/ConcertSync/ConcertSyncCore/Source/ConcertSyncCore/Public/Replication/Messages/ClientQuery.h
+++ b/Engine/Plugins/Developer/Concert/ConcertSync/ConcertSyncCore/Source/ConcertSyncCore/Public/Replication/Messages/ClientQuery.h
@@ -2,6 +2,7 @@
 
 #pragma once
 
+#include "EReplicationResponseErrorCode.h"
 #include "Replication/Data/ClientQueriedInfo.h"
 #include "ClientQuery.generated.h"
 
@@ -51,10 +52,14 @@ struct FConcertReplication_QueryReplicationInfo_Request
 };
 
 USTRUCT()
-struct FConcertReplication_QueryReplicationInfo_Response
+struct FConcertReplication_QueryReplicationInfo_Response 
 {
 	GENERATED_BODY()
-
+	
+	/** Concert's custom requests are default constructed when they timeout. Server always sets this to Handled when processed. */
+	UPROPERTY()
+	EReplicationResponseErrorCode ErrorCode = EReplicationResponseErrorCode::Timeout;
+	
 	/**
 	 * Binds client participating in the replication session to their info.
 	 * 
diff --git a/Engine/Plugins/Developer/Concert/ConcertSync/ConcertSyncCore/Source/ConcertSyncCore/Public/Replication/Messages/EReplicationResponseErrorCode.h b/Engine/Plugins/Developer/Concert/ConcertSync/ConcertSyncCore/Source/ConcertSyncCore/Public/Replication/Messages/EReplicationResponseErrorCode.h
new file mode 100644
--- /dev/null
+++ b/Engine/Plugins/Developer/Concert/ConcertSync/ConcertSyncCore/Source/ConcertSyncCore/Public/Replication/Messages/EReplicationResponseErrorCode.h
@@ -0,0 +1,21 @@
+// Copyright Epic Games, Inc. All Rights Reserved.
+
+#pragma once
+
+#include "EReplicationResponseErrorCode.generated.h"
+
+/**
+ * Added as field to all replication response data.
+ * Used to detect timeouts or other delivery failures.
+ * 
+ * Custom requests and responses are considered failed if they are default constructed.
+ * The default constructed responses should be EReplicationResponseErrorCode::Timout by default.
+ */
+UENUM()
+enum class EReplicationResponseErrorCode : uint8
+{
+	/** Value to set when default constructed. Indicates responses not received. */
+	Timeout,
+	/** Set when request was handled and there is valid response data */
+	Handled
+};
\ No newline at end of file
diff --git a/Engine/Plugins/Developer/Concert/ConcertSync/ConcertSyncCore/Source/ConcertSyncCore/Public/Replication/Messages/Handshake.h b/Engine/Plugins/Developer/Concert/ConcertSync/ConcertSyncCore/Source/ConcertSyncCore/Public/Replication/Messages/Handshake.h
--- a/Engine/Plugins/Developer/Concert/ConcertSync/ConcertSyncCore/Source/ConcertSyncCore/Public/Replication/Messages/Handshake.h
+++ b/Engine/Plugins/Developer/Concert/ConcertSync/ConcertSyncCore/Source/ConcertSyncCore/Public/Replication/Messages/Handshake.h
@@ -81,7 +81,7 @@ struct FConcertReplication_Join_Response
 	 * Hence NetworkError means timeout, etc.
 	 */
 	UPROPERTY()
-	EJoinReplicationErrorCode ErrorCode = EJoinReplicationErrorCode::NetworkError;
+	EJoinReplicationErrorCode JoinErrorCode = EJoinReplicationErrorCode::NetworkError;
 
 	/** More information about ErrorCode to help the user resolve the issue. */
 	UPROPERTY()
diff --git a/Engine/Plugins/Developer/Concert/ConcertSync/ConcertSyncCore/Source/ConcertSyncCore/Public/Replication/Messages/ObjectReplication.h b/Engine/Plugins/Developer/Concert/ConcertSync/ConcertSyncCore/Source/ConcertSyncCore/Public/Replication/Messages/ObjectReplication.h
--- a/Engine/Plugins/Developer/Concert/ConcertSync/ConcertSyncCore/Source/ConcertSyncCore/Public/Replication/Messages/ObjectReplication.h
+++ b/Engine/Plugins/Developer/Concert/ConcertSync/ConcertSyncCore/Source/ConcertSyncCore/Public/Replication/Messages/ObjectReplication.h
@@ -3,8 +3,6 @@
 #pragma once
 
 #include "ConcertMessageData.h"
-#include "Replication/Data/ClientQueriedInfo.h"
-#include "Replication/Data/ObjectIds.h"
 #include "Misc/Optional.h"
 #include "ObjectReplication.generated.h"
 
diff --git a/Engine/Plugins/Developer/Concert/ConcertSync/ConcertSyncServer/Source/ConcertSyncServer/Private/Replication/AuthorityManager.cpp b/Engine/Plugins/Developer/Concert/ConcertSync/ConcertSyncServer/Source/ConcertSyncServer/Private/Replication/AuthorityManager.cpp
--- a/Engine/Plugins/Developer/Concert/ConcertSync/ConcertSyncServer/Source/ConcertSyncServer/Private/Replication/AuthorityManager.cpp
+++ b/Engine/Plugins/Developer/Concert/ConcertSync/ConcertSyncServer/Source/ConcertSyncServer/Private/Replication/AuthorityManager.cpp
@@ -185,6 +185,7 @@ namespace UE::ConcertSyncServer::Replication
 			}
 		});
 
+		Response.ErrorCode = EReplicationResponseErrorCode::Handled;
 		return EConcertSessionResponseCode::Success;
 	}
 	
diff --git a/Engine/Plugins/Developer/Concert/ConcertSync/ConcertSyncServer/Source/ConcertSyncServer/Private/Replication/ConcertServerReplicationManager.cpp b/Engine/Plugins/Developer/Concert/ConcertSync/ConcertSyncServer/Source/ConcertSyncServer/Private/Replication/ConcertServerReplicationManager.cpp
--- a/Engine/Plugins/Developer/Concert/ConcertSync/ConcertSyncServer/Source/ConcertSyncServer/Private/Replication/ConcertServerReplicationManager.cpp
+++ b/Engine/Plugins/Developer/Concert/ConcertSync/ConcertSyncServer/Source/ConcertSyncServer/Private/Replication/ConcertServerReplicationManager.cpp
@@ -6,6 +6,7 @@
 #include "ConcertLogGlobal.h"
 #include "IConcertSession.h"
 #include "Replication/ConcertReplicationClient.h"
+#include "Replication/Data/ClientQueriedInfo.h"
 #include "Replication/Formats/FullObjectFormat.h"
 #include "Replication/Messages/ChangeStream.h"
 #include "Replication/Messages/ClientQuery.h"
@@ -80,8 +81,8 @@ namespace UE::ConcertSyncServer::Replication
 		
 		const EConcertSessionResponseCode Result = InternalHandleJoinReplicationSessionRequest(ConcertSessionContext, Request, Response);
 		
-		UE_CLOG(Response.ErrorCode == EJoinReplicationErrorCode::Success, LogConcert, Log, TEXT("Accepted replication join request"));
-		UE_CLOG(Response.ErrorCode != EJoinReplicationErrorCode::Success, LogConcert, Log, TEXT("Rejected replication join request. %s: %s"), *ConcertSyncCore::Replication::LexJoinErrorCode(Response.ErrorCode), *Response.DetailedErrorMessage);
+		UE_CLOG(Response.JoinErrorCode == EJoinReplicationErrorCode::Success, LogConcert, Log, TEXT("Accepted replication join request"));
+		UE_CLOG(Response.JoinErrorCode != EJoinReplicationErrorCode::Success, LogConcert, Log, TEXT("Rejected replication join request. %s: %s"), *ConcertSyncCore::Replication::LexJoinErrorCode(Response.JoinErrorCode), *Response.DetailedErrorMessage);
 		return Result;
 	}
 
@@ -144,6 +145,7 @@ namespace UE::ConcertSyncServer::Replication
 			}
 		}
 
+		Response.ErrorCode = EReplicationResponseErrorCode::Handled;
 		return EConcertSessionResponseCode::Success;
 	}
 
diff --git a/Engine/Plugins/Developer/Concert/ConcertSync/ConcertSyncServer/Source/ConcertSyncServer/Private/Replication/ConcertServerReplicationManager.h b/Engine/Plugins/Developer/Concert/ConcertSync/ConcertSyncServer/Source/ConcertSyncServer/Private/Replication/ConcertServerReplicationManager.h
--- a/Engine/Plugins/Developer/Concert/ConcertSync/ConcertSyncServer/Source/ConcertSyncServer/Private/Replication/ConcertServerReplicationManager.h
+++ b/Engine/Plugins/Developer/Concert/ConcertSync/ConcertSyncServer/Source/ConcertSyncServer/Private/Replication/ConcertServerReplicationManager.h
@@ -19,6 +19,7 @@ class IConcertServerSession;
 struct FConcertReplication_ChangeStream_Response;
 struct FConcertReplication_QueryReplicationInfo_Response;
 struct FConcertReplication_QueryReplicationInfo_Request;
+struct FReplicationAuthorityInfo;
 
 namespace UE::ConcertSyncCore
 {
diff --git a/Engine/Plugins/Developer/Concert/ConcertSync/ConcertSyncServer/Source/ConcertSyncServer/Private/Replication/ConcertServerReplicationManager_ChangeStream.cpp b/Engine/Plugins/Developer/Concert/ConcertSync/ConcertSyncServer/Source/ConcertSyncServer/Private/Replication/ConcertServerReplicationManager_ChangeStream.cpp
--- a/Engine/Plugins/Developer/Concert/ConcertSync/ConcertSyncServer/Source/ConcertSyncServer/Private/Replication/ConcertServerReplicationManager_ChangeStream.cpp
+++ b/Engine/Plugins/Developer/Concert/ConcertSync/ConcertSyncServer/Source/ConcertSyncServer/Private/Replication/ConcertServerReplicationManager_ChangeStream.cpp
@@ -130,6 +130,7 @@ namespace UE::ConcertSyncServer::Replication
 		/** Checks whether this request is valid to apply. */
 		static bool ShouldAcceptRequest(const FConcertReplication_ChangeStream_Request& Request, const FConcertReplicationClient& Client, const FAuthorityManager& AuthorityManager, FConcertReplication_ChangeStream_Response& OutResponse)
 		{
+			OutResponse.ErrorCode = EReplicationResponseErrorCode::Handled;
 			ValidatePutObjectsRequestSemantics(Request, Client, OutResponse);
 			ValidateAddedStreamsAreUnique(Request, Client, OutResponse);
 			LookForAuthorityConflicts(Request, Client, AuthorityManager, OutResponse);
@@ -164,6 +165,7 @@ namespace UE::ConcertSyncServer::Replication
 			UE_LOG(LogConcert, Warning, TEXT("Rejecting ChangeStream request from %s"), *SendingClientId.ToString(EGuidFormats::Short));
 		}
 		
+		Response.ErrorCode = EReplicationResponseErrorCode::Handled;
 		return EConcertSessionResponseCode::Success;
 	}
 }
diff --git a/Engine/Plugins/Developer/Concert/ConcertSync/ConcertSyncTest/Source/ConcertSyncTest/Private/Replication/ReplicationAuthorityTests.cpp b/Engine/Plugins/Developer/Concert/ConcertSync/ConcertSyncTest/Source/ConcertSyncTest/Private/Replication/ReplicationAuthorityTests.cpp
--- a/Engine/Plugins/Developer/Concert/ConcertSync/ConcertSyncTest/Source/ConcertSyncTest/Private/Replication/ReplicationAuthorityTests.cpp
+++ b/Engine/Plugins/Developer/Concert/ConcertSync/ConcertSyncTest/Source/ConcertSyncTest/Private/Replication/ReplicationAuthorityTests.cpp
@@ -111,7 +111,8 @@ namespace UE::ConcertSyncTests::Replication::Authority
 			.Next([this, &bSenderReceivedResponse](const ConcertSyncClient::Replication::FAuthorityChangeResponse& Response) mutable
 			{
 				bSenderReceivedResponse = true;
-				TestEqual(TEXT("No rejection taking authority"), Response.RejectedObjects.Num(), 0); 
+				TestEqual(TEXT("No rejection taking authority"), Response.RejectedObjects.Num(), 0);
+				TestTrue(TEXT("ErrorCode == Handled"), Response.ErrorCode == EReplicationResponseErrorCode::Handled);
 			});
 		TickClient(Client_Sender);
 		TickServer();
@@ -124,6 +125,7 @@ namespace UE::ConcertSyncTests::Replication::Authority
 			{
 				bReceiverReceivedResponse = true;
 				TestEqual(TEXT("Rejected because Sender already has authority"), Response.RejectedObjects.Num(), 1);
+				TestTrue(TEXT("ErrorCode == Handled"), Response.ErrorCode == EReplicationResponseErrorCode::Handled);
 				
 				if (const FConcertStreamArray* RejectedStreams = Response.RejectedObjects.Find(TestObjectPath))
 				{
@@ -146,6 +148,7 @@ namespace UE::ConcertSyncTests::Replication::Authority
 			{
 				bSenderReceivedResponse = true;
 				TestEqual(TEXT("No rejection releasing object"), Response.RejectedObjects.Num(), 0); 
+				TestTrue(TEXT("ErrorCode == Handled"), Response.ErrorCode == EReplicationResponseErrorCode::Handled);
 			});
 		TickClient(Client_Sender);
 		TickServer();
@@ -158,6 +161,7 @@ namespace UE::ConcertSyncTests::Replication::Authority
 			{
 				bReceiverReceivedResponse = true;
 				TestEqual(TEXT("No rejection because the object should not be released"), Response.RejectedObjects.Num(), 0); 
+				TestTrue(TEXT("ErrorCode == Handled"), Response.ErrorCode == EReplicationResponseErrorCode::Handled);
 			});
 		TickClient(Client_Receiver);
 		TickServer();
@@ -188,6 +192,7 @@ namespace UE::ConcertSyncTests::Replication::Authority
 				bReceivedResponse = true;
 				TestTrue(TEXT("Cannot take authority over unregistered object"), Response.RejectedObjects.Contains(SomePath));
 				TestEqual(TEXT("Rejected exactly 1 object"), Response.RejectedObjects.Num(), 1);
+				TestTrue(TEXT("ErrorCode == Handled"), Response.ErrorCode == EReplicationResponseErrorCode::Handled);
 
 				if (const FConcertStreamArray* RejectedStreams = Response.RejectedObjects.Find(SomePath))
 				{
@@ -226,6 +231,7 @@ namespace UE::ConcertSyncTests::Replication::Authority
 			{
 				bSenderReceivedResponse = true;
 				TestEqual(TEXT("No rejection taking authority (sender)"), Response.RejectedObjects.Num(), 0); 
+				TestTrue(TEXT("ErrorCode == Handled"), Response.ErrorCode == EReplicationResponseErrorCode::Handled);
 			});
 		TickClient(Client_Sender);
 		TickServer();
@@ -243,6 +249,7 @@ namespace UE::ConcertSyncTests::Replication::Authority
 			{
 				bReceiverReceivedResponse = true;
 				TestEqual(TEXT("No rejection taking authority (receiver)"), Response.RejectedObjects.Num(), 0); 
+				TestTrue(TEXT("ErrorCode == Handled"), Response.ErrorCode == EReplicationResponseErrorCode::Handled);
 			});
 		TickClient(Client_Receiver);
 		TickServer();
@@ -270,6 +277,7 @@ namespace UE::ConcertSyncTests::Replication::Authority
 			{
 				bSenderReceivedResponse = true;
 				TestEqual(TEXT("No rejection taking authority (sender)"), Response.RejectedObjects.Num(), 0); 
+				TestTrue(TEXT("ErrorCode == Handled"), Response.ErrorCode == EReplicationResponseErrorCode::Handled);
 			});
 		TickClient(Client_Sender);
 		TickServer();
@@ -287,6 +295,7 @@ namespace UE::ConcertSyncTests::Replication::Authority
 			{
 				bReceiverReceivedResponse = true;
 				TestEqual(TEXT("No rejection taking authority (receiver)"), Response.RejectedObjects.Num(), 0); 
+				TestTrue(TEXT("ErrorCode == Handled"), Response.ErrorCode == EReplicationResponseErrorCode::Handled);
 			});
 		TickClient(Client_Receiver);
 		TickServer();
@@ -314,6 +323,7 @@ namespace UE::ConcertSyncTests::Replication::Authority
 			{
 				bSenderReceivedTakeResponse = true;
 				TestEqual(TEXT("No rejection taking authority"), Response.RejectedObjects.Num(), 0); 
+				TestTrue(TEXT("ErrorCode == Handled"), Response.ErrorCode == EReplicationResponseErrorCode::Handled);
 			});
 		TestTrue(TEXT("Sender received response taking authority"), bSenderReceivedTakeResponse);
 
@@ -328,6 +338,7 @@ namespace UE::ConcertSyncTests::Replication::Authority
 			{
 				bSenderReceivedReleaseResponse = true;
 				TestEqual(TEXT("No rejection releasing authority"), Response.RejectedObjects.Num(), 0); 
+				TestTrue(TEXT("ErrorCode == Handled"), Response.ErrorCode == EReplicationResponseErrorCode::Handled);
 			});
 		TestTrue(TEXT("Sender received response releasing authority"), bSenderReceivedReleaseResponse);
 		
@@ -336,4 +347,22 @@ namespace UE::ConcertSyncTests::Replication::Authority
 
 		return true;
 	}
+
+	/**
+	 * Tests that client updates its local cache of the server state when RequestAuthorityChange times out.
+	 */
+	IMPLEMENT_CUSTOM_SIMPLE_AUTOMATION_TEST(FChangingAuthorityTimeoutRetainsServerState, FSendReceiveObjectTestBase, "Concert.Replication.Authority.ChangeAuthorityTimeoutRetainsServerState", EAutomationTestFlags::EditorContext | EAutomationTestFlags::EngineFilter);
+	bool FChangingAuthorityTimeoutRetainsServerState::RunTest(const FString& Parameters)
+	{
+		// 1. Init
+		SetUpClientAndServer();
+		ServerSession->SetTestFlags(EServerSessionTestingFlags::AllowRequestTimeouts);
+		
+		ClientReplicationManager_Sender->TakeAuthorityOver({ TestObject });
+		ServerSession->UnregisterCustomRequestHandler<FConcertReplication_ChangeAuthority_Request>();
+		ClientReplicationManager_Sender->ReleaseAuthorityOf({ TestObject });
+
+		TestTrue(TEXT("Timed out ReleaseAuthorityOf reverted local server prediction"), ClientReplicationManager_Sender->GetClientOwnedObjects().Contains(TestObject));
+		return true;
+	}
 }
diff --git a/Engine/Plugins/Developer/Concert/ConcertSync/ConcertSyncTest/Source/ConcertSyncTest/Private/Replication/ReplicationSendReceiveTests.cpp b/Engine/Plugins/Developer/Concert/ConcertSync/ConcertSyncTest/Source/ConcertSyncTest/Private/Replication/ReplicationSendReceiveTests.cpp
--- a/Engine/Plugins/Developer/Concert/ConcertSync/ConcertSyncTest/Source/ConcertSyncTest/Private/Replication/ReplicationSendReceiveTests.cpp
+++ b/Engine/Plugins/Developer/Concert/ConcertSync/ConcertSyncTest/Source/ConcertSyncTest/Private/Replication/ReplicationSendReceiveTests.cpp
@@ -83,6 +83,7 @@ namespace UE::ConcertSyncTests::Replication::SendReceiveFlow
 			{
 				bReceivedChangeStreamResponse = true;
 				Test.TestTrue(TEXT("Changed Stream"), Response.IsSuccess());
+				Test.TestTrue(TEXT("ErrorCode == Handled"), Response.ErrorCode == EReplicationResponseErrorCode::Handled);
 			});
 		Test.TestTrue(TEXT("Received change stream response"), bReceivedChangeStreamResponse);
 	}
diff --git a/Engine/Plugins/Developer/Concert/ConcertSync/ConcertSyncTest/Source/ConcertSyncTest/Private/Replication/ReplicationStreamRequestTests.cpp b/Engine/Plugins/Developer/Concert/ConcertSync/ConcertSyncTest/Source/ConcertSyncTest/Private/Replication/ReplicationStreamRequestTests.cpp
--- a/Engine/Plugins/Developer/Concert/ConcertSync/ConcertSyncTest/Source/ConcertSyncTest/Private/Replication/ReplicationStreamRequestTests.cpp
+++ b/Engine/Plugins/Developer/Concert/ConcertSync/ConcertSyncTest/Source/ConcertSyncTest/Private/Replication/ReplicationStreamRequestTests.cpp
@@ -32,6 +32,8 @@ namespace UE::ConcertSyncTests::Replication::Stream
 		Request.ClientEndpointIds = { SenderEndpointId };
 		auto TestReplicationMapContent = [this, &SenderEndpointId](const ConcertSyncClient::Replication::FClientQueryResponse& Response)
 		{
+			TestTrue(TEXT("ErrorCode == Handled"), Response.ErrorCode == EReplicationResponseErrorCode::Handled);
+			
 			const FReplicationClientQueriedInfo* Info = Response.ClientInfo.Find(SenderEndpointId);
 			if (!Info)
 			{
@@ -60,6 +62,8 @@ namespace UE::ConcertSyncTests::Replication::Stream
 			.Next([this, &SenderEndpointId, &TestReplicationMapContent, &bReceivedFirstQueryResponse](ConcertSyncClient::Replication::FClientQueryResponse&& Response) mutable
 			{
 				bReceivedFirstQueryResponse = true;
+				TestTrue(TEXT("ErrorCode == Handled"), Response.ErrorCode == EReplicationResponseErrorCode::Handled);
+				
 				TestReplicationMapContent(Response);
 				if (const FReplicationClientQueriedInfo* Info = Response.ClientInfo.Find(SenderEndpointId))
 				{
@@ -75,6 +79,7 @@ namespace UE::ConcertSyncTests::Replication::Stream
 			{
 				bSenderReceivedResponse = true;
 				TestEqual(TEXT("No rejection taking authority"), Response.RejectedObjects.Num(), 0); 
+				TestTrue(TEXT("ErrorCode == Handled"), Response.ErrorCode == EReplicationResponseErrorCode::Handled);
 			});
 		TestTrue(TEXT("Received response taking authority"), bSenderReceivedResponse);
 
@@ -83,6 +88,7 @@ namespace UE::ConcertSyncTests::Replication::Stream
 			.Next([this, &SenderEndpointId, &TestReplicationMapContent, &TestObjectPath, &bReceivedSecondQueryResponse](ConcertSyncClient::Replication::FClientQueryResponse&& Response) mutable
 			{
 				bReceivedSecondQueryResponse = true;
+				TestTrue(TEXT("ErrorCode == Handled"), Response.ErrorCode == EReplicationResponseErrorCode::Handled);
 				TestReplicationMapContent(Response);
 				
 				const FReplicationClientQueriedInfo* Info = Response.ClientInfo.Find(SenderEndpointId);
@@ -108,6 +114,7 @@ namespace UE::ConcertSyncTests::Replication::Stream
 			.Next([this, &SenderEndpointId, &bReceivedResponse_SkipStreamInfo](ConcertSyncClient::Replication::FClientQueryResponse&& Response) mutable
 			{
 				bReceivedResponse_SkipStreamInfo = true;
+				TestTrue(TEXT("ErrorCode == Handled"), Response.ErrorCode == EReplicationResponseErrorCode::Handled);
 				if (const FReplicationClientQueriedInfo* Info = Response.ClientInfo.Find(SenderEndpointId))
 				{
 					TestEqual(TEXT("SkipStreamInfo > No stream data"), Info->Streams.Num(), 0);
@@ -123,6 +130,7 @@ namespace UE::ConcertSyncTests::Replication::Stream
 			{
 				bReceivedResponse_SkipProperties = true;
 				const FReplicationClientQueriedInfo* Info = Response.ClientInfo.Find(SenderEndpointId);
+				TestTrue(TEXT("ErrorCode == Handled"), Response.ErrorCode == EReplicationResponseErrorCode::Handled);
 				if (!Info || Info->Streams.IsEmpty())
 				{
 					AddError(TEXT("SkipProperties > No stream info"));
@@ -146,6 +154,7 @@ namespace UE::ConcertSyncTests::Replication::Stream
         	.Next([this, &SenderEndpointId, &bReceivedResponse_SkipAuthority](ConcertSyncClient::Replication::FClientQueryResponse&& Response) mutable
         	{
         		bReceivedResponse_SkipAuthority = true;
+				TestTrue(TEXT("ErrorCode == Handled"), Response.ErrorCode == EReplicationResponseErrorCode::Handled);
         		if (const FReplicationClientQueriedInfo* Info = Response.ClientInfo.Find(SenderEndpointId))
         		{
         			TestEqual(TEXT("SkipAuthority > No authority data"), Info->Authority.Num(), 0);
@@ -234,6 +243,7 @@ namespace UE::ConcertSyncTests::Replication::Stream
 				.Next([this, &InTestName, &bModifiedInitialStream](FChangeStreamResponse&& Response)
 				{
 					bModifiedInitialStream = true;
+					TestTrue(TEXT("ErrorCode == Handled"), Response.ErrorCode == EReplicationResponseErrorCode::Handled);
 					TestTrue(FString::Printf(TEXT("%s > Modify Stream > Success"), *InTestName), Response.IsSuccess());
 				});
 			TestTrue(FString::Printf(TEXT("%s > Modify Stream > Received response"), *InTestName), bModifiedInitialStream);
@@ -251,6 +261,7 @@ namespace UE::ConcertSyncTests::Replication::Stream
 				.Next([this, &InTestName, &bQueriedClientInfo, &Streams, &SenderId](FClientQueryResponse&& Response)
 				{
 					bQueriedClientInfo = true;
+					TestTrue(TEXT("ErrorCode == Handled"), Response.ErrorCode == EReplicationResponseErrorCode::Handled);
 					const FReplicationClientQueriedInfo* Info = Response.ClientInfo.Find(SenderId);
 					const bool bRegisteredAndQueriedStreamsAreEqual = Info->Streams == Streams;
 					TestTrue(FString::Printf(TEXT("%s > Registered and queried streams are equal"), *InTestName), bRegisteredAndQueriedStreamsAreEqual);
@@ -331,6 +342,7 @@ namespace UE::ConcertSyncTests::Replication::Stream
 			.Next([this, &bTookAuthority](FAuthorityChangeResponse&& Response)
 			{
 				bTookAuthority = true;
+				TestTrue(TEXT("ErrorCode == Handled"), Response.ErrorCode == EReplicationResponseErrorCode::Handled);
 				TestTrue(TEXT("Authority request > Success"), Response.RejectedObjects.IsEmpty());
 			});
 		TestTrue(TEXT("Authority request > Received response"), bTookAuthority);
@@ -370,6 +382,7 @@ namespace UE::ConcertSyncTests::Replication::Stream
 			.Next([this, &bSenderTookAuthority](const FAuthorityChangeResponse& Response) mutable
 			{
 				bSenderTookAuthority = true;
+				TestTrue(TEXT("ErrorCode == Handled"), Response.ErrorCode == EReplicationResponseErrorCode::Handled);
 				TestEqual(TEXT("Sender > No rejection taking authority"), Response.RejectedObjects.Num(), 0); 
 			});
 		TestTrue(TEXT("Sender > Received response taking authority"), bSenderTookAuthority);
@@ -383,6 +396,7 @@ namespace UE::ConcertSyncTests::Replication::Stream
 			.Next([this, &bAddedVectorStream](FChangeStreamResponse&& Response)
 			{
 				bAddedVectorStream = true;
+				TestTrue(TEXT("ErrorCode == Handled"), Response.ErrorCode == EReplicationResponseErrorCode::Handled);
 				TestTrue(TEXT("Vector Stream > Success"), Response.IsSuccess());
 			});
 		TestTrue(TEXT("Vector Stream> Received response"), bAddedVectorStream);
@@ -394,6 +408,7 @@ namespace UE::ConcertSyncTests::Replication::Stream
 			.Next([this, &bReceiverTookAuthority](const ConcertSyncClient::Replication::FAuthorityChangeResponse& Response) mutable
 			{
 				bReceiverTookAuthority = true;
+				TestTrue(TEXT("ErrorCode == Handled"), Response.ErrorCode == EReplicationResponseErrorCode::Handled);
 				TestEqual(TEXT("Receiver > No rejection taking authority"), Response.RejectedObjects.Num(), 0); 
 			});
 		TestTrue(TEXT("Receiver > Received response taking authority"), bReceiverTookAuthority);
@@ -412,6 +427,7 @@ namespace UE::ConcertSyncTests::Replication::Stream
 			.Next([this, TestObject, SenderStreamID = SenderStreamID, &TestObjectInReceiverStreamId, &bReceivedResponseAddingFloat](FChangeStreamResponse&& Response)
 			{
 				bReceivedResponseAddingFloat = true;
+				TestTrue(TEXT("ErrorCode == Handled"), Response.ErrorCode == EReplicationResponseErrorCode::Handled);
 				TestEqual(TEXT("Append Float > 1 conflict"), Response.AuthorityConflicts.Num(), 1);
 				if (const FReplicatedObjectId* ConflictingObject = Response.AuthorityConflicts.Find(TestObjectInReceiverStreamId))
 				{
@@ -461,6 +477,7 @@ namespace UE::ConcertSyncTests::Replication::Stream
 			.Next([this, &bReceivedChangeStreamResponse](FChangeStreamResponse&& Response)
 			{
 				bReceivedChangeStreamResponse = true;
+				TestTrue(TEXT("ErrorCode == Handled"), Response.ErrorCode == EReplicationResponseErrorCode::Handled);
 				TestTrue(TEXT("Invalid Request > Modify Stream > Failure"), Response.IsFailure());
 			});
 		
@@ -488,6 +505,7 @@ namespace UE::ConcertSyncTests::Replication::Stream
 		ClientReplicationManager_Sender->TakeAuthorityOver({ TestObject })
 			.Next([this](FAuthorityChangeResponse&& Response)
 			{
+				TestTrue(TEXT("ErrorCode == Handled"), Response.ErrorCode == EReplicationResponseErrorCode::Handled);
 				TestTrue(TEXT("Sender > Success taking authority"), Response.RejectedObjects.IsEmpty());
 			});
 
@@ -500,9 +518,41 @@ namespace UE::ConcertSyncTests::Replication::Stream
 		ClientReplicationManager_Receiver->TakeAuthorityOver({ TestObject })
 			.Next([this](FAuthorityChangeResponse&& Response)
 			{
+				TestTrue(TEXT("ErrorCode == Handled"), Response.ErrorCode == EReplicationResponseErrorCode::Handled);
 				TestTrue(TEXT("Receiver > Success taking authority"), Response.RejectedObjects.IsEmpty());
 			});
 		
 		return true;
 	}
+
+	/**
+     * Tests that client updates its local cache of the server state when RequestAuthorityChange times out.
+     */
+    IMPLEMENT_CUSTOM_SIMPLE_AUTOMATION_TEST(FChangingStreamTimeoutRetainsServerState, FChangeStreamsTestBase, "Concert.Replication.Stream.ChangeStreamTimeoutRetainsServerState", EAutomationTestFlags::EditorContext | EAutomationTestFlags::EngineFilter);
+    bool FChangingStreamTimeoutRetainsServerState::RunTest(const FString& Parameters)
+    {
+    	// 1. Start up client with float property stream
+    	UTestReflectionObject* TestObject = NewObject<UTestReflectionObject>(GetTransientPackage());
+    	auto[InitialStreamId, InitialStream] = CreateFloatPropertyStream(*TestObject);
+    	SenderArgs.Streams = { InitialStream };
+    	SetUpClientAndServer();
+		ServerSession->SetTestFlags(EServerSessionTestingFlags::AllowRequestTimeouts);
+
+    	// 2. Timeout request
+    	ServerSession->UnregisterCustomRequestHandler<FConcertReplication_ChangeStream_Request>();
+    	ConcertSyncClient::Replication::FChangeStreamRequest Request;
+    	Request.ObjectsToRemove.Add({ InitialStreamId, TestObject });
+    	ClientReplicationManager_Sender->ChangeStream(Request);
+
+    	// 3. Check that the local client's server prediction was reverted
+    	TArray<FReplicationStreamDescription> Streams = ClientReplicationManager_Sender->GetRegisteredStreams();
+    	if (Streams.IsEmpty())
+    	{
+    		AddError(TEXT("Stream change was not reverted"));
+    		return true;
+    	}
+    	TestTrue(TEXT("Timed out Stream change reverted"), Streams[0].BaseDescription.ReplicationMap.ReplicatedObjects.Contains(TestObject));
+    	
+    	return true;
+    }
 }
diff --git a/Engine/Plugins/Developer/Concert/ConcertSync/ConcertSyncTest/Source/ConcertSyncTest/Private/Util/ClientServerCommunicationTest.h b/Engine/Plugins/Developer/Concert/ConcertSync/ConcertSyncTest/Source/ConcertSyncTest/Private/Util/ClientServerCommunicationTest.h
--- a/Engine/Plugins/Developer/Concert/ConcertSync/ConcertSyncTest/Source/ConcertSyncTest/Private/Util/ClientServerCommunicationTest.h
+++ b/Engine/Plugins/Developer/Concert/ConcertSync/ConcertSyncTest/Source/ConcertSyncTest/Private/Util/ClientServerCommunicationTest.h
@@ -147,6 +147,14 @@ namespace UE::ConcertSyncTests
 		FConcertClientInfo ClientInfo;
 	};
 
+	enum class EServerSessionTestingFlags
+	{
+		None = 0,
+		/** In FConcertServerSessionMock::DispatchEvent if no request handler is registered, it is allowed to timeout the request. Otherwise a check() is triggered. */
+		AllowRequestTimeouts = 1 << 0
+	};
+	ENUM_CLASS_FLAGS(EServerSessionTestingFlags);
+
 	/** Specializes the base concert server session to act as a fake server session. */
 	class FConcertServerSessionMock : public FConcertServerSessionBaseMock
 	{
@@ -157,6 +165,11 @@ namespace UE::ConcertSyncTests
 			RegisterCustomEventHandler<FConcertSession_LeaveSessionEvent>(this, &FConcertServerSessionMock::HandleClientLeaveEvent);
 		}
 		
+		void SetTestFlags(EServerSessionTestingFlags InFlags)
+		{
+			TestFlags = InFlags;
+		}
+		
 		virtual void InternalSendCustomEvent(const UScriptStruct* EventType, const void* EventData, const TArray<FGuid>& TargetEndpointIds, EConcertMessageFlags, TOptional<FConcertSequencedCustomEvent> InSequenceId={}) override
 		{
 			for (const FGuid& TargetEndPointId : TargetEndpointIds)
@@ -279,6 +292,14 @@ namespace UE::ConcertSyncTests
 					check(false); // The test suite is not expected to fire any other result than Success or Failed.
 				}
 			}
+			else 
+			{
+				const bool bAllowTimeout = EnumHasAnyFlags(TestFlags, EServerSessionTestingFlags::AllowRequestTimeouts);
+				if (ensureAlways(bAllowTimeout))
+				{
+					ResponseHandler->HandleResponse(nullptr);
+				}
+			}
 		}
 
 		virtual bool FindSessionClient(const FGuid& EndpointId, FConcertSessionClientInfo& OutSessionClientInfo) const override
@@ -309,6 +330,8 @@ namespace UE::ConcertSyncTests
 
 		/** Connected clients sessions. */
 		TArray<FConcertClientSessionBaseMock*> ClientSessions;
+
+		EServerSessionTestingFlags TestFlags = EServerSessionTestingFlags::None;
 		
 		void HandleClientLeaveEvent(const FConcertSessionContext& Context, const FConcertSession_LeaveSessionEvent& Event)
 		{
