Automatically take authority for added objects

diff --git a/Engine/Plugins/Developer/Concert/ConcertApp/MultiUserClient/Source/MultiUserClient/Private/Replication/Client/ReplicationClient.cpp b/Engine/Plugins/Developer/Concert/ConcertApp/MultiUserClient/Source/MultiUserClient/Private/Replication/Client/ReplicationClient.cpp
--- a/Engine/Plugins/Developer/Concert/ConcertApp/MultiUserClient/Source/MultiUserClient/Private/Replication/Client/ReplicationClient.cpp
+++ b/Engine/Plugins/Developer/Concert/ConcertApp/MultiUserClient/Source/MultiUserClient/Private/Replication/Client/ReplicationClient.cpp
@@ -43,18 +43,30 @@ namespace UE::MultiUserClient
 	}
 
 	void FReplicationClient::OnObjectsChanged(
-		TConstArrayView<UObject*> Objects,
-		TConstArrayView<FSoftObjectPath> SoftObjectPaths,
+		TConstArrayView<UObject*> AddedObjects,
+		TConstArrayView<FSoftObjectPath> RemovedObjects,
 		ConcertClientSharedSlate::EReplicatedObjectChangeReason ReplicatedObjectChangeReason
 		)
 	{
 		// Could improve performance by just considering what actually changed instead of doing a full rebuild
+		// This must be done before SetAuthorityIfAllowed because it uses the cache for checking whether the object has properties assigned
 		LocalClientStreamDiffer.RefreshChangesCache();
+		
+		// Util for clients: automatically take authority for newly added objects
+		for (const UObject* Object : AddedObjects)
+		{
+			LocalAuthorityDiffer.SetAuthorityIfAllowed(Object, true);
+		}
+
+		// Refresh because authority changes may no longer be valid after modifying the stream
+		LocalAuthorityDiffer.RefreshChanges();
 	}
 
 	void FReplicationClient::OnPropertiesChanged()
 	{
 		// Could improve performance by just considering what actually changed instead of doing a full rebuild
 		LocalClientStreamDiffer.RefreshChangesCache();
+		// Refresh because authority changes may no longer be valid after modifying the stream
+		LocalAuthorityDiffer.RefreshChanges();
 	}
 }
diff --git a/Engine/Plugins/Developer/Concert/ConcertApp/MultiUserClient/Source/MultiUserClient/Private/Replication/Client/ReplicationClient.h b/Engine/Plugins/Developer/Concert/ConcertApp/MultiUserClient/Source/MultiUserClient/Private/Replication/Client/ReplicationClient.h
--- a/Engine/Plugins/Developer/Concert/ConcertApp/MultiUserClient/Source/MultiUserClient/Private/Replication/Client/ReplicationClient.h
+++ b/Engine/Plugins/Developer/Concert/ConcertApp/MultiUserClient/Source/MultiUserClient/Private/Replication/Client/ReplicationClient.h
@@ -92,7 +92,7 @@ namespace UE::MultiUserClient
 		FOnModelExternallyChanged OnModelExternallyChangedDelegate;
 
 		// Respond to model changing
-		void OnObjectsChanged(TConstArrayView<UObject*> Objects, TConstArrayView<FSoftObjectPath> SoftObjectPaths, ConcertClientSharedSlate::EReplicatedObjectChangeReason ReplicatedObjectChangeReason);
+		void OnObjectsChanged(TConstArrayView<UObject*> AddedObjects, TConstArrayView<FSoftObjectPath> RemovedObjects, ConcertClientSharedSlate::EReplicatedObjectChangeReason ReplicatedObjectChangeReason);
 		void OnPropertiesChanged();
 	};
 }
