Automatically submit changes as they are made by the user

diff --git a/Engine/Plugins/Developer/Concert/ConcertApp/MultiUserClient/Source/MultiUserClient/MultiUserClient.Build.cs b/Engine/Plugins/Developer/Concert/ConcertApp/MultiUserClient/Source/MultiUserClient/MultiUserClient.Build.cs
--- a/Engine/Plugins/Developer/Concert/ConcertApp/MultiUserClient/Source/MultiUserClient/MultiUserClient.Build.cs
+++ b/Engine/Plugins/Developer/Concert/ConcertApp/MultiUserClient/Source/MultiUserClient/MultiUserClient.Build.cs
@@ -20,6 +20,7 @@ namespace UnrealBuildTool.Rules
 					// Engine
 					"ApplicationCore",
 					"ContentBrowser",
+					"Core",
 					"DesktopPlatform",
 					"EditorStyle",
 					"InputCore",
diff --git a/Engine/Plugins/Developer/Concert/ConcertApp/MultiUserClient/Source/MultiUserClient/Private/Replication/Authority/AuthorityChangeTracker.cpp b/Engine/Plugins/Developer/Concert/ConcertApp/MultiUserClient/Source/MultiUserClient/Private/Replication/Authority/AuthorityChangeTracker.cpp
--- a/Engine/Plugins/Developer/Concert/ConcertApp/MultiUserClient/Source/MultiUserClient/Private/Replication/Authority/AuthorityChangeTracker.cpp
+++ b/Engine/Plugins/Developer/Concert/ConcertApp/MultiUserClient/Source/MultiUserClient/Private/Replication/Authority/AuthorityChangeTracker.cpp
@@ -20,11 +20,19 @@ namespace UE::MultiUserClient
 		AuthoritySynchronizer.OnServerStateChanged().RemoveAll(this);
 	}
 
-	void FAuthorityChangeTracker::SetAuthorityIfAllowed(const FSoftObjectPath& ObjectPath, bool bNewAuthorityState)
+	void FAuthorityChangeTracker::SetAuthorityIfAllowed(TConstArrayView<FSoftObjectPath> ObjectPaths, bool bNewAuthorityState)
 	{
-		if (CanSetAuthorityFor(ObjectPath))
+		for (const FSoftObjectPath& ObjectPath : ObjectPaths)
 		{
-			NewAuthorityStates.Add(ObjectPath, bNewAuthorityState);
+			if (CanSetAuthorityFor(ObjectPath))
+			{
+				NewAuthorityStates.Add(ObjectPath, bNewAuthorityState);
+			}
+		}
+
+		if (HasChanges())
+		{
+			FOnAuthorityChangeMadeDelegate.Broadcast();
 		}
 	}
 
diff --git a/Engine/Plugins/Developer/Concert/ConcertApp/MultiUserClient/Source/MultiUserClient/Private/Replication/Authority/AuthorityChangeTracker.h b/Engine/Plugins/Developer/Concert/ConcertApp/MultiUserClient/Source/MultiUserClient/Private/Replication/Authority/AuthorityChangeTracker.h
--- a/Engine/Plugins/Developer/Concert/ConcertApp/MultiUserClient/Source/MultiUserClient/Private/Replication/Authority/AuthorityChangeTracker.h
+++ b/Engine/Plugins/Developer/Concert/ConcertApp/MultiUserClient/Source/MultiUserClient/Private/Replication/Authority/AuthorityChangeTracker.h
@@ -3,6 +3,7 @@
 #pragma once
 
 #include "Containers/Map.h"
+#include "Containers/ContainersFwd.h"
 #include "Replication/IConcertClientReplicationManager.h"
 
 struct FSoftObjectPath;
@@ -20,8 +21,10 @@ namespace UE::MultiUserClient
 		FAuthorityChangeTracker(IClientAuthoritySynchronizer& InAuthoritySynchronizer);
 		~FAuthorityChangeTracker();
 		
-		/** Marks that the authority should be changed to bNewAuthorityState. */
-		void SetAuthorityIfAllowed(const FSoftObjectPath& ObjectPath, bool bNewAuthorityState);
+		/**
+		 * Marks that the authority should be changed to bNewAuthorityState.
+		 */
+		void SetAuthorityIfAllowed(TConstArrayView<FSoftObjectPath> ObjectPaths, bool bNewAuthorityState);
 
 		/** Diffs NewAuthorityStates to the current authority states and removes entries. */
 		void RefreshChanges();
@@ -39,6 +42,10 @@ namespace UE::MultiUserClient
 
 		/** Builds a change request from the local changes. */
 		ConcertSyncClient::Replication::FAuthorityChangeRequest BuildChangeRequest(const FGuid& StreamId) const;
+
+		/** Called when NewAuthorityStates is updated. */
+		DECLARE_MULTICAST_DELEGATE(FOnAuthorityChangeMade);
+		FOnAuthorityChangeMade& OnAuthorityChangeMade() { return FOnAuthorityChangeMadeDelegate; }
 		
 	private:
 
@@ -47,6 +54,9 @@ namespace UE::MultiUserClient
 
 		/** Object to the authority state it should have. */
 		TMap<FSoftObjectPath, bool> NewAuthorityStates;
+
+		/** Called when NewAuthorityStates is updated. */
+		FOnAuthorityChangeMade FOnAuthorityChangeMadeDelegate;
 	};
 }
 
diff --git a/Engine/Plugins/Developer/Concert/ConcertApp/MultiUserClient/Source/MultiUserClient/Private/Replication/Client/LocalReplicationClient.cpp b/Engine/Plugins/Developer/Concert/ConcertApp/MultiUserClient/Source/MultiUserClient/Private/Replication/Client/LocalReplicationClient.cpp
--- a/Engine/Plugins/Developer/Concert/ConcertApp/MultiUserClient/Source/MultiUserClient/Private/Replication/Client/LocalReplicationClient.cpp
+++ b/Engine/Plugins/Developer/Concert/ConcertApp/MultiUserClient/Source/MultiUserClient/Private/Replication/Client/LocalReplicationClient.cpp
@@ -21,9 +21,9 @@ namespace UE::MultiUserClient
 				{
 					return GetStreamDiffer().DoesObjectHavePropertiesAfterSubmit(ObjectPath);
 				})),
-			[this, InClient]()
+			[this, InClient](FStreamChangeTracker& InStreamChangeTracker, FAuthorityChangeTracker& InAuthorityChangeTracker, IClientStreamSynchronizer& InStreamSynchronizer)
 			{
-				return MakeUnique<FSubmissionWorkflow_LocalClient>(InClient, GetStreamDiffer(), GetAuthorityDiffer(), GetStreamSynchronizer());
+				return MakeUnique<FSubmissionWorkflow_LocalClient>(InClient, InStreamChangeTracker, InAuthorityChangeTracker, InStreamSynchronizer);
 			})
 	{}
 }
diff --git a/Engine/Plugins/Developer/Concert/ConcertApp/MultiUserClient/Source/MultiUserClient/Private/Replication/Client/RemoteReplicationClient.cpp b/Engine/Plugins/Developer/Concert/ConcertApp/MultiUserClient/Source/MultiUserClient/Private/Replication/Client/RemoteReplicationClient.cpp
--- a/Engine/Plugins/Developer/Concert/ConcertApp/MultiUserClient/Source/MultiUserClient/Private/Replication/Client/RemoteReplicationClient.cpp
+++ b/Engine/Plugins/Developer/Concert/ConcertApp/MultiUserClient/Source/MultiUserClient/Private/Replication/Client/RemoteReplicationClient.cpp
@@ -24,8 +24,10 @@ namespace UE::MultiUserClient
 				{
 					return GetStreamDiffer().DoesObjectHavePropertiesAfterSubmit(ObjectPath);
 				})),
-			[]() { return MakeUnique<FSubmissionWorkflow_RemoteClient>(); }
-			)
+			[](FStreamChangeTracker& InStreamChangeTracker, FAuthorityChangeTracker& InAuthorityChangeTracker, IClientStreamSynchronizer& InStreamSynchronizer)
+			{
+				return MakeUnique<FSubmissionWorkflow_RemoteClient>();
+			})
 		, RemoteEndpointId(InConcertClientId)
 	{
 		// When the remote client's state has changed, refresh the UI.
diff --git a/Engine/Plugins/Developer/Concert/ConcertApp/MultiUserClient/Source/MultiUserClient/Private/Replication/Client/ReplicationClient.cpp b/Engine/Plugins/Developer/Concert/ConcertApp/MultiUserClient/Source/MultiUserClient/Private/Replication/Client/ReplicationClient.cpp
--- a/Engine/Plugins/Developer/Concert/ConcertApp/MultiUserClient/Source/MultiUserClient/Private/Replication/Client/ReplicationClient.cpp
+++ b/Engine/Plugins/Developer/Concert/ConcertApp/MultiUserClient/Source/MultiUserClient/Private/Replication/Client/ReplicationClient.cpp
@@ -14,7 +14,7 @@ namespace UE::MultiUserClient
 		UMultiUserReplicationClientPreset& InSessionContent,
 		TUniquePtr<IClientStreamSynchronizer> InStreamSynchronizer,
 		TUniquePtr<IClientAuthoritySynchronizer> InAuthoritySynchronizer,
-		TFunctionRef<TUniquePtr<ISubmissionWorkflow>()> MakeSubmissionWorkflowFunc
+		TFunctionRef<FMakeSubmissionWorkflow> MakeSubmissionWorkflowFunc
 		)
 		: ClientContentStorage(&InSessionContent)
 		, StreamSynchronizer(MoveTemp(InStreamSynchronizer))
@@ -30,9 +30,9 @@ namespace UE::MultiUserClient
 			FStreamChangeTracker::FOnModifyReplicationMap::CreateLambda([this](){ ClientContentStorage->Stream->Modify(); })
 			)
 		, LocalAuthorityDiffer(*AuthoritySynchronizer)
+		, SubmissionWorkflow(MakeSubmissionWorkflowFunc(LocalClientStreamDiffer, LocalAuthorityDiffer, *StreamSynchronizer.Get()))
+		, AutoSubmissionPolicy(*SubmissionWorkflow.Get(), LocalClientEditModel.Get(), LocalAuthorityDiffer)
 	{
-		SubmissionWorkflow = MakeSubmissionWorkflowFunc();
-		
 		LocalClientEditModel->OnObjectsChanged().AddRaw(this, &FReplicationClient::OnObjectsChanged);
 		LocalClientEditModel->OnPropertiesChanged().AddRaw(this, &FReplicationClient::OnPropertiesChanged);
 		
@@ -52,11 +52,10 @@ namespace UE::MultiUserClient
 		// This must be done before SetAuthorityIfAllowed because it uses the cache for checking whether the object has properties assigned
 		LocalClientStreamDiffer.RefreshChangesCache();
 		
-		// Util for clients: automatically take authority for newly added objects
-		for (const UObject* Object : AddedObjects)
-		{
-			LocalAuthorityDiffer.SetAuthorityIfAllowed(Object, true);
-		}
+		// Better UX for user: automatically take authority for newly added objects
+		TArray<FSoftObjectPath> ObjectPaths;
+		Algo::Transform(AddedObjects, ObjectPaths, [](const UObject* Object){ return FSoftObjectPath(Object); });
+		LocalAuthorityDiffer.SetAuthorityIfAllowed(ObjectPaths, true);
 
 		// Refresh because authority changes may no longer be valid after modifying the stream
 		LocalAuthorityDiffer.RefreshChanges();
diff --git a/Engine/Plugins/Developer/Concert/ConcertApp/MultiUserClient/Source/MultiUserClient/Private/Replication/Client/ReplicationClient.h b/Engine/Plugins/Developer/Concert/ConcertApp/MultiUserClient/Source/MultiUserClient/Private/Replication/Client/ReplicationClient.h
--- a/Engine/Plugins/Developer/Concert/ConcertApp/MultiUserClient/Source/MultiUserClient/Private/Replication/Client/ReplicationClient.h
+++ b/Engine/Plugins/Developer/Concert/ConcertApp/MultiUserClient/Source/MultiUserClient/Private/Replication/Client/ReplicationClient.h
@@ -6,7 +6,9 @@
 #include "Replication/Authority/IClientAuthoritySynchronizer.h"
 #include "Replication/Stream/IClientStreamSynchronizer.h"
 #include "Replication/Stream/StreamChangeTracker.h"
+#include "Replication/Submission/AutoSubmissionPolicy.h"
 #include "Replication/Submission/ISubmissionWorkflow.h"
+
 #include "Templates/SharedPointer.h"
 #include "Templates/UnrealTemplate.h"
 
@@ -26,14 +28,15 @@ namespace UE::MultiUserClient
 	class FReplicationClient : public FNoncopyable
 	{
 	public:
-
-		DECLARE_DELEGATE_RetVal(TUniquePtr<ISubmissionWorkflow>, FMakeSubmissionWorkflow);
+		
+		/** Indirection for creating ISubmissionWorkflow because some ISubmissionWorkflow implementations needs members constructed in FReplicationClient. */
+		using FMakeSubmissionWorkflow = TUniquePtr<ISubmissionWorkflow>(FStreamChangeTracker&, FAuthorityChangeTracker&, IClientStreamSynchronizer&);
 
 		FReplicationClient(
 			UMultiUserReplicationClientPreset& InSessionContent,
 			TUniquePtr<IClientStreamSynchronizer> InStreamSynchronizer,
 			TUniquePtr<IClientAuthoritySynchronizer> InAuthoritySynchronizer,
-			TFunctionRef<TUniquePtr<ISubmissionWorkflow>()> MakeSubmissionWorkflowFunc
+			TFunctionRef<FMakeSubmissionWorkflow> MakeSubmissionWorkflowFunc
 			);
 
 		UMultiUserReplicationClientPreset* GetClientContent() const { return ClientContentStorage; }
@@ -69,8 +72,6 @@ namespace UE::MultiUserClient
 		TUniquePtr<IClientStreamSynchronizer> StreamSynchronizer;
 		/** Keeps the client's authority state on the server in sync.*/
 		TUniquePtr<IClientAuthoritySynchronizer> AuthoritySynchronizer;
-		/** Handles the logic of submitting and reverting for this client. */
-		TUniquePtr<ISubmissionWorkflow> SubmissionWorkflow;
 		
 		/**
 		 * Used to detect changes made to the client's config by the local editor.
@@ -87,6 +88,11 @@ namespace UE::MultiUserClient
 		FStreamChangeTracker LocalClientStreamDiffer;
 		/** Tracks changes made to the client's authority state. */
 		FAuthorityChangeTracker LocalAuthorityDiffer;
+		
+		/** Handles the logic of submitting and reverting for this client. */
+		TUniquePtr<ISubmissionWorkflow> SubmissionWorkflow;
+		/** Automatically submits changes as they are made by the user. */
+		FAutoSubmissionPolicy AutoSubmissionPolicy;
 
 		/** Called when the data underlying the model has changed (and the UI needs to be refreshed). */
 		FOnModelExternallyChanged OnModelExternallyChangedDelegate;
diff --git a/Engine/Plugins/Developer/Concert/ConcertApp/MultiUserClient/Source/MultiUserClient/Private/Replication/Client/ReplicationClientManager.cpp b/Engine/Plugins/Developer/Concert/ConcertApp/MultiUserClient/Source/MultiUserClient/Private/Replication/Client/ReplicationClientManager.cpp
--- a/Engine/Plugins/Developer/Concert/ConcertApp/MultiUserClient/Source/MultiUserClient/Private/Replication/Client/ReplicationClientManager.cpp
+++ b/Engine/Plugins/Developer/Concert/ConcertApp/MultiUserClient/Source/MultiUserClient/Private/Replication/Client/ReplicationClientManager.cpp
@@ -25,6 +25,7 @@ namespace UE::MultiUserClient
 			UMultiUserReplicationClientPreset* ClientPreset = SessionContent->AddClient();
 			return FLocalReplicationClient(*ClientPreset, MakeUnique<FStreamSynchronizer_LocalClient>(InClient, ClientPreset->Stream->StreamId), InClient);
 		}())
+		, SubmissionNotifier(*this)
 	{
 		InSession->OnSessionClientChanged().AddRaw(this, &FReplicationClientManager::OnSessionClientChanged);
 
@@ -92,6 +93,7 @@ namespace UE::MultiUserClient
 
 				{
 					const TUniquePtr<FRemoteReplicationClient> Client = MoveTemp(RemoteClients[Index]);
+					OnPreRemoteClientRemovedDelegate.Broadcast(*Client.Get());
 					SessionContent->RemoveClient(*Client->GetClientContent());
 					RemoteClients.RemoveAtSwap(Index);
 				}
@@ -108,12 +110,15 @@ namespace UE::MultiUserClient
 	
 	void FReplicationClientManager::CreateRemoteClient(const FGuid& ClientEndpointId, bool bBroadcastDelegate)
 	{
+		TUniquePtr<FRemoteReplicationClient> RemoteClientPtr = MakeUnique<FRemoteReplicationClient>(*SessionContent->AddClient(), ClientEndpointId, QueryService);
+		FRemoteReplicationClient& RemoteClient = *RemoteClientPtr;
 		RemoteClients.Emplace(
-			MakeUnique<FRemoteReplicationClient>(*SessionContent->AddClient(), ClientEndpointId, QueryService)
+			MoveTemp(RemoteClientPtr)
 			);
 
 		if (bBroadcastDelegate)
 		{
+			OnPostRemoteClientAddedDelegate.Broadcast(RemoteClient);
 			OnRemoteClientsChangedDelegate.Broadcast();
 		}
 	}
diff --git a/Engine/Plugins/Developer/Concert/ConcertApp/MultiUserClient/Source/MultiUserClient/Private/Replication/Client/ReplicationClientManager.h b/Engine/Plugins/Developer/Concert/ConcertApp/MultiUserClient/Source/MultiUserClient/Private/Replication/Client/ReplicationClientManager.h
--- a/Engine/Plugins/Developer/Concert/ConcertApp/MultiUserClient/Source/MultiUserClient/Private/Replication/Client/ReplicationClientManager.h
+++ b/Engine/Plugins/Developer/Concert/ConcertApp/MultiUserClient/Source/MultiUserClient/Private/Replication/Client/ReplicationClientManager.h
@@ -6,6 +6,7 @@
 #include "ConcertMessageData.h"
 #include "LocalReplicationClient.h"
 #include "ReplicationClient.h"
+#include "Replication/Submission/Notification/SubmissionNotifier.h"
 #include "Replication/Util/RegularQueryService.h"
 
 #include "UObject/GCObject.h"
@@ -60,9 +61,15 @@ namespace UE::MultiUserClient
 		}
 		
 		DECLARE_MULTICAST_DELEGATE(FRemoteClientsChanged);
-		/** Called when RemoteClients changes. */
+		/** Called when RemoteClients changes. Called after OnPostRemoteClientAdded. */
 		FRemoteClientsChanged& OnRemoteClientsChanged() { return OnRemoteClientsChangedDelegate; }
 		
+		DECLARE_MULTICAST_DELEGATE_OneParam(FRemoteClientDelegate, FRemoteReplicationClient&);
+		/** Called just after a remote client is has been added to RemoteClients. Called before OnRemoteClientsChanged. */
+		FRemoteClientDelegate& OnPostRemoteClientAdded() { return OnPostRemoteClientAddedDelegate; }
+		/** Called just before a remote client is about to be removed from RemoteClients. */
+		FRemoteClientDelegate& OnPreRemoteClientRemoved() { return OnPreRemoteClientRemovedDelegate; }
+		
 		//~ Begin FGCObject Interface
 		virtual void AddReferencedObjects(FReferenceCollector& Collector) override;
 		virtual FString GetReferencerName() const override { return TEXT("FReplicationStreamSynchronizer"); }
@@ -94,10 +101,17 @@ namespace UE::MultiUserClient
 		 * UI keeps references to systems inside the client so it is TSharedRef in case TArray is reallocated.
 		 */
 		TArray<TUniquePtr<FRemoteReplicationClient>> RemoteClients;
-
+		
 		/** Called when RemoteClients changes. */
 		FRemoteClientsChanged OnRemoteClientsChangedDelegate;
-
+		/** Called when RemoteClients changes. Called after OnPostRemoteClientAdded. */
+		FRemoteClientDelegate OnPostRemoteClientAddedDelegate;
+		/** Called just before a remote client is about to be removed from RemoteClients. */
+		FRemoteClientDelegate OnPreRemoteClientRemovedDelegate; 
+		
+		/** Manages SNotificationItems when submission to the server fails. */
+		FSubmissionNotifier SubmissionNotifier;
+		
 		/** Updates RemoteClients depending on the change. */
 		void OnSessionClientChanged(IConcertClientSession&, EConcertClientStatus NewStatus, const FConcertSessionClientInfo& ClientInfo);
 		
diff --git a/Engine/Plugins/Developer/Concert/ConcertApp/MultiUserClient/Source/MultiUserClient/Private/Replication/Submission/AutoSubmissionPolicy.cpp b/Engine/Plugins/Developer/Concert/ConcertApp/MultiUserClient/Source/MultiUserClient/Private/Replication/Submission/AutoSubmissionPolicy.cpp
new file mode 100644
--- /dev/null
+++ b/Engine/Plugins/Developer/Concert/ConcertApp/MultiUserClient/Source/MultiUserClient/Private/Replication/Submission/AutoSubmissionPolicy.cpp
@@ -0,0 +1,51 @@
+// Copyright Epic Games, Inc. All Rights Reserved.
+
+#include "AutoSubmissionPolicy.h"
+
+#include "ISubmissionWorkflow.h"
+#include "Replication/Editor/Model/IEditableObjectToPropertiesModel.h"
+
+#include "Misc/CoreDelegates.h"
+#include "Replication/Authority/AuthorityChangeTracker.h"
+
+namespace UE::MultiUserClient
+{
+	FAutoSubmissionPolicy::FAutoSubmissionPolicy(
+		ISubmissionWorkflow& InSubmissionWorkflow,
+		ConcertClientSharedSlate::IEditableObjectToPropertiesModel& InStreamEditorModel,
+		FAuthorityChangeTracker& InAuthorityChangeTracker
+		)
+		: SubmissionWorkflow(InSubmissionWorkflow)
+		, StreamEditorModel(InStreamEditorModel)
+		, AuthorityChangeTracker(InAuthorityChangeTracker)
+	{
+		StreamEditorModel.OnObjectsChanged().AddRaw(this, &FAutoSubmissionPolicy::OnObjectsChanged);
+		StreamEditorModel.OnPropertiesChanged().AddRaw(this, &FAutoSubmissionPolicy::OnChangesDetected);
+		AuthorityChangeTracker.OnAuthorityChangeMade().AddRaw(this, &FAutoSubmissionPolicy::OnChangesDetected);
+	}
+
+	FAutoSubmissionPolicy::~FAutoSubmissionPolicy()
+	{
+		FCoreDelegates::OnEndFrame.RemoveAll(this);
+		StreamEditorModel.OnObjectsChanged().RemoveAll(this);
+		StreamEditorModel.OnPropertiesChanged().RemoveAll(this);
+		AuthorityChangeTracker.OnAuthorityChangeMade().RemoveAll(this);
+	}
+
+	void FAutoSubmissionPolicy::OnChangesDetected()
+	{
+		if (!FCoreDelegates::OnEndFrame.IsBoundToObject(this))
+		{
+			FCoreDelegates::OnEndFrame.AddRaw(this, &FAutoSubmissionPolicy::OnEndFrame);
+		}
+	}
+
+	void FAutoSubmissionPolicy::OnEndFrame()
+	{
+		if (SubmissionWorkflow.CanSubmit())
+		{
+			FCoreDelegates::OnEndFrame.RemoveAll(this);
+			SubmissionWorkflow.SubmitChanges();
+		}
+	}
+}
diff --git a/Engine/Plugins/Developer/Concert/ConcertApp/MultiUserClient/Source/MultiUserClient/Private/Replication/Submission/AutoSubmissionPolicy.h b/Engine/Plugins/Developer/Concert/ConcertApp/MultiUserClient/Source/MultiUserClient/Private/Replication/Submission/AutoSubmissionPolicy.h
new file mode 100644
--- /dev/null
+++ b/Engine/Plugins/Developer/Concert/ConcertApp/MultiUserClient/Source/MultiUserClient/Private/Replication/Submission/AutoSubmissionPolicy.h
@@ -0,0 +1,53 @@
+// Copyright Epic Games, Inc. All Rights Reserved.
+
+#pragma once
+
+#include "Replication/Editor/Model/IEditableObjectToPropertiesModel.h"
+#include "Templates/UnrealTemplate.h"
+
+namespace UE::ConcertClientSharedSlate
+{
+	class IEditableObjectToPropertiesModel;
+}
+
+namespace UE::MultiUserClient
+{
+	class FAuthorityChangeTracker;
+	class ISubmissionWorkflow;
+	
+	/**
+	 * Automatically submits stream and authority changes as the user makes them.
+	 * If a change is in progress when the change is made, the changes are queued and sent once the pending submission is done.
+	 *
+	 * This policy accumulates all changes until the end of frame and then sends them. This means multiple changes are made in the same frame,
+	 * only one request will be sent containing all changes.
+	 */
+	class FAutoSubmissionPolicy : public FNoncopyable
+	{
+	public:
+		
+		FAutoSubmissionPolicy(
+			ISubmissionWorkflow& InSubmissionWorkflow,
+			ConcertClientSharedSlate::IEditableObjectToPropertiesModel& InStreamEditorModel,
+			FAuthorityChangeTracker& InAuthorityChangeTracker
+			);
+		~FAutoSubmissionPolicy();
+		
+	private:
+
+		/** Handles performing the submission */
+		ISubmissionWorkflow& SubmissionWorkflow;
+
+		/** Informs us when the stream is structurally changed by the user. */
+		ConcertClientSharedSlate::IEditableObjectToPropertiesModel& StreamEditorModel;
+		/** Informs us when authority is changed by the user. */
+		FAuthorityChangeTracker& AuthorityChangeTracker;
+		
+		void OnObjectsChanged(TArrayView<UObject* const>, TArrayView<const FSoftObjectPath>, ConcertClientSharedSlate::EReplicatedObjectChangeReason) { OnChangesDetected(); }
+		void OnChangesDetected();
+
+		/** Checks whether any changes were made last frame and submits a change request if so. */
+		void OnEndFrame();
+	};
+}
+
diff --git a/Engine/Plugins/Developer/Concert/ConcertApp/MultiUserClient/Source/MultiUserClient/Private/Replication/Submission/ISubmissionWorkflow.h b/Engine/Plugins/Developer/Concert/ConcertApp/MultiUserClient/Source/MultiUserClient/Private/Replication/Submission/ISubmissionWorkflow.h
--- a/Engine/Plugins/Developer/Concert/ConcertApp/MultiUserClient/Source/MultiUserClient/Private/Replication/Submission/ISubmissionWorkflow.h
+++ b/Engine/Plugins/Developer/Concert/ConcertApp/MultiUserClient/Source/MultiUserClient/Private/Replication/Submission/ISubmissionWorkflow.h
@@ -5,6 +5,7 @@
 #include "Data/EChangeRevertability.h"
 #include "Data/EChangeUploadability.h"
 
+#include "Delegates/Delegate.h"
 #include "Templates/Function.h"
 #include "Templates/SharedPointer.h"
 
@@ -12,6 +13,8 @@ namespace UE::MultiUserClient
 {
 	class ISubmissionOperation;
 	struct FSubmitStreamChangesResponse;
+	struct FSubmitAuthorityChangesRequest;
+	struct FSubmitAuthorityChangesResponse;
 	
 	/**
 	 * Manages the flow of changing stream and authority on the server.
@@ -37,9 +40,6 @@ namespace UE::MultiUserClient
 		 */
 		virtual TSharedPtr<ISubmissionOperation> SubmitChanges() = 0;
 
-		/** Reverts all local changes */
-		virtual void RevertChanges() = 0;
-
 		/** @return Detailed information about whether Submit can be called */
 		virtual EChangeUploadability GetUploadability() const = 0;
 		bool CanSubmit() const { return GetUploadability() == EChangeUploadability::Ready; }
@@ -47,7 +47,37 @@ namespace UE::MultiUserClient
 		/** @return Detailed information about whether Revert can be called */
 		virtual EChangeRevertability GetRevertability() const = 0;
 		bool CanRevert() const { return GetRevertability() == EChangeRevertability::Revertable; }
+		
+		DECLARE_MULTICAST_DELEGATE_OneParam(FOnStreamRequestCompleted, const FSubmitStreamChangesResponse&);
+		/**
+		 * Called whenever a submit operation completes the stream change request stage.
+		 * @note No stream changes may have been requested. Check the error code.
+		 */
+		virtual FOnStreamRequestCompleted& OnStreamRequestCompleted() = 0;
+
+		DECLARE_MULTICAST_DELEGATE_TwoParams(FOnAuthorityRequestCompleted, const FSubmitAuthorityChangesRequest&, const FSubmitAuthorityChangesResponse&);
+		/**
+		 * Called whenever a submit operation completes the authority change request stage.
+		 * @note No authority changes may have been requested. Check the error code.
+		 */
+		virtual FOnAuthorityRequestCompleted& OnAuthorityRequestCompleted() = 0;
 
 		virtual ~ISubmissionWorkflow() = default;
 	};
+
+	/** Shared implementation for ISubmissionWorkflow. */
+	class FSubmissionWorkflowBase : public ISubmissionWorkflow
+	{
+	public:
+
+		//~ Begin ISubmissionWorkflow Interface
+		virtual FOnStreamRequestCompleted& OnStreamRequestCompleted() override { return StreamRequestCompletedDelegate; }
+		virtual FOnAuthorityRequestCompleted& OnAuthorityRequestCompleted() override { return AuthorityRequestCompletedDelegate; }
+		//~ End ISubmissionWorkflow Interface
+
+	protected:
+		
+		FOnStreamRequestCompleted StreamRequestCompletedDelegate;
+		FOnAuthorityRequestCompleted AuthorityRequestCompletedDelegate;
+	};
 }
diff --git a/Engine/Plugins/Developer/Concert/ConcertApp/MultiUserClient/Source/MultiUserClient/Private/Replication/Submission/Notification/AccumulatedSubmissionErrors.h b/Engine/Plugins/Developer/Concert/ConcertApp/MultiUserClient/Source/MultiUserClient/Private/Replication/Submission/Notification/AccumulatedSubmissionErrors.h
new file mode 100644
--- /dev/null
+++ b/Engine/Plugins/Developer/Concert/ConcertApp/MultiUserClient/Source/MultiUserClient/Private/Replication/Submission/Notification/AccumulatedSubmissionErrors.h
@@ -0,0 +1,30 @@
+// Copyright Epic Games, Inc. All Rights Reserved.
+
+#pragma once
+
+#include "Containers/Set.h"
+#include "HAL/Platform.h"
+#include "UObject/SoftObjectPath.h"
+
+namespace UE::MultiUserClient
+{
+	struct FAccumulatedStreamErrors
+	{
+		int32 NumTimeouts = 0;
+
+		/** Changing the these objects would cause authority conflicts. */
+		TSet<FSoftObjectPath> AuthorityConflicts;
+		/** This should be empty. If not, that means our client code made a bad request. */
+		TSet<FSoftObjectPath> SemanticErrors;
+		/** Should be false. If not, that means our client code made a bad request. */
+		bool bFailedStreamCreation = false;
+	};
+	
+	struct FAccumulatedAuthorityErrors
+	{
+		int32 NumTimeouts = 0;
+
+		/** These objects had conflicts */
+		TSet<FSoftObjectPath> Rejected;
+	};
+}
\ No newline at end of file
diff --git a/Engine/Plugins/Developer/Concert/ConcertApp/MultiUserClient/Source/MultiUserClient/Private/Replication/Submission/Notification/SAuthorityRejectedNotification.cpp b/Engine/Plugins/Developer/Concert/ConcertApp/MultiUserClient/Source/MultiUserClient/Private/Replication/Submission/Notification/SAuthorityRejectedNotification.cpp
new file mode 100644
--- /dev/null
+++ b/Engine/Plugins/Developer/Concert/ConcertApp/MultiUserClient/Source/MultiUserClient/Private/Replication/Submission/Notification/SAuthorityRejectedNotification.cpp
@@ -0,0 +1,69 @@
+// Copyright Epic Games, Inc. All Rights Reserved.
+
+#include "SAuthorityRejectedNotification.h"
+
+#include "AccumulatedSubmissionErrors.h"
+
+#include "Styling/AppStyle.h"
+#include "Widgets/Input/SButton.h"
+#include "Widgets/SBoxPanel.h"
+#include "Widgets/Text/STextBlock.h"
+
+#define LOCTEXT_NAMESPACE "SAuthorityRejectedNotification"
+
+namespace UE::MultiUserClient
+{
+	void SAuthorityRejectedNotification::Construct(const FArguments& InArgs)
+	{
+		ErrorsAttribute = InArgs._Errors;
+		
+		Super::Construct(
+			Super::FArguments()
+			.Message(LOCTEXT("MainMessage", "Errors taking ownership"))
+			.OnCloseClicked(InArgs._OnCloseClicked)
+			);
+	}
+
+	void SAuthorityRejectedNotification::Refresh()
+	{
+		TSharedRef<SVerticalBox> Result = SNew(SVerticalBox);
+		const FAccumulatedAuthorityErrors& Errors = *ErrorsAttribute.Get();
+		
+		if (Errors.NumTimeouts > 0)
+		{
+			AddTimeoutErrorWidget(Result, Errors);
+		}
+		if (!Errors.Rejected.IsEmpty())
+		{
+			AddRejectionErrorWidget(Result, Errors);
+		}
+		
+		SetErrorContent(Result);
+	}
+
+	void SAuthorityRejectedNotification::AddTimeoutErrorWidget(TSharedRef<SVerticalBox> Result, const FAccumulatedAuthorityErrors& Errors)
+	{
+		Result->AddSlot()
+			.AutoHeight()
+			[
+				SNew(STextBlock)
+				.Text(FText::Format(LOCTEXT("TimeoutsFmt", "{0} {0}|plural(one=Timeout,other=Timeouts)"), Errors.NumTimeouts))
+				.Font(FAppStyle::Get().GetFontStyle(TEXT("NotificationList.FontLight")))
+				.TextStyle(&FAppStyle::GetWidgetStyle<FTextBlockStyle>("NotificationList.WidgetText"))
+			];
+	}
+
+	void SAuthorityRejectedNotification::AddRejectionErrorWidget(TSharedRef<SVerticalBox> Result, const FAccumulatedAuthorityErrors& Errors)
+	{
+		Result->AddSlot()
+			.AutoHeight()
+			[
+				SNew(STextBlock)
+				.Text(FText::Format(LOCTEXT("RejectionsFmt", "{0} {0}|plural(one=Rejection,other=Rejections)"), Errors.Rejected.Num()))
+				.Font(FAppStyle::Get().GetFontStyle(TEXT("NotificationList.FontLight")))
+				.TextStyle(&FAppStyle::GetWidgetStyle<FTextBlockStyle>("NotificationList.WidgetText"))
+			];
+	}
+}
+
+#undef LOCTEXT_NAMESPACE
\ No newline at end of file
diff --git a/Engine/Plugins/Developer/Concert/ConcertApp/MultiUserClient/Source/MultiUserClient/Private/Replication/Submission/Notification/SAuthorityRejectedNotification.h b/Engine/Plugins/Developer/Concert/ConcertApp/MultiUserClient/Source/MultiUserClient/Private/Replication/Submission/Notification/SAuthorityRejectedNotification.h
new file mode 100644
--- /dev/null
+++ b/Engine/Plugins/Developer/Concert/ConcertApp/MultiUserClient/Source/MultiUserClient/Private/Replication/Submission/Notification/SAuthorityRejectedNotification.h
@@ -0,0 +1,43 @@
+// Copyright Epic Games, Inc. All Rights Reserved.
+
+#pragma once
+
+#include "SBaseRejectionNotification.h"
+
+namespace UE::MultiUserClient
+{
+	struct FAccumulatedAuthorityErrors;
+	
+	/** Displays FAccumulatedAuthorityErrors. */
+	class SAuthorityRejectedNotification : public SBaseRejectionNotification
+	{
+		using Super = SBaseRejectionNotification;
+	public:
+
+		SLATE_BEGIN_ARGS(SAuthorityRejectedNotification)
+		{}
+			/** The errors to display */
+			SLATE_ATTRIBUTE(const FAccumulatedAuthorityErrors*, Errors)
+			/** Called when the close button is pressed. Is supposed to close the notification. */
+			SLATE_EVENT(FOnClicked, OnCloseClicked)
+		SLATE_END_ARGS()
+
+		void Construct(const FArguments& InArgs);
+
+		/** Rebuilds the UI from the errors. */
+		void Refresh();
+
+		//~ Begin INotificationWidget interface
+		virtual void OnSetCompletionState(SNotificationItem::ECompletionState State) override {}
+		virtual TSharedRef<SWidget> AsWidget() override { return AsShared(); }
+		//~ End INotificationWidget interface
+
+	private:
+
+		TAttribute<const FAccumulatedAuthorityErrors*> ErrorsAttribute;
+		
+		void AddTimeoutErrorWidget(TSharedRef<SVerticalBox> Result, const FAccumulatedAuthorityErrors& Errors);
+		void AddRejectionErrorWidget(TSharedRef<SVerticalBox> Result, const FAccumulatedAuthorityErrors& Errors);
+	};
+}
+
diff --git a/Engine/Plugins/Developer/Concert/ConcertApp/MultiUserClient/Source/MultiUserClient/Private/Replication/Submission/Notification/SBaseRejectionNotification.cpp b/Engine/Plugins/Developer/Concert/ConcertApp/MultiUserClient/Source/MultiUserClient/Private/Replication/Submission/Notification/SBaseRejectionNotification.cpp
new file mode 100644
--- /dev/null
+++ b/Engine/Plugins/Developer/Concert/ConcertApp/MultiUserClient/Source/MultiUserClient/Private/Replication/Submission/Notification/SBaseRejectionNotification.cpp
@@ -0,0 +1,86 @@
+// Copyright Epic Games, Inc. All Rights Reserved.
+
+#include "SBaseRejectionNotification.h"
+
+#include "Styling/AppStyle.h"
+#include "Widgets/Layout/SBox.h"
+#include "Widgets/Images/SImage.h"
+#include "Widgets/Input/SButton.h"
+#include "Widgets/SBoxPanel.h"
+#include "Widgets/Text/STextBlock.h"
+
+#define LOCTEXT_NAMESPACE "SBaseRejectionNotification"
+
+namespace UE::MultiUserClient
+{
+	void SBaseRejectionNotification::Construct(const FArguments& InArgs)
+	{
+		ChildSlot
+		[
+			SNew(SHorizontalBox)
+			
+			+SHorizontalBox::Slot()
+			.AutoWidth()
+			.VAlign(VAlign_Top)
+			.HAlign(HAlign_Left)
+			[
+				SNew(SImage)
+				.Image(FAppStyle::Get().GetBrush("Icons.WarningWithColor.Large"))
+				.DesiredSizeOverride(FVector2D(32,32))
+			]
+			
+			+SHorizontalBox::Slot()
+			.Padding(10.f, 0.f, 0.f, 0.f)
+			[
+				BuildErrorContent(InArgs)
+			]
+		];
+
+		Refresh();
+	}
+
+	void SBaseRejectionNotification::SetErrorContent(const TSharedRef<SWidget>& Widget)
+	{
+		ErrorContent->SetContent(Widget);
+	}
+
+	TSharedRef<SWidget> SBaseRejectionNotification::BuildErrorContent(const FArguments& InArgs)
+	{
+		return SNew(SVerticalBox)
+			
+			+SVerticalBox::Slot()
+			.AutoHeight()
+			[
+				SNew(STextBlock)
+				.Text(InArgs._Message)
+				.Font(FAppStyle::Get().GetFontStyle(TEXT("NotificationList.FontBold")))
+				.TextStyle(&FAppStyle::GetWidgetStyle<FTextBlockStyle>("NotificationList.WidgetText"))
+			]
+
+			+SVerticalBox::Slot()
+			.Padding(FMargin(0.0f, 10.0f, 0.0f, 2.0f))
+			.AutoHeight()
+			[
+				SAssignNew(ErrorContent, SBox)
+			]
+			
+			+SVerticalBox::Slot()
+			.AutoHeight()
+			.Padding(0.0f, 5.0f, 0.0f, 0.0f)
+			.VAlign(VAlign_Center)
+			.HAlign(HAlign_Right)
+			[
+				SNew(SHorizontalBox)
+				+SHorizontalBox::Slot()
+				.AutoWidth()
+				[
+					SNew(SButton)
+					.Text(LOCTEXT("Close", "Close"))
+					.TextStyle(&FAppStyle::GetWidgetStyle<FTextBlockStyle>("NotificationList.WidgetText"))
+					.OnClicked(InArgs._OnCloseClicked) 
+				]
+			];
+	}
+}
+
+#undef LOCTEXT_NAMESPACE
\ No newline at end of file
diff --git a/Engine/Plugins/Developer/Concert/ConcertApp/MultiUserClient/Source/MultiUserClient/Private/Replication/Submission/Notification/SBaseRejectionNotification.h b/Engine/Plugins/Developer/Concert/ConcertApp/MultiUserClient/Source/MultiUserClient/Private/Replication/Submission/Notification/SBaseRejectionNotification.h
new file mode 100644
--- /dev/null
+++ b/Engine/Plugins/Developer/Concert/ConcertApp/MultiUserClient/Source/MultiUserClient/Private/Replication/Submission/Notification/SBaseRejectionNotification.h
@@ -0,0 +1,50 @@
+// Copyright Epic Games, Inc. All Rights Reserved.
+
+#pragma once
+
+#include "Widgets/DeclarativeSyntaxSupport.h"
+#include "Widgets/SCompoundWidget.h"
+#include "Widgets/Notifications/INotificationWidget.h"
+#include "Widgets/Notifications/SNotificationList.h"
+
+class SBox;
+	
+namespace UE::MultiUserClient
+{
+	/** Shared base class for rejection notifications. Builds the general UI structure. */
+	class SBaseRejectionNotification : public SCompoundWidget, public INotificationWidget
+	{
+	public:
+
+		SLATE_BEGIN_ARGS(SBaseRejectionNotification)
+		{}
+			/** The main message displayed at the top */
+			SLATE_ARGUMENT(FText, Message)
+			/** Called when the close button is pressed. Is supposed to close the notification. */
+			SLATE_EVENT(FOnClicked, OnCloseClicked)
+		SLATE_END_ARGS()
+
+		void Construct(const FArguments& InArgs);
+
+		/** Rebuilds the UI from the errors. */
+		virtual void Refresh() = 0;
+		
+		//~ Begin INotificationWidget interface
+		virtual void OnSetCompletionState(SNotificationItem::ECompletionState State) override {}
+		virtual TSharedRef<SWidget> AsWidget() override { return AsShared(); }
+		//~ End INotificationWidget interface
+
+	protected:
+
+		void SetErrorContent(const TSharedRef<SWidget>& Widget);
+		
+	private:
+
+		/** Subclasses place their content in this */
+		TSharedPtr<SBox> ErrorContent;
+
+		TSharedRef<SWidget> BuildErrorContent(const FArguments& InArgs);
+	};
+}
+
+
diff --git a/Engine/Plugins/Developer/Concert/ConcertApp/MultiUserClient/Source/MultiUserClient/Private/Replication/Submission/Notification/SStreamRejectedNotification.cpp b/Engine/Plugins/Developer/Concert/ConcertApp/MultiUserClient/Source/MultiUserClient/Private/Replication/Submission/Notification/SStreamRejectedNotification.cpp
new file mode 100644
--- /dev/null
+++ b/Engine/Plugins/Developer/Concert/ConcertApp/MultiUserClient/Source/MultiUserClient/Private/Replication/Submission/Notification/SStreamRejectedNotification.cpp
@@ -0,0 +1,101 @@
+// Copyright Epic Games, Inc. All Rights Reserved.
+
+#include "SStreamRejectedNotification.h"
+
+#include "AccumulatedSubmissionErrors.h"
+
+#include "Styling/AppStyle.h"
+#include "Widgets/Input/SButton.h"
+#include "Widgets/SBoxPanel.h"
+#include "Widgets/Text/STextBlock.h"
+
+#define LOCTEXT_NAMESPACE "SStreamRejectedNotification"
+
+namespace UE::MultiUserClient
+{
+	void SStreamRejectedNotification::Construct(const FArguments& InArgs)
+	{
+		ErrorsAttribute = InArgs._Errors;
+		
+		Super::Construct(
+			Super::FArguments()
+			.Message(LOCTEXT("MainMessage", "Errors updating stream"))
+			.OnCloseClicked(InArgs._OnCloseClicked)
+			);
+	}
+
+	void SStreamRejectedNotification::Refresh()
+	{
+		TSharedRef<SVerticalBox> Result = SNew(SVerticalBox);
+		const FAccumulatedStreamErrors& Errors = *ErrorsAttribute.Get();
+		
+		if (Errors.NumTimeouts > 0)
+		{
+			AddTimeoutErrorWidget(Result, Errors);
+		}
+		if (!Errors.AuthorityConflicts.IsEmpty())
+		{
+			AddAuthorityConflictErrorWidget(Result, Errors);
+		}
+		if (!Errors.SemanticErrors.IsEmpty())
+		{
+			AddSemanticErrorWidget(Result, Errors);
+		}
+		if (Errors.bFailedStreamCreation)
+		{
+			AddStreamCreationErrorWidget(Result, Errors);
+		}
+		
+		SetErrorContent(Result);
+	}
+
+	void SStreamRejectedNotification::AddTimeoutErrorWidget(TSharedRef<SVerticalBox> Result, const FAccumulatedStreamErrors& Errors)
+	{
+		Result->AddSlot()
+			.AutoHeight()
+			[
+				SNew(STextBlock)
+				.Text(FText::Format(LOCTEXT("TimeoutsFmt", "{0} {0}|plural(one=Timeout,other=Timeouts)"), Errors.NumTimeouts))
+				.Font(FAppStyle::Get().GetFontStyle(TEXT("NotificationList.FontLight")))
+				.TextStyle(&FAppStyle::GetWidgetStyle<FTextBlockStyle>("NotificationList.WidgetText"))
+			];
+	}
+
+	void SStreamRejectedNotification::AddAuthorityConflictErrorWidget(TSharedRef<SVerticalBox> Result, const FAccumulatedStreamErrors& Errors)
+	{
+		Result->AddSlot()
+			.AutoHeight()
+			[
+				SNew(STextBlock)
+				.Text(FText::Format(LOCTEXT("ConflictsFmt", "{0} {0}|plural(one=Conflict,other=Conflicts)"), Errors.AuthorityConflicts.Num()))
+				.Font(FAppStyle::Get().GetFontStyle(TEXT("NotificationList.FontLight")))
+				.TextStyle(&FAppStyle::GetWidgetStyle<FTextBlockStyle>("NotificationList.WidgetText"))
+			];
+	}
+	
+	void SStreamRejectedNotification::AddSemanticErrorWidget(TSharedRef<SVerticalBox> Result, const FAccumulatedStreamErrors& Errors)
+	{
+		Result->AddSlot()
+			.AutoHeight()
+			[
+				SNew(STextBlock)
+				.Text(FText::Format(LOCTEXT("SemanticFmt", "{0} {0}|plural(one=Semantic Error,other=Semantic Errors)"), Errors.SemanticErrors.Num()))
+				.Font(FAppStyle::Get().GetFontStyle(TEXT("NotificationList.FontLight")))
+				.TextStyle(&FAppStyle::GetWidgetStyle<FTextBlockStyle>("NotificationList.WidgetText"))
+			];
+	}
+	
+	void SStreamRejectedNotification::AddStreamCreationErrorWidget(TSharedRef<SVerticalBox> Result, const FAccumulatedStreamErrors& Errors)
+	{
+		Result->AddSlot()
+			.AutoHeight()
+			[
+				SNew(STextBlock)
+				.Text(LOCTEXT("RejectionsFmt", "Failed to create stream"))
+				.Font(FAppStyle::Get().GetFontStyle(TEXT("NotificationList.FontLight")))
+				.TextStyle(&FAppStyle::GetWidgetStyle<FTextBlockStyle>("NotificationList.WidgetText"))
+			];
+	}
+}
+
+#undef LOCTEXT_NAMESPACE
\ No newline at end of file
diff --git a/Engine/Plugins/Developer/Concert/ConcertApp/MultiUserClient/Source/MultiUserClient/Private/Replication/Submission/Notification/SStreamRejectedNotification.h b/Engine/Plugins/Developer/Concert/ConcertApp/MultiUserClient/Source/MultiUserClient/Private/Replication/Submission/Notification/SStreamRejectedNotification.h
new file mode 100644
--- /dev/null
+++ b/Engine/Plugins/Developer/Concert/ConcertApp/MultiUserClient/Source/MultiUserClient/Private/Replication/Submission/Notification/SStreamRejectedNotification.h
@@ -0,0 +1,45 @@
+// Copyright Epic Games, Inc. All Rights Reserved.
+
+#pragma once
+
+#include "SBaseRejectionNotification.h"
+
+namespace UE::MultiUserClient
+{
+	struct FAccumulatedStreamErrors;
+	
+	/** Displays FAccumulatedStreamErrors. */
+	class SStreamRejectedNotification : public SBaseRejectionNotification
+	{
+		using Super = SBaseRejectionNotification;
+	public:
+
+		SLATE_BEGIN_ARGS(SStreamRejectedNotification)
+		{}
+			/** The errors to display */
+			SLATE_ATTRIBUTE(const FAccumulatedStreamErrors*, Errors)
+			/** Called when the close button is pressed. Is supposed to close the notification. */
+			SLATE_EVENT(FOnClicked, OnCloseClicked)
+		SLATE_END_ARGS()
+
+		void Construct(const FArguments& InArgs);
+
+		/** Rebuilds the UI from the errors. */
+		void Refresh();
+		
+		//~ Begin INotificationWidget interface
+		virtual void OnSetCompletionState(SNotificationItem::ECompletionState State) override {}
+		virtual TSharedRef<SWidget> AsWidget() override { return AsShared(); }
+		//~ End INotificationWidget interface
+
+	private:
+
+		TAttribute<const FAccumulatedStreamErrors*> ErrorsAttribute;
+		
+		void AddTimeoutErrorWidget(TSharedRef<SVerticalBox> Result, const FAccumulatedStreamErrors& Errors);
+		void AddAuthorityConflictErrorWidget(TSharedRef<SVerticalBox> Result, const FAccumulatedStreamErrors& Errors);
+		void AddSemanticErrorWidget(TSharedRef<SVerticalBox> Result, const FAccumulatedStreamErrors& Errors);
+		void AddStreamCreationErrorWidget(TSharedRef<SVerticalBox> Result, const FAccumulatedStreamErrors& Errors);
+	};
+}
+
diff --git a/Engine/Plugins/Developer/Concert/ConcertApp/MultiUserClient/Source/MultiUserClient/Private/Replication/Submission/Notification/SubmissionNotifier.cpp b/Engine/Plugins/Developer/Concert/ConcertApp/MultiUserClient/Source/MultiUserClient/Private/Replication/Submission/Notification/SubmissionNotifier.cpp
new file mode 100644
--- /dev/null
+++ b/Engine/Plugins/Developer/Concert/ConcertApp/MultiUserClient/Source/MultiUserClient/Private/Replication/Submission/Notification/SubmissionNotifier.cpp
@@ -0,0 +1,193 @@
+// Copyright Epic Games, Inc. All Rights Reserved.
+
+#include "SubmissionNotifier.h"
+
+#include "SAuthorityRejectedNotification.h"
+#include "SStreamRejectedNotification.h"
+#include "Replication/Client/RemoteReplicationClient.h"
+#include "Replication/Client/ReplicationClientManager.h"
+#include "Replication/Submission/Data/AuthoritySubmission.h"
+#include "Replication/Submission/Data/StreamSubmission.h"
+
+#include "Framework/Notifications/NotificationManager.h"
+#include "Input/Reply.h"
+#include "Widgets/Notifications/SNotificationList.h"
+
+namespace UE::MultiUserClient
+{
+	FSubmissionNotifier::FSubmissionNotifier(FReplicationClientManager& InReplicationClientManager)
+		: ReplicationClientManager(InReplicationClientManager)
+	{
+		ReplicationClientManager.OnPostRemoteClientAdded().AddRaw(this, &FSubmissionNotifier::OnPostRemoteClientAdded);
+		RegisterClient(ReplicationClientManager.GetLocalClient());
+	}
+
+	FSubmissionNotifier::~FSubmissionNotifier()
+	{
+		CloseStreamNotification();
+		CloseAuthorityNotification();
+		
+		// FSubmissionNotifier is usually owned by ReplicationClientManager so this clean-up is not really needed.
+		// But we'll follow RAII here in case the ownership changes the future. 
+		ReplicationClientManager.OnPostRemoteClientAdded().RemoveAll(this);
+		ReplicationClientManager.OnPreRemoteClientRemoved().RemoveAll(this);
+		
+		UnregisterClient(ReplicationClientManager.GetLocalClient());
+		for (const TNonNullPtr<FRemoteReplicationClient>& RemoteClient : ReplicationClientManager.GetRemoteClients())
+		{
+			UnregisterClient(*RemoteClient);
+		}
+	}
+
+	void FSubmissionNotifier::OnPostRemoteClientAdded(FRemoteReplicationClient& RemoteReplicationClient)
+	{
+		RegisterClient(RemoteReplicationClient);
+	}
+
+	void FSubmissionNotifier::RegisterClient(FReplicationClient& Client)
+	{
+		Client.GetSubmissionWorkflow().OnStreamRequestCompleted().AddRaw(this, &FSubmissionNotifier::OnStreamRequestCompleted);
+		Client.GetSubmissionWorkflow().OnAuthorityRequestCompleted().AddRaw(this, &FSubmissionNotifier::OnAuthorityRequestCompleted);
+	}
+
+	void FSubmissionNotifier::UnregisterClient(FReplicationClient& Client)
+	{
+		Client.GetSubmissionWorkflow().OnStreamRequestCompleted().RemoveAll(this);
+		Client.GetSubmissionWorkflow().OnAuthorityRequestCompleted().RemoveAll(this);
+	}
+
+	void FSubmissionNotifier::OnStreamRequestCompleted(const FSubmitStreamChangesResponse& Request)
+	{
+		AccumulateStreamRejections(Request);
+		CreateOrUpdateStreamNotification();
+	}
+
+	void FSubmissionNotifier::OnAuthorityRequestCompleted(const FSubmitAuthorityChangesRequest& Request, const FSubmitAuthorityChangesResponse& Response)
+	{
+		AccumulateAuthorityRejections(Request, Response);
+		CreateOrUpdateAuthorityNotification();
+		
+	}
+
+	void FSubmissionNotifier::AccumulateStreamRejections(const FSubmitStreamChangesResponse& CompletedOp)
+	{
+		switch(CompletedOp.ErrorCode)
+		{
+		case EStreamSubmissionErrorCode::Success:
+			Algo::Transform(
+				CompletedOp.SubmissionInfo->Response.AuthorityConflicts, StreamErrors.AuthorityConflicts,
+				[](const TPair<FObjectInStreamID, FReplicatedObjectId>& Pair) { return Pair.Key.Object; }
+				);
+			Algo::Transform(
+				CompletedOp.SubmissionInfo->Response.ObjectsToPutSemanticErrors, StreamErrors.SemanticErrors,
+				[](const TPair<FObjectInStreamID, EConcertPutObjectErrorCode>& Pair) { return Pair.Key.Object; }
+				);
+			StreamErrors.bFailedStreamCreation |= !CompletedOp.SubmissionInfo->Response.FailedStreamCreation.IsEmpty();
+			break;
+			
+		case EStreamSubmissionErrorCode::Timeout:
+			++StreamErrors.NumTimeouts;
+			break;
+			
+		case EStreamSubmissionErrorCode::NoChange: break;
+		case EStreamSubmissionErrorCode::Cancelled: break;
+		default: ;
+		}
+	}
+
+	void FSubmissionNotifier::AccumulateAuthorityRejections(const FSubmitAuthorityChangesRequest& RequestOp, const FSubmitAuthorityChangesResponse& ResponseOp)
+	{
+		switch (ResponseOp.ErrorCode)
+		{
+		case EAuthoritySubmissionResponseErrorCode::Success:
+			Algo::Transform(
+				ResponseOp.Response->RejectedObjects, AuthorityErrors.Rejected,
+				[](const TPair<FSoftObjectPath, FConcertStreamArray>& Pair) { return Pair.Key; }
+				);
+			break;
+			
+		case EAuthoritySubmissionResponseErrorCode::Timeout:
+			++AuthorityErrors.NumTimeouts;
+			break;
+			
+		case EAuthoritySubmissionResponseErrorCode::NoChange: break;
+		case EAuthoritySubmissionResponseErrorCode::CancelledDueToStreamUpdate: break;
+		case EAuthoritySubmissionResponseErrorCode::Cancelled: break;
+		default: ;
+		}
+	}
+
+	void FSubmissionNotifier::CreateOrUpdateStreamNotification()
+	{
+		if (!HasStreamRejections())
+		{
+			return;
+		}
+
+		if (StreamRejectedNotification)
+		{
+			StreamRejectedNotification->Refresh();
+		}
+		else
+		{
+			StreamRejectedNotification = SNew(SStreamRejectedNotification)
+				.Errors_Lambda([this](){ return &StreamErrors; })
+				.OnCloseClicked_Raw(this, &FSubmissionNotifier::CloseStreamNotification);
+			FNotificationInfo Info(StreamRejectedNotification.ToSharedRef());
+			Info.bFireAndForget = false;
+			Info.ExpireDuration = 1.f;
+			StreamNotificationItem = FSlateNotificationManager::Get().AddNotification(Info);
+		}
+	}
+
+	void FSubmissionNotifier::CreateOrUpdateAuthorityNotification()
+	{
+		if (!HasAuthorityRejections())
+		{
+			return;
+		}
+
+		if (AuthorityRejectedNotification)
+		{
+			AuthorityRejectedNotification->Refresh();
+		}
+		else
+		{
+			AuthorityRejectedNotification = SNew(SAuthorityRejectedNotification)
+				.Errors_Lambda([this](){ return &AuthorityErrors; })
+				.OnCloseClicked_Raw(this, &FSubmissionNotifier::CloseAuthorityNotification);
+			FNotificationInfo Info(AuthorityRejectedNotification.ToSharedRef());
+			Info.bFireAndForget = false;
+			Info.ExpireDuration = 1.f;
+			AuthorityNotificationItem = FSlateNotificationManager::Get().AddNotification(Info);
+		}
+	}
+
+	FReply FSubmissionNotifier::CloseStreamNotification()
+	{
+		if (StreamNotificationItem)
+		{
+			StreamRejectedNotification.Reset();
+			StreamNotificationItem->ExpireAndFadeout();
+			StreamNotificationItem.Reset();
+
+			StreamErrors = {};
+		}
+		
+		return FReply::Handled();
+	}
+
+	FReply FSubmissionNotifier::CloseAuthorityNotification()
+	{
+		if (AuthorityNotificationItem)
+		{
+			AuthorityRejectedNotification.Reset();
+			AuthorityNotificationItem->ExpireAndFadeout();
+			AuthorityNotificationItem.Reset();
+
+			AuthorityErrors = {};
+		}
+
+		return FReply::Handled();
+	}
+}
diff --git a/Engine/Plugins/Developer/Concert/ConcertApp/MultiUserClient/Source/MultiUserClient/Private/Replication/Submission/Notification/SubmissionNotifier.h b/Engine/Plugins/Developer/Concert/ConcertApp/MultiUserClient/Source/MultiUserClient/Private/Replication/Submission/Notification/SubmissionNotifier.h
new file mode 100644
--- /dev/null
+++ b/Engine/Plugins/Developer/Concert/ConcertApp/MultiUserClient/Source/MultiUserClient/Private/Replication/Submission/Notification/SubmissionNotifier.h
@@ -0,0 +1,79 @@
+// Copyright Epic Games, Inc. All Rights Reserved.
+
+#pragma once
+
+#include "AccumulatedSubmissionErrors.h"
+#include "Templates/UnrealTemplate.h"
+
+class SNotificationItem;
+class FReply;
+
+namespace UE::MultiUserClient
+{
+	class FRemoteReplicationClient;
+	class FReplicationClient;
+	class FReplicationClientManager;
+	
+	class SAuthorityRejectedNotification;
+	class SStreamRejectedNotification;
+	
+	struct FSubmitAuthorityChangesResponse;
+	struct FSubmitAuthorityChangesRequest;
+	struct FSubmitStreamChangesResponse;
+	
+	/**
+	 * Adds a SNotificationItem when ISubmissionWorkflow operations fail.
+	 * Failing operations are bundled into a single notification item until it is manually dismissed.
+	 */
+	class FSubmissionNotifier : public FNoncopyable
+	{
+	public:
+		
+		FSubmissionNotifier(FReplicationClientManager& InReplicationClientManager);
+		~FSubmissionNotifier();
+
+	private:
+
+		/** Emits events when remote clients are discovered. */
+		FReplicationClientManager& ReplicationClientManager;
+
+		/** Only valid if changing streams has failed and the user has not yet dismissed the notification. */
+		TSharedPtr<SNotificationItem> StreamNotificationItem;
+		/** Only valid if changing authority has failed and the user has not yet dismissed the notification. */
+		TSharedPtr<SNotificationItem> AuthorityNotificationItem;
+		
+		/** Wrapped by StreamNotification */
+		TSharedPtr<SStreamRejectedNotification> StreamRejectedNotification;
+		/** Wrapped by AuthorityNotification */
+		TSharedPtr<SAuthorityRejectedNotification> AuthorityRejectedNotification;
+
+		/** Displayed by StreamRejectedNotification */
+		FAccumulatedStreamErrors StreamErrors;
+		/** Displayed by AuthorityRejectedNotification */
+		FAccumulatedAuthorityErrors AuthorityErrors;
+		
+		// Client events
+		void OnPostRemoteClientAdded(FRemoteReplicationClient& RemoteReplicationClient);
+		void RegisterClient(FReplicationClient& Client);
+		void UnregisterClient(FReplicationClient& Client);
+
+		// Per-client submission events
+		void OnStreamRequestCompleted(const FSubmitStreamChangesResponse& Request);
+		void OnAuthorityRequestCompleted(const FSubmitAuthorityChangesRequest& Request, const FSubmitAuthorityChangesResponse& Response);
+
+		// Handle rejections
+		void AccumulateStreamRejections(const FSubmitStreamChangesResponse& CompletedOp);
+		void AccumulateAuthorityRejections(const FSubmitAuthorityChangesRequest& RequestOp, const FSubmitAuthorityChangesResponse& ResponseOp);
+		bool HasStreamRejections() const { return StreamErrors.NumTimeouts > 0 || !StreamErrors.AuthorityConflicts.IsEmpty() || !StreamErrors.SemanticErrors.IsEmpty() || StreamErrors.bFailedStreamCreation; }
+		bool HasAuthorityRejections() const { return AuthorityErrors.NumTimeouts > 0 || !AuthorityErrors.Rejected.IsEmpty(); }
+
+		// Updating the notifications
+		void CreateOrUpdateStreamNotification();
+		void CreateOrUpdateAuthorityNotification();
+
+		// Button events
+		FReply CloseStreamNotification();
+		FReply CloseAuthorityNotification();
+	};
+}
+
diff --git a/Engine/Plugins/Developer/Concert/ConcertApp/MultiUserClient/Source/MultiUserClient/Private/Replication/Submission/SingleClientSubmissionOperation.h b/Engine/Plugins/Developer/Concert/ConcertApp/MultiUserClient/Source/MultiUserClient/Private/Replication/Submission/SingleClientSubmissionOperation.h
--- a/Engine/Plugins/Developer/Concert/ConcertApp/MultiUserClient/Source/MultiUserClient/Private/Replication/Submission/SingleClientSubmissionOperation.h
+++ b/Engine/Plugins/Developer/Concert/ConcertApp/MultiUserClient/Source/MultiUserClient/Private/Replication/Submission/SingleClientSubmissionOperation.h
@@ -36,17 +36,17 @@ namespace UE::MultiUserClient
 			}
 		}
 
-		void EmplaceStreamPromise(FSubmitStreamChangesResponse&& Result)
+		void EmplaceStreamPromise(FSubmitStreamChangesResponse Result)
 		{
 			bStreamPromiseWasSet = true;
 			StreamChangesPromise.EmplaceValue(MoveTemp(Result));
 		}
-		void EmplaceAuthorityRequestPromise(FSubmitAuthorityChangesRequest&& Result)
+		void EmplaceAuthorityRequestPromise(FSubmitAuthorityChangesRequest Result)
 		{
 			bAuthorityRequestPromiseWasSet = true;
 			AuthorityChangeRequestPromise.EmplaceValue(MoveTemp(Result));
 		}
-		void EmplaceAuthorityResponsePromise(FSubmitAuthorityChangesResponse&& Result)
+		void EmplaceAuthorityResponsePromise(FSubmitAuthorityChangesResponse Result)
 		{
 			bAuthorityResponsePromiseWasSet = true;
 			AuthorityChangeResponsePromise.EmplaceValue(MoveTemp(Result));
diff --git a/Engine/Plugins/Developer/Concert/ConcertApp/MultiUserClient/Source/MultiUserClient/Private/Replication/Submission/SubmissionWorkflow_LocalClient.cpp b/Engine/Plugins/Developer/Concert/ConcertApp/MultiUserClient/Source/MultiUserClient/Private/Replication/Submission/SubmissionWorkflow_LocalClient.cpp
--- a/Engine/Plugins/Developer/Concert/ConcertApp/MultiUserClient/Source/MultiUserClient/Private/Replication/Submission/SubmissionWorkflow_LocalClient.cpp
+++ b/Engine/Plugins/Developer/Concert/ConcertApp/MultiUserClient/Source/MultiUserClient/Private/Replication/Submission/SubmissionWorkflow_LocalClient.cpp
@@ -21,12 +21,6 @@ namespace UE::MultiUserClient
 		, StreamSynchronizer(InStreamSynchronizer)
 	{}
 
-	void FSubmissionWorkflow_LocalClient::RevertChanges()
-	{
-		StreamChangeTracker.RevertCachedChanges();
-		AuthorityChangeTracker.ClearChanges();
-	}
-
 	EChangeUploadability FSubmissionWorkflow_LocalClient::GetUploadability() const
 	{
 		const bool bOperationInProgress = InProgressOperation.IsSet(); 
@@ -74,7 +68,10 @@ namespace UE::MultiUserClient
 		
 		if (bIsChangelistEmpty)
 		{
-			Operation->EmplaceStreamPromise(FSubmitStreamChangesResponse{ EStreamSubmissionErrorCode::NoChange });
+			const FSubmitStreamChangesResponse CompletedChange { EStreamSubmissionErrorCode::NoChange };
+			Operation->EmplaceStreamPromise(CompletedChange);
+			StreamRequestCompletedDelegate.Broadcast(CompletedChange);
+			
 			SendAuthorityChangeRequest(MoveTemp(AuthorityChangeRequest));
 		}
 		else
@@ -103,7 +100,7 @@ namespace UE::MultiUserClient
 	}
 
 	void FSubmissionWorkflow_LocalClient::OnStreamChangeCompleted(
-		const ConcertSyncClient::Replication::FChangeStreamRequest& Request,
+		const ConcertSyncClient::Replication::FChangeStreamRequest& StreamChangeRequest,
 		const ConcertSyncClient::Replication::FChangeStreamResponse& ChangeStreamResponse,
 		ConcertSyncClient::Replication::FAuthorityChangeRequest AuthorityChangeRequest
 		)
@@ -113,10 +110,9 @@ namespace UE::MultiUserClient
 		const EStreamSubmissionErrorCode ErrorCode = ChangeStreamResponse.ErrorCode == EReplicationResponseErrorCode::Handled
 			? EStreamSubmissionErrorCode::Success
 			: EStreamSubmissionErrorCode::Timeout;
-		Operation->EmplaceStreamPromise(FSubmitStreamChangesResponse{
-			ErrorCode,
-			{ FCompletedChangeSubmission{ Request, ChangeStreamResponse } }
-		});
+		const FSubmitStreamChangesResponse CompletedChange{ ErrorCode, { FCompletedChangeSubmission{ StreamChangeRequest, ChangeStreamResponse }} };
+		Operation->EmplaceStreamPromise(CompletedChange);
+		StreamRequestCompletedDelegate.Broadcast(CompletedChange);
 		
 		if (ChangeStreamResponse.IsSuccess())
 		{
@@ -124,9 +120,14 @@ namespace UE::MultiUserClient
 		}
 		else
 		{
-			Operation->EmplaceAuthorityRequestPromise(FSubmitAuthorityChangesRequest{ EAuthoritySubmissionRequestErrorCode::CancelledDueToStreamUpdate });
-			Operation->EmplaceAuthorityResponsePromise(FSubmitAuthorityChangesResponse{ EAuthoritySubmissionResponseErrorCode::CancelledDueToStreamUpdate });
+			const FSubmitAuthorityChangesRequest Request{ EAuthoritySubmissionRequestErrorCode::CancelledDueToStreamUpdate };
+			const FSubmitAuthorityChangesResponse Response{ EAuthoritySubmissionResponseErrorCode::CancelledDueToStreamUpdate };
+			
+			Operation->EmplaceAuthorityRequestPromise(Request);
+			Operation->EmplaceAuthorityResponsePromise(Response);
 			InProgressOperation.Reset();
+			
+			AuthorityRequestCompletedDelegate.Broadcast(Request, Response);
 		}
 	}
 
@@ -137,6 +138,7 @@ namespace UE::MultiUserClient
 		if (!ensure(ReplicationManager))
 		{
 			InProgressOperation.Reset(); // Automatically cancels the pending promises
+			AuthorityRequestCompletedDelegate.Broadcast({ EAuthoritySubmissionRequestErrorCode::Cancelled }, { EAuthoritySubmissionResponseErrorCode::Cancelled });
 			return;
 		}
 
@@ -144,23 +146,33 @@ namespace UE::MultiUserClient
 		const bool bHasNoChanges = AuthorityChangeRequest.ReleaseAuthority.IsEmpty() && AuthorityChangeRequest.TakeAuthority.IsEmpty();
 		if (bHasNoChanges)
 		{
-			Operation->EmplaceAuthorityRequestPromise(FSubmitAuthorityChangesRequest{ EAuthoritySubmissionRequestErrorCode::NoChange });
-			Operation->EmplaceAuthorityResponsePromise(FSubmitAuthorityChangesResponse{ EAuthoritySubmissionResponseErrorCode::NoChange });
+			const FSubmitAuthorityChangesRequest Request{ EAuthoritySubmissionRequestErrorCode::NoChange };
+			const FSubmitAuthorityChangesResponse Response{ EAuthoritySubmissionResponseErrorCode::NoChange };
+			
+			Operation->EmplaceAuthorityRequestPromise(Request);
+			Operation->EmplaceAuthorityResponsePromise(Response);
 			InProgressOperation.Reset();
+			
+			AuthorityRequestCompletedDelegate.Broadcast(Request, Response);
 			return;
 		}
-		
-		Operation->EmplaceAuthorityRequestPromise(FSubmitAuthorityChangesRequest{ EAuthoritySubmissionRequestErrorCode::Success, AuthorityChangeRequest });
+
+		FSubmitAuthorityChangesRequest Request{ EAuthoritySubmissionRequestErrorCode::Success, AuthorityChangeRequest };
+		Operation->EmplaceAuthorityRequestPromise(Request);
 		ReplicationManager->RequestAuthorityChange(MoveTemp(AuthorityChangeRequest))
-			.Next([this, DestructionDetection = LifetimeToken->AsWeak()](FAuthorityChangeResponse&& Response)
+			.Next([this, Request = MoveTemp(Request), DestructionDetection = LifetimeToken->AsWeak()](FAuthorityChangeResponse&& Response)
 			{
 				if (DestructionDetection.IsValid())
 				{
 					const EAuthoritySubmissionResponseErrorCode ErrorCode = Response.ErrorCode == EReplicationResponseErrorCode::Handled
 						? EAuthoritySubmissionResponseErrorCode::Success
 						: EAuthoritySubmissionResponseErrorCode::Timeout;
-					InProgressOperation->Get().EmplaceAuthorityResponsePromise(FSubmitAuthorityChangesResponse{ ErrorCode, MoveTemp(Response) });
+					const FSubmitAuthorityChangesResponse Result { ErrorCode, MoveTemp(Response) };
+					
+					InProgressOperation->Get().EmplaceAuthorityResponsePromise(Result);
 					InProgressOperation.Reset();
+					
+					AuthorityRequestCompletedDelegate.Broadcast(Request, Result);
 				}
 			});
 	}
diff --git a/Engine/Plugins/Developer/Concert/ConcertApp/MultiUserClient/Source/MultiUserClient/Private/Replication/Submission/SubmissionWorkflow_LocalClient.h b/Engine/Plugins/Developer/Concert/ConcertApp/MultiUserClient/Source/MultiUserClient/Private/Replication/Submission/SubmissionWorkflow_LocalClient.h
--- a/Engine/Plugins/Developer/Concert/ConcertApp/MultiUserClient/Source/MultiUserClient/Private/Replication/Submission/SubmissionWorkflow_LocalClient.h
+++ b/Engine/Plugins/Developer/Concert/ConcertApp/MultiUserClient/Source/MultiUserClient/Private/Replication/Submission/SubmissionWorkflow_LocalClient.h
@@ -3,10 +3,13 @@
 #pragma once
 
 #include "ISubmissionWorkflow.h"
+
 #include "Replication/Stream/IClientStreamSynchronizer.h"
 #include "Replication/IToken.h"
 #include "SingleClientSubmissionOperation.h"
 
+#include "Templates/UnrealTemplate.h"
+
 class IConcertSyncClient;
 
 namespace UE::MultiUserClient
@@ -17,7 +20,7 @@ namespace UE::MultiUserClient
 	class FStreamChangeTracker;
 	
 	/** Handles the submission workflow for a single client. */
-	class FSubmissionWorkflow_LocalClient : public ISubmissionWorkflow
+	class FSubmissionWorkflow_LocalClient : public FSubmissionWorkflowBase, public FNoncopyable
 	{
 	public:
 		
@@ -30,7 +33,6 @@ namespace UE::MultiUserClient
 		
 		//~ Begin ISubmissionWorkflow Interface
 		virtual TSharedPtr<ISubmissionOperation> SubmitChanges() override;
-		virtual void RevertChanges() override;
 		virtual EChangeUploadability GetUploadability() const override;
 		virtual EChangeRevertability GetRevertability() const override;
 		//~ End ISubmissionWorkflow Interface
@@ -61,7 +63,7 @@ namespace UE::MultiUserClient
 
 		/** Advances the request by requesting authority. */
 		void OnStreamChangeCompleted(
-			const ConcertSyncClient::Replication::FChangeStreamRequest& Request,
+			const ConcertSyncClient::Replication::FChangeStreamRequest& StreamChangeRequest,
 			const ConcertSyncClient::Replication::FChangeStreamResponse& ChangeStreamResponse,
 			ConcertSyncClient::Replication::FAuthorityChangeRequest AuthorityChangeRequest
 			);
diff --git a/Engine/Plugins/Developer/Concert/ConcertApp/MultiUserClient/Source/MultiUserClient/Private/Replication/Submission/SubmissionWorkflow_RemoteClient.cpp b/Engine/Plugins/Developer/Concert/ConcertApp/MultiUserClient/Source/MultiUserClient/Private/Replication/Submission/SubmissionWorkflow_RemoteClient.cpp
--- a/Engine/Plugins/Developer/Concert/ConcertApp/MultiUserClient/Source/MultiUserClient/Private/Replication/Submission/SubmissionWorkflow_RemoteClient.cpp
+++ b/Engine/Plugins/Developer/Concert/ConcertApp/MultiUserClient/Source/MultiUserClient/Private/Replication/Submission/SubmissionWorkflow_RemoteClient.cpp
@@ -11,9 +11,6 @@ namespace UE::MultiUserClient
 		return nullptr;
 	}
 
-	void FSubmissionWorkflow_RemoteClient::RevertChanges()
-	{}
-
 	EChangeUploadability FSubmissionWorkflow_RemoteClient::GetUploadability() const
 	{
 		// TODO DP UE-180657: Changing remote client's stream is not implemented for now
diff --git a/Engine/Plugins/Developer/Concert/ConcertApp/MultiUserClient/Source/MultiUserClient/Private/Replication/Submission/SubmissionWorkflow_RemoteClient.h b/Engine/Plugins/Developer/Concert/ConcertApp/MultiUserClient/Source/MultiUserClient/Private/Replication/Submission/SubmissionWorkflow_RemoteClient.h
--- a/Engine/Plugins/Developer/Concert/ConcertApp/MultiUserClient/Source/MultiUserClient/Private/Replication/Submission/SubmissionWorkflow_RemoteClient.h
+++ b/Engine/Plugins/Developer/Concert/ConcertApp/MultiUserClient/Source/MultiUserClient/Private/Replication/Submission/SubmissionWorkflow_RemoteClient.h
@@ -3,6 +3,7 @@
 #pragma once
 
 #include "ISubmissionWorkflow.h"
+
 #include "Templates/UnrealTemplate.h"
 
 class IConcertSyncClient;
@@ -10,14 +11,13 @@ class IConcertSyncClient;
 namespace UE::MultiUserClient
 {
 	class FSubmissionWorkflow_RemoteClient
-		: public ISubmissionWorkflow
+		: public FSubmissionWorkflowBase
 		, public FNoncopyable
 	{
 	public:
 
 		//~ Begin ISubmissionWorkflow Interface
 		virtual TSharedPtr<ISubmissionOperation> SubmitChanges() override;
-		virtual void RevertChanges() override;
 		virtual EChangeUploadability GetUploadability() const override;
 		virtual EChangeRevertability GetRevertability() const override;
 		//~ End ISubmissionWorkflow Interface
diff --git a/Engine/Plugins/Developer/Concert/ConcertApp/MultiUserClient/Source/MultiUserClient/Private/Widgets/ActiveSession/Replication/Client/SSubmissionControls.cpp b/Engine/Plugins/Developer/Concert/ConcertApp/MultiUserClient/Source/MultiUserClient/Private/Widgets/ActiveSession/Replication/Client/SSubmissionControls.cpp
deleted file mode 100644
diff --git a/Engine/Plugins/Developer/Concert/ConcertApp/MultiUserClient/Source/MultiUserClient/Private/Widgets/ActiveSession/Replication/Client/SSubmissionControls.h b/Engine/Plugins/Developer/Concert/ConcertApp/MultiUserClient/Source/MultiUserClient/Private/Widgets/ActiveSession/Replication/Client/SSubmissionControls.h
deleted file mode 100644
diff --git a/Engine/Plugins/Developer/Concert/ConcertApp/MultiUserClient/Source/MultiUserClient/Private/Widgets/ActiveSession/Replication/Client/Single/SReplicationClientView.cpp b/Engine/Plugins/Developer/Concert/ConcertApp/MultiUserClient/Source/MultiUserClient/Private/Widgets/ActiveSession/Replication/Client/Single/SReplicationClientView.cpp
--- a/Engine/Plugins/Developer/Concert/ConcertApp/MultiUserClient/Source/MultiUserClient/Private/Widgets/ActiveSession/Replication/Client/Single/SReplicationClientView.cpp
+++ b/Engine/Plugins/Developer/Concert/ConcertApp/MultiUserClient/Source/MultiUserClient/Private/Widgets/ActiveSession/Replication/Client/Single/SReplicationClientView.cpp
@@ -56,7 +56,6 @@ namespace UE::MultiUserClient
 			.Padding(2.f)
 			[
 				SNew(SSingleClientToolbar)
-				.SubmissionWorkflow(&ReplicationClient->GetSubmissionWorkflow())
 				.ViewSelectionArea() [ InArgs._ViewSelectionArea.Widget ]
 			]
 
diff --git a/Engine/Plugins/Developer/Concert/ConcertApp/MultiUserClient/Source/MultiUserClient/Private/Widgets/ActiveSession/Replication/Client/Single/SSingleClientToolbar.cpp b/Engine/Plugins/Developer/Concert/ConcertApp/MultiUserClient/Source/MultiUserClient/Private/Widgets/ActiveSession/Replication/Client/Single/SSingleClientToolbar.cpp
--- a/Engine/Plugins/Developer/Concert/ConcertApp/MultiUserClient/Source/MultiUserClient/Private/Widgets/ActiveSession/Replication/Client/Single/SSingleClientToolbar.cpp
+++ b/Engine/Plugins/Developer/Concert/ConcertApp/MultiUserClient/Source/MultiUserClient/Private/Widgets/ActiveSession/Replication/Client/Single/SSingleClientToolbar.cpp
@@ -2,8 +2,6 @@
 
 #include "SSingleClientToolbar.h"
 
-#include "Widgets/ActiveSession/Replication/Client/SSubmissionControls.h"
-
 #include "Widgets/Layout/SSpacer.h"
 #include "Widgets/SBoxPanel.h"
 
@@ -21,20 +19,6 @@ namespace UE::MultiUserClient
 			[
 				InArgs._ViewSelectionArea.Widget
 			]
-
-			+SHorizontalBox::Slot()
-			.FillWidth(1.f)
-			[
-				SNew(SSpacer)
-			]
-
-			// Submission controls
-			+SHorizontalBox::Slot()
-			.AutoWidth()
-			[
-				SNew(SSubmissionControls)
-				.SubmissionWorkflow(InArgs._SubmissionWorkflow)
-			]
 		];
 	}
 }
\ No newline at end of file
diff --git a/Engine/Plugins/Developer/Concert/ConcertApp/MultiUserClient/Source/MultiUserClient/Private/Widgets/ActiveSession/Replication/Client/Single/SSingleClientToolbar.h b/Engine/Plugins/Developer/Concert/ConcertApp/MultiUserClient/Source/MultiUserClient/Private/Widgets/ActiveSession/Replication/Client/Single/SSingleClientToolbar.h
--- a/Engine/Plugins/Developer/Concert/ConcertApp/MultiUserClient/Source/MultiUserClient/Private/Widgets/ActiveSession/Replication/Client/Single/SSingleClientToolbar.h
+++ b/Engine/Plugins/Developer/Concert/ConcertApp/MultiUserClient/Source/MultiUserClient/Private/Widgets/ActiveSession/Replication/Client/Single/SSingleClientToolbar.h
@@ -12,9 +12,6 @@ namespace UE::MultiUserClient
 
 	/**
 	 * Contains a bunch of actions that can be performed on client view.
-	 *
-	 * Layout:
-	 * | AdditionalToolbarWidgets                   | Submit | Revert |
 	 */
 	class SSingleClientToolbar : public SCompoundWidget
 	{
@@ -22,9 +19,6 @@ namespace UE::MultiUserClient
 		
 		SLATE_BEGIN_ARGS(SSingleClientToolbar)
 		{}
-			/** Manages submission and reverting of changes. Used for upload and revert buttons. */
-			SLATE_ATTRIBUTE(ISubmissionWorkflow*, SubmissionWorkflow)
-
 			/** Dedicated space for a widget with which to change the view. */
 			SLATE_NAMED_SLOT(FArguments, ViewSelectionArea)
 		SLATE_END_ARGS()
diff --git a/Engine/Plugins/Developer/Concert/ConcertApp/MultiUserClient/Source/MultiUserClient/Private/Widgets/ActiveSession/Replication/Client/Single/StreamEditorColumns.cpp b/Engine/Plugins/Developer/Concert/ConcertApp/MultiUserClient/Source/MultiUserClient/Private/Widgets/ActiveSession/Replication/Client/Single/StreamEditorColumns.cpp
--- a/Engine/Plugins/Developer/Concert/ConcertApp/MultiUserClient/Source/MultiUserClient/Private/Widgets/ActiveSession/Replication/Client/Single/StreamEditorColumns.cpp
+++ b/Engine/Plugins/Developer/Concert/ConcertApp/MultiUserClient/Source/MultiUserClient/Private/Widgets/ActiveSession/Replication/Client/Single/StreamEditorColumns.cpp
@@ -76,12 +76,14 @@ namespace UE::MultiUserClient::StreamEditorColumns
 			)
 		{
 			const FSoftObjectPath& ObjectPath = ObjectData.GetObjectPath();
-			ChangeTracker->SetAuthorityIfAllowed(ObjectPath, bIsChecked);
-			ClientStreamModel->ForEachSubobject(ObjectPath, [&ChangeTracker, bIsChecked](const FSoftObjectPath& Child)
+			TArray<FSoftObjectPath> Paths { ObjectPath };
+			ClientStreamModel->ForEachSubobject(ObjectPath, [&Paths](const FSoftObjectPath& Child)
 			{
-				ChangeTracker->SetAuthorityIfAllowed(Child, bIsChecked);
+				Paths.Add(Child);
 				return EBreakBehavior::Continue;
 			});
+			
+			ChangeTracker->SetAuthorityIfAllowed(Paths, bIsChecked);
 		}
 
 		static bool IsEnabled(
@@ -175,7 +177,7 @@ namespace UE::MultiUserClient::StreamEditorColumns
 					FSubobjectColumnDelegates::FOnColumnCheckboxChanged::CreateLambda(
 					[&ChangeTracker](bool bIsChecked, const FReplicatedObjectData& ObjectData)
 					{
-						ChangeTracker.SetAuthorityIfAllowed(ObjectData.GetObjectPath(), bIsChecked);
+						ChangeTracker.SetAuthorityIfAllowed({ ObjectData.GetObjectPath() }, bIsChecked);
 					}),
 					FSubobjectColumnDelegates::FGetToolTipText::CreateLambda([&ChangeTracker, &SubmissionWorkflow](const FReplicatedObjectData& ObjectData)
 					{
diff --git a/Engine/Plugins/Developer/Concert/ConcertSync/ConcertSyncClient/Source/ConcertSyncClient/Private/Replication/Manager/ReplicationManagerState_Connected.cpp b/Engine/Plugins/Developer/Concert/ConcertSync/ConcertSyncClient/Source/ConcertSyncClient/Private/Replication/Manager/ReplicationManagerState_Connected.cpp
--- a/Engine/Plugins/Developer/Concert/ConcertSync/ConcertSyncClient/Source/ConcertSyncClient/Private/Replication/Manager/ReplicationManagerState_Connected.cpp
+++ b/Engine/Plugins/Developer/Concert/ConcertSync/ConcertSyncClient/Source/ConcertSyncClient/Private/Replication/Manager/ReplicationManagerState_Connected.cpp
@@ -36,6 +36,12 @@ namespace UE::ConcertSyncClient::Replication
 		false,
 		TEXT("Whether the client should pretend that stream change requests timed out instead of sending to the server.")
 		);
+
+	TAutoConsoleVariable<bool> CVarSimulateAuthorityRejection(
+		TEXT("Concert.Replication.SimulateAuthorityRejection"),
+		false,
+		TEXT("Whether the client should pretend that authority change requests were rejected.")
+		);
 	
 	FReplicationManagerState_Connected::FReplicationManagerState_Connected(
 		TSharedRef<IConcertClientSession> LiveSession,
@@ -102,6 +108,10 @@ namespace UE::ConcertSyncClient::Replication
 		{
 			return MakeFulfilledPromise<FAuthorityChangeResponse>(FAuthorityChangeResponse{ EReplicationResponseErrorCode::Timeout }).GetFuture();
 		}
+		if (CVarSimulateAuthorityRejection.GetValueOnGameThread())
+		{
+			return RejectAll(MoveTemp(Args));
+		}
 		
 		// Stop replicating removed objects right now: the server will remove authority after processing this request.
 		// At that point, it will log errors for receiving replication data from a client without authority.
diff --git a/Engine/Plugins/Developer/Concert/ConcertSync/ConcertSyncClient/Source/ConcertSyncClient/Private/Replication/Manager/ReplicationManagerUtils.cpp b/Engine/Plugins/Developer/Concert/ConcertSync/ConcertSyncClient/Source/ConcertSyncClient/Private/Replication/Manager/ReplicationManagerUtils.cpp
--- a/Engine/Plugins/Developer/Concert/ConcertSync/ConcertSyncClient/Source/ConcertSyncClient/Private/Replication/Manager/ReplicationManagerUtils.cpp
+++ b/Engine/Plugins/Developer/Concert/ConcertSync/ConcertSyncClient/Source/ConcertSyncClient/Private/Replication/Manager/ReplicationManagerUtils.cpp
@@ -12,7 +12,7 @@ namespace UE::ConcertSyncClient::Replication
 	TFuture<FAuthorityChangeResponse> RejectAll(FAuthorityChangeRequest&& Args)
 	{
 		return MakeFulfilledPromise<FAuthorityChangeResponse>(
-			FAuthorityChangeResponse{{ .RejectedObjects = MoveTemp(Args.TakeAuthority) }}
+			FAuthorityChangeResponse{{ EReplicationResponseErrorCode::Handled, MoveTemp(Args.TakeAuthority) }}
 			).GetFuture();
 	}
 }
diff --git a/Engine/Plugins/Developer/Concert/ConcertUI/ConcertClientSharedSlate/Source/ConcertClientSharedSlate/Private/Replication/Editor/Model/DisplayUtils.cpp b/Engine/Plugins/Developer/Concert/ConcertUI/ConcertClientSharedSlate/Source/ConcertClientSharedSlate/Private/Replication/Editor/Model/DisplayUtils.cpp
deleted file mode 100644
diff --git a/Engine/Plugins/Developer/Concert/ConcertUI/ConcertClientSharedSlate/Source/ConcertClientSharedSlate/Private/Replication/Editor/Model/DisplayUtils.h b/Engine/Plugins/Developer/Concert/ConcertUI/ConcertClientSharedSlate/Source/ConcertClientSharedSlate/Private/Replication/Editor/Model/DisplayUtils.h
deleted file mode 100644
diff --git a/Engine/Plugins/Developer/Concert/ConcertUI/ConcertClientSharedSlate/Source/ConcertClientSharedSlate/Private/Replication/Editor/View/DisplayUtils.cpp b/Engine/Plugins/Developer/Concert/ConcertUI/ConcertClientSharedSlate/Source/ConcertClientSharedSlate/Private/Replication/Editor/View/DisplayUtils.cpp
new file mode 100644
--- /dev/null
+++ b/Engine/Plugins/Developer/Concert/ConcertUI/ConcertClientSharedSlate/Source/ConcertClientSharedSlate/Private/Replication/Editor/View/DisplayUtils.cpp
@@ -0,0 +1,104 @@
+// Copyright Epic Games, Inc. All Rights Reserved.
+
+#include "Replication/Editor/View/DisplayUtils.h"
+
+#include "Replication/Editor/Model/IObjectToPropertiesModel.h"
+
+#include "AssetRegistry/AssetData.h"
+#include "AssetRegistry/AssetRegistryModule.h"
+#include "Blueprint/BlueprintSupport.h"
+#include "ConcertLogGlobal.h"
+#include "GameFramework/Actor.h"
+#include "Styling/SlateIconFinder.h"
+
+namespace UE::ConcertClientSharedSlate::DisplayUtils
+{
+	FText GetObjectDisplayText(const FSoftObjectPath& Object)
+	{
+		return FText::FromString(GetObjectDisplayString(Object));
+	}
+
+	FString GetObjectDisplayString(const FSoftObjectPath& Object)
+	{
+		// Important! The object may not be loaded, yet. This could be if the asset is using a level that was not opened. 
+		if (const UObject* LoadedObject = Object.ResolveObject())
+		{
+			return GetObjectDisplayString(*LoadedObject);
+		}
+
+		// Subpath looks like this PersistentLevel.Actor.Component
+		const FString& Subpath = Object.GetSubPathString();
+		const int32 LastDotIndex = Subpath.Find(TEXT("."), ESearchCase::CaseSensitive, ESearchDir::FromEnd);
+		if (LastDotIndex == INDEX_NONE)
+		{
+			return {};
+		}
+		return Subpath.RightChop(LastDotIndex + 1);
+	}
+
+	FString GetObjectDisplayString(const UObject& Object)
+	{
+		if (const AActor* AsActor = Cast<AActor>(&Object))
+		{
+			return AsActor->GetActorLabel();
+		}
+		return Object.GetName();
+	}
+
+	FText GetObjectTypeText(const IObjectToPropertiesModel& Model, const FSoftObjectPath& Object)
+	{
+		if (const FSoftClassPath ClassPath = Model.GetObjectClass(Object); ClassPath.IsValid())
+		{
+			// For C++ classes or loaded Blueprints
+			if (UClass* Class = ClassPath.ResolveClass())
+			{
+				return Class->GetDisplayNameText();
+			}
+
+			// Assuming it is a Blueprint, try to get its display name from the registry tags
+			const FAssetRegistryModule& AssetRegistryModule = FModuleManager::LoadModuleChecked<FAssetRegistryModule>(TEXT("AssetRegistry"));
+			const IAssetRegistry& AssetRegistry = AssetRegistryModule.Get();
+			
+			const FAssetData BlueprintData = AssetRegistry.GetAssetByObjectPath(ClassPath);
+			if (BlueprintData.IsValid())
+			{
+				FString DisplayName;
+				return BlueprintData.GetTagValue(FBlueprintTags::BlueprintDisplayName, DisplayName)
+					? FText::FromString(DisplayName)
+					: FText::GetEmpty();
+			}
+
+			// TODO: Check redirectors
+			// It was not a Blueprint - the class was most likely removed
+			UE_LOG(LogConcert, Error, TEXT("Failed to get display name for unknown class \"%s\""), *ClassPath.ToString());
+		}
+		return FText::GetEmpty();
+	}
+
+	FSlateIcon GetObjectIcon(const IObjectToPropertiesModel& Model, const FSoftObjectPath& Object)
+	{
+		if (const FSoftClassPath ClassPath = Model.GetObjectClass(Object); ClassPath.IsValid())
+		{
+			UClass* Class = ClassPath.ResolveClass();
+			return Class
+				? FSlateIconFinder::FindIconForClass(Class)
+				: FSlateIcon();
+		}
+		return FSlateIcon();
+	}
+
+	FSlateIcon GetObjectIcon(UObject& Object)
+	{
+		return FSlateIconFinder::FindIconForClass(Object.GetClass());
+	}
+	
+	FText GetPropertyDisplayText(const FConcertPropertyChain& Property)
+	{
+		return FText::FromString(GetPropertyDisplayString(Property));
+    }
+	
+	FString GetPropertyDisplayString(const FConcertPropertyChain& Property)
+	{
+		return Property.ToString(FConcertPropertyChain::EToStringMethod::LeafProperty);
+	}
+}
diff --git a/Engine/Plugins/Developer/Concert/ConcertUI/ConcertClientSharedSlate/Source/ConcertClientSharedSlate/Private/Replication/Editor/View/DisplayUtils.h b/Engine/Plugins/Developer/Concert/ConcertUI/ConcertClientSharedSlate/Source/ConcertClientSharedSlate/Private/Replication/Editor/View/DisplayUtils.h
new file mode 100644
--- /dev/null
+++ b/Engine/Plugins/Developer/Concert/ConcertUI/ConcertClientSharedSlate/Source/ConcertClientSharedSlate/Private/Replication/Editor/View/DisplayUtils.h
@@ -0,0 +1,40 @@
+// Copyright Epic Games, Inc. All Rights Reserved.
+
+#pragma once
+
+class FString;
+class FText;
+class UObject;
+
+struct FConcertPropertyChain;
+struct FSlateIcon;
+struct FSoftObjectPath;
+
+namespace UE::ConcertClientSharedSlate
+{
+	class IObjectToPropertiesModel;
+}
+
+namespace UE::ConcertClientSharedSlate::DisplayUtils
+{
+	/** @return The text to use for displaying this object's name */
+	FText GetObjectDisplayText(const FSoftObjectPath& Object);
+	
+	/** @return More lightweight version of GetObjectDisplayText which does not construct any FText. */
+	FString GetObjectDisplayString(const FSoftObjectPath& Object);
+	/** @return More lightweight version of GetObjectDisplayText which does not construct any FText. */
+	FString GetObjectDisplayString(const UObject& Object);
+	
+	/** @return The text to use for displaying this object's type */
+	FText GetObjectTypeText(const IObjectToPropertiesModel& Model, const FSoftObjectPath& Object);
+	
+	/** @return The icon to use for this object */
+	FSlateIcon GetObjectIcon(const IObjectToPropertiesModel& Model, const FSoftObjectPath& Object);
+	/** @return The icon to use for this object */
+	FSlateIcon GetObjectIcon(UObject& Object);
+
+	/** @return The text to use for displaying this property's name. */
+	FText GetPropertyDisplayText(const FConcertPropertyChain& Property);
+	/** @return More lightweight version of GetPropertyDisplayString which does not construct any FText. */
+	FString GetPropertyDisplayString(const FConcertPropertyChain& Property);
+}
diff --git a/Engine/Plugins/Developer/Concert/ConcertUI/ConcertClientSharedSlate/Source/ConcertClientSharedSlate/Private/Replication/Editor/View/ObjectEditor/SBaseReplicationStreamEditor.cpp b/Engine/Plugins/Developer/Concert/ConcertUI/ConcertClientSharedSlate/Source/ConcertClientSharedSlate/Private/Replication/Editor/View/ObjectEditor/SBaseReplicationStreamEditor.cpp
--- a/Engine/Plugins/Developer/Concert/ConcertUI/ConcertClientSharedSlate/Source/ConcertClientSharedSlate/Private/Replication/Editor/View/ObjectEditor/SBaseReplicationStreamEditor.cpp
+++ b/Engine/Plugins/Developer/Concert/ConcertUI/ConcertClientSharedSlate/Source/ConcertClientSharedSlate/Private/Replication/Editor/View/ObjectEditor/SBaseReplicationStreamEditor.cpp
@@ -4,13 +4,11 @@
 
 #include "FakeObjectToPropertiesEditorModel.h"
 #include "Model/Item/SourceModelBuilders.h"
-#include "Replication/Data/ConcertPropertySelection.h"
-#include "Replication/Editor/Model/DisplayUtils.h"
 #include "Replication/Editor/Model/IEditableObjectToPropertiesModel.h"
 #include "Replication/Editor/Model/Object/IObjectSelectionSourceModel.h"
 #include "Replication/Editor/Model/Property/IPropertySelectionSourceModel.h"
-#include "Replication/Editor/Model/ReplicatedPropertyData.h"
 #include "Replication/Editor/Model/ReplicatedObjectData.h"
+#include "Replication/Editor/View/DisplayUtils.h"
 #include "Replication/Editor/View/ObjectViewer/SReplicationStreamViewer.h"
 
 #include "Algo/AnyOf.h"
diff --git a/Engine/Plugins/Developer/Concert/ConcertUI/ConcertClientSharedSlate/Source/ConcertClientSharedSlate/Private/Replication/Editor/View/ObjectViewer/Property/SReplicatedPropertiesView.cpp b/Engine/Plugins/Developer/Concert/ConcertUI/ConcertClientSharedSlate/Source/ConcertClientSharedSlate/Private/Replication/Editor/View/ObjectViewer/Property/SReplicatedPropertiesView.cpp
--- a/Engine/Plugins/Developer/Concert/ConcertUI/ConcertClientSharedSlate/Source/ConcertClientSharedSlate/Private/Replication/Editor/View/ObjectViewer/Property/SReplicatedPropertiesView.cpp
+++ b/Engine/Plugins/Developer/Concert/ConcertUI/ConcertClientSharedSlate/Source/ConcertClientSharedSlate/Private/Replication/Editor/View/ObjectViewer/Property/SReplicatedPropertiesView.cpp
@@ -5,8 +5,8 @@
 #include "Filters/SBasicFilterBar.h"
 #include "PropertyFilter_ByPropertyType.h"
 #include "PropertyFrontendFilter.h"
-#include "Replication/Editor/Model/DisplayUtils.h"
 #include "Replication/Editor/Model/ReplicatedPropertyData.h"
+#include "Replication/Editor/View/DisplayUtils.h"
 
 #include "UObject/UnrealType.h"
 
diff --git a/Engine/Plugins/Developer/Concert/ConcertUI/ConcertClientSharedSlate/Source/ConcertClientSharedSlate/Private/Replication/Editor/View/ObjectViewer/Tree/SelectionViewerColumns.cpp b/Engine/Plugins/Developer/Concert/ConcertUI/ConcertClientSharedSlate/Source/ConcertClientSharedSlate/Private/Replication/Editor/View/ObjectViewer/Tree/SelectionViewerColumns.cpp
--- a/Engine/Plugins/Developer/Concert/ConcertUI/ConcertClientSharedSlate/Source/ConcertClientSharedSlate/Private/Replication/Editor/View/ObjectViewer/Tree/SelectionViewerColumns.cpp
+++ b/Engine/Plugins/Developer/Concert/ConcertUI/ConcertClientSharedSlate/Source/ConcertClientSharedSlate/Private/Replication/Editor/View/ObjectViewer/Tree/SelectionViewerColumns.cpp
@@ -3,11 +3,11 @@
 #include "SelectionViewerColumns.h"
 
 #include "ClassIconFinder.h"
-#include "Replication/Editor/Model/DisplayUtils.h"
 #include "Replication/Editor/Model/IEditableObjectToPropertiesModel.h"
 #include "Replication/Editor/Model/ReplicatedPropertyData.h"
 #include "Replication/Editor/Model/ReplicatedObjectData.h"
 #include "Replication/Editor/Model/Subobject/ISubobjectModel.h"
+#include "Replication/Editor/View/DisplayUtils.h"
 #include "Replication/Editor/View/ObjectEditor/SBaseReplicationStreamEditor.h"
 #include "Replication/Editor/View/ReplicationColumnsUtils.h"
 #include "Replication/PropertyChainUtils.h"
