Add owner list to columns

diff --git a/Engine/Plugins/Developer/Concert/ConcertApp/MultiUserClient/Source/MultiUserClient/Private/Replication/Authority/AuthoritySynchronizer_RemoteClient.cpp b/Engine/Plugins/Developer/Concert/ConcertApp/MultiUserClient/Source/MultiUserClient/Private/Replication/Authority/AuthoritySynchronizer_RemoteClient.cpp
--- a/Engine/Plugins/Developer/Concert/ConcertApp/MultiUserClient/Source/MultiUserClient/Private/Replication/Authority/AuthoritySynchronizer_RemoteClient.cpp
+++ b/Engine/Plugins/Developer/Concert/ConcertApp/MultiUserClient/Source/MultiUserClient/Private/Replication/Authority/AuthoritySynchronizer_RemoteClient.cpp
@@ -34,11 +34,16 @@ namespace UE::MultiUserClient
 
 	void FAuthoritySynchronizer_RemoteClient::HandleAuthorityQuery(const TArray<FReplicationAuthorityInfo>& PerStreamAuthority)
 	{
-		LastServerState.Empty();
-
+		TSet<FSoftObjectPath> OldServerState = MoveTemp(LastServerState);
 		for (const FReplicationAuthorityInfo& Info : PerStreamAuthority)
 		{
 			LastServerState.Append(Info.AuthoredObjects);
 		}
+		
+		const bool bAreEqual = OldServerState.Num() == LastServerState.Num() && LastServerState.Includes(OldServerState);
+		if (!bAreEqual)
+		{
+			OnServerStateChangedDelegate.Broadcast();
+		}
 	}
 }
diff --git a/Engine/Plugins/Developer/Concert/ConcertApp/MultiUserClient/Source/MultiUserClient/Private/Replication/Client/LocalReplicationClient.cpp b/Engine/Plugins/Developer/Concert/ConcertApp/MultiUserClient/Source/MultiUserClient/Private/Replication/Client/LocalReplicationClient.cpp
--- a/Engine/Plugins/Developer/Concert/ConcertApp/MultiUserClient/Source/MultiUserClient/Private/Replication/Client/LocalReplicationClient.cpp
+++ b/Engine/Plugins/Developer/Concert/ConcertApp/MultiUserClient/Source/MultiUserClient/Private/Replication/Client/LocalReplicationClient.cpp
@@ -2,6 +2,7 @@
 
 #include "LocalReplicationClient.h"
 
+#include "IConcertSyncClient.h"
 #include "Replication/Authority/AuthoritySynchronizer_LocalClient.h"
 #include "Replication/Submission/SubmissionWorkflow_LocalClient.h"
 
@@ -13,6 +14,7 @@ namespace UE::MultiUserClient
 		TSharedRef<IConcertSyncClient> InClient
 		)
 		: FReplicationClient(
+			InClient->GetConcertClient()->GetCurrentSession()->GetSessionClientEndpointId(),
 			InSessionContent,
 			MoveTemp(InStreamSynchronizer),
 			MakeUnique<FAuthoritySynchronizer_LocalClient>(
diff --git a/Engine/Plugins/Developer/Concert/ConcertApp/MultiUserClient/Source/MultiUserClient/Private/Replication/Client/RemoteReplicationClient.cpp b/Engine/Plugins/Developer/Concert/ConcertApp/MultiUserClient/Source/MultiUserClient/Private/Replication/Client/RemoteReplicationClient.cpp
--- a/Engine/Plugins/Developer/Concert/ConcertApp/MultiUserClient/Source/MultiUserClient/Private/Replication/Client/RemoteReplicationClient.cpp
+++ b/Engine/Plugins/Developer/Concert/ConcertApp/MultiUserClient/Source/MultiUserClient/Private/Replication/Client/RemoteReplicationClient.cpp
@@ -10,11 +10,12 @@
 namespace UE::MultiUserClient
 {
 	FRemoteReplicationClient::FRemoteReplicationClient(
-		UMultiUserReplicationClientPreset& InSessionContent,
 		const FGuid& InConcertClientId,
+		UMultiUserReplicationClientPreset& InSessionContent,
 		FRegularQueryService& QueryService
 		)
 		: FReplicationClient(
+			InConcertClientId,
 			InSessionContent,
 			MakeUnique<FStreamSynchronizer_RemoteClient>(InConcertClientId, QueryService),
 			MakeUnique<FAuthoritySynchronizer_RemoteClient>(
@@ -28,7 +29,6 @@ namespace UE::MultiUserClient
 			{
 				return MakeUnique<FSubmissionWorkflow_RemoteClient>();
 			})
-		, RemoteEndpointId(InConcertClientId)
 	{
 		// When the remote client's state has changed, refresh the UI.
 		GetStreamSynchronizer().OnServerStateChanged().AddLambda([this]()
diff --git a/Engine/Plugins/Developer/Concert/ConcertApp/MultiUserClient/Source/MultiUserClient/Private/Replication/Client/RemoteReplicationClient.h b/Engine/Plugins/Developer/Concert/ConcertApp/MultiUserClient/Source/MultiUserClient/Private/Replication/Client/RemoteReplicationClient.h
--- a/Engine/Plugins/Developer/Concert/ConcertApp/MultiUserClient/Source/MultiUserClient/Private/Replication/Client/RemoteReplicationClient.h
+++ b/Engine/Plugins/Developer/Concert/ConcertApp/MultiUserClient/Source/MultiUserClient/Private/Replication/Client/RemoteReplicationClient.h
@@ -14,22 +14,10 @@ namespace UE::MultiUserClient
 	public:
 
 		FRemoteReplicationClient(
-			UMultiUserReplicationClientPreset& InSessionContent,
 			const FGuid& InConcertClientId,
+			UMultiUserReplicationClientPreset& InSessionContent,
 			FRegularQueryService& QueryService
 			);
-		
-		const FGuid& GetRemoteEndpointId() const { return RemoteEndpointId; }
-		
-	private:
-
-		/**
-		 * Endpoint of the Concert client.
-		 *
-		 * Can be used to look up client info.
-		 * @see IConcertSession::FindSessionClient
-		 */
-		const FGuid RemoteEndpointId;
 	};
 }
 
diff --git a/Engine/Plugins/Developer/Concert/ConcertApp/MultiUserClient/Source/MultiUserClient/Private/Replication/Client/ReplicationClient.cpp b/Engine/Plugins/Developer/Concert/ConcertApp/MultiUserClient/Source/MultiUserClient/Private/Replication/Client/ReplicationClient.cpp
--- a/Engine/Plugins/Developer/Concert/ConcertApp/MultiUserClient/Source/MultiUserClient/Private/Replication/Client/ReplicationClient.cpp
+++ b/Engine/Plugins/Developer/Concert/ConcertApp/MultiUserClient/Source/MultiUserClient/Private/Replication/Client/ReplicationClient.cpp
@@ -11,12 +11,14 @@
 namespace UE::MultiUserClient
 {
 	FReplicationClient::FReplicationClient(
+		const FGuid& EndpointId,
 		UMultiUserReplicationClientPreset& InSessionContent,
 		TUniquePtr<IClientStreamSynchronizer> InStreamSynchronizer,
 		TUniquePtr<IClientAuthoritySynchronizer> InAuthoritySynchronizer,
 		TFunctionRef<FMakeSubmissionWorkflow> MakeSubmissionWorkflowFunc
 		)
-		: ClientContentStorage(&InSessionContent)
+		: EndpointId(EndpointId)
+		, ClientContentStorage(&InSessionContent)
 		, StreamSynchronizer(MoveTemp(InStreamSynchronizer))
 		, AuthoritySynchronizer(MoveTemp(InAuthoritySynchronizer))
 		, LocalClientEditModel(ConcertClientSharedSlate::CreatePropertySelectionModel(
diff --git a/Engine/Plugins/Developer/Concert/ConcertApp/MultiUserClient/Source/MultiUserClient/Private/Replication/Client/ReplicationClient.h b/Engine/Plugins/Developer/Concert/ConcertApp/MultiUserClient/Source/MultiUserClient/Private/Replication/Client/ReplicationClient.h
--- a/Engine/Plugins/Developer/Concert/ConcertApp/MultiUserClient/Source/MultiUserClient/Private/Replication/Client/ReplicationClient.h
+++ b/Engine/Plugins/Developer/Concert/ConcertApp/MultiUserClient/Source/MultiUserClient/Private/Replication/Client/ReplicationClient.h
@@ -33,6 +33,7 @@ namespace UE::MultiUserClient
 		using FMakeSubmissionWorkflow = TUniquePtr<ISubmissionWorkflow>(FStreamChangeTracker&, FAuthorityChangeTracker&, IClientStreamSynchronizer&);
 
 		FReplicationClient(
+			const FGuid& EndpointId,
 			UMultiUserReplicationClientPreset& InSessionContent,
 			TUniquePtr<IClientStreamSynchronizer> InStreamSynchronizer,
 			TUniquePtr<IClientAuthoritySynchronizer> InAuthoritySynchronizer,
@@ -46,11 +47,15 @@ namespace UE::MultiUserClient
 		
 		const FStreamChangeTracker& GetStreamDiffer() const { return LocalClientStreamDiffer; }
 		FStreamChangeTracker& GetStreamDiffer() { return LocalClientStreamDiffer; }
+		
 		const FAuthorityChangeTracker& GetAuthorityDiffer() const { return LocalAuthorityDiffer; }
 		FAuthorityChangeTracker& GetAuthorityDiffer() { return LocalAuthorityDiffer; }
+		
 		const ISubmissionWorkflow& GetSubmissionWorkflow() const { return *SubmissionWorkflow; }
 		ISubmissionWorkflow& GetSubmissionWorkflow() { return *SubmissionWorkflow; }
 
+		const FGuid& GetEndpointId() const { return EndpointId; }
+
 		/**
 		 * Called when the data underlying the model has changed externally. Since the change was not caused by the model,
 		 * its events, like IEditableObjectToPropertiesModel::OnObjectsChanged, were not called.
@@ -64,6 +69,9 @@ namespace UE::MultiUserClient
 		FOnModelExternallyChanged& OnModelExternallyChanged() { return OnModelExternallyChangedDelegate; }
 		
 	private:
+
+		/** This client's Concert Endpoint ID. */
+		const FGuid EndpointId;
 		
 		/** The state of the server is synched up with this object and displayed in the UI. */
 		TObjectPtr<UMultiUserReplicationClientPreset> ClientContentStorage;
diff --git a/Engine/Plugins/Developer/Concert/ConcertApp/MultiUserClient/Source/MultiUserClient/Private/Replication/Client/ReplicationClientManager.cpp b/Engine/Plugins/Developer/Concert/ConcertApp/MultiUserClient/Source/MultiUserClient/Private/Replication/Client/ReplicationClientManager.cpp
--- a/Engine/Plugins/Developer/Concert/ConcertApp/MultiUserClient/Source/MultiUserClient/Private/Replication/Client/ReplicationClientManager.cpp
+++ b/Engine/Plugins/Developer/Concert/ConcertApp/MultiUserClient/Source/MultiUserClient/Private/Replication/Client/ReplicationClientManager.cpp
@@ -19,21 +19,20 @@ namespace UE::MultiUserClient
 		)
 		: SessionContent(NewObject<UMultiUserReplicationSessionPreset>(GetTransientPackage(), NAME_None, RF_Transient))
 		, Session(InSession)
-		, QueryService(InClient)
 		, LocalClient([this, InClient]()
 		{
 			UMultiUserReplicationClientPreset* ClientPreset = SessionContent->AddClient();
 			return FLocalReplicationClient(*ClientPreset, MakeUnique<FStreamSynchronizer_LocalClient>(InClient, ClientPreset->Stream->StreamId), InClient);
 		}())
+		, QueryService(InClient)
+		, AuthorityCache(*this)
 		, SubmissionNotifier(*this)
 	{
 		InSession->OnSessionClientChanged().AddRaw(this, &FReplicationClientManager::OnSessionClientChanged);
 
 		for (const FGuid& ClientEndpointId : InSession->GetSessionClientEndpointIds())
 		{
-			// Nobody is subscribed at this point, so no need to broadcast
-			constexpr bool bBroadcast = false;
-			CreateRemoteClient(ClientEndpointId, bBroadcast);
+			CreateRemoteClient(ClientEndpointId);
 		}
 	}
 
@@ -59,7 +58,7 @@ namespace UE::MultiUserClient
 	{
 		const TUniquePtr<FRemoteReplicationClient>* Client = RemoteClients.FindByPredicate([&EndpointId](const TUniquePtr<FRemoteReplicationClient>& Client)
 		{
-			return Client->GetRemoteEndpointId() == EndpointId;
+			return Client->GetEndpointId() == EndpointId;
 		});
 		return Client ? Client->Get() : nullptr;
 	}
@@ -84,7 +83,7 @@ namespace UE::MultiUserClient
 				const int32 Index = RemoteClients.IndexOfByPredicate(
 					[&ClientEndpointId](const TUniquePtr<FRemoteReplicationClient>& Client)
 					{
-						return Client->GetRemoteEndpointId() == ClientEndpointId;
+						return Client->GetEndpointId() == ClientEndpointId;
 					});
 				if (!ensure(RemoteClients.IsValidIndex(Index)))
 				{
@@ -110,7 +109,7 @@ namespace UE::MultiUserClient
 	
 	void FReplicationClientManager::CreateRemoteClient(const FGuid& ClientEndpointId, bool bBroadcastDelegate)
 	{
-		TUniquePtr<FRemoteReplicationClient> RemoteClientPtr = MakeUnique<FRemoteReplicationClient>(*SessionContent->AddClient(), ClientEndpointId, QueryService);
+		TUniquePtr<FRemoteReplicationClient> RemoteClientPtr = MakeUnique<FRemoteReplicationClient>(ClientEndpointId, *SessionContent->AddClient(), QueryService);
 		FRemoteReplicationClient& RemoteClient = *RemoteClientPtr;
 		RemoteClients.Emplace(
 			MoveTemp(RemoteClientPtr)
diff --git a/Engine/Plugins/Developer/Concert/ConcertApp/MultiUserClient/Source/MultiUserClient/Private/Replication/Client/ReplicationClientManager.h b/Engine/Plugins/Developer/Concert/ConcertApp/MultiUserClient/Source/MultiUserClient/Private/Replication/Client/ReplicationClientManager.h
--- a/Engine/Plugins/Developer/Concert/ConcertApp/MultiUserClient/Source/MultiUserClient/Private/Replication/Client/ReplicationClientManager.h
+++ b/Engine/Plugins/Developer/Concert/ConcertApp/MultiUserClient/Source/MultiUserClient/Private/Replication/Client/ReplicationClientManager.h
@@ -7,6 +7,7 @@
 #include "LocalReplicationClient.h"
 #include "ReplicationClient.h"
 #include "Replication/Submission/Notification/SubmissionNotifier.h"
+#include "Replication/Util/GlobalAuthorityCache.h"
 #include "Replication/Util/RegularQueryService.h"
 
 #include "UObject/GCObject.h"
@@ -51,6 +52,9 @@ namespace UE::MultiUserClient
 		const FLocalReplicationClient& GetLocalClient() const { return LocalClient; }
 		FLocalReplicationClient& GetLocalClient() { return LocalClient; }
 		TArray<TNonNullPtr<FRemoteReplicationClient>> GetRemoteClients() const;
+		
+		const FGlobalAuthorityCache& GetAuthorityCache() const { return AuthorityCache; }
+		FGlobalAuthorityCache& GetAuthorityCache() { return AuthorityCache; }
 
 		/** Util for finding a remote client by its EndpointId. */
 		const FRemoteReplicationClient* FindRemoteClient(const FGuid& EndpointId) const;
@@ -59,6 +63,20 @@ namespace UE::MultiUserClient
 			const FReplicationClientManager* ConstThis = this;
 			return const_cast<FRemoteReplicationClient*>(ConstThis->FindRemoteClient(EndpointId));
 		}
+
+		/** Util for finding a local or remote client by its EndpointId. */
+		const FReplicationClient* FindClient(const FGuid& EndpointId) const
+		{
+			return GetLocalClient().GetEndpointId() == EndpointId
+				? &GetLocalClient()
+				: static_cast<const FReplicationClient*>(FindRemoteClient(EndpointId));
+		}
+		FReplicationClient* FindClient(const FGuid& EndpointId)
+		{
+			const FReplicationClientManager* ConstThis = this;
+			return const_cast<FReplicationClient*>(ConstThis->FindClient(EndpointId));
+		}
+
 		
 		DECLARE_MULTICAST_DELEGATE(FRemoteClientsChanged);
 		/** Called when RemoteClients changes. Called after OnPostRemoteClientAdded. */
@@ -88,14 +106,9 @@ namespace UE::MultiUserClient
 		 */
 		const TWeakPtr<IConcertClientSession> Session;
 
-		/**
-		 * Sends FConcertReplication_QueryReplicationInfo_Request in regular intervals.
-		 * Shared by all remote clients so all requests are bundled reducing the number of network requests. 
-		 */
-		FRegularQueryService QueryService;
-		
 		/** Manages the local client */
 		FLocalReplicationClient LocalClient;
+		
 		/**
 		 * Manages remote clients. Updated when client connects or disconnects to the active session.
 		 * UI keeps references to systems inside the client so it is TSharedRef in case TArray is reallocated.
@@ -109,6 +122,14 @@ namespace UE::MultiUserClient
 		/** Called just before a remote client is about to be removed from RemoteClients. */
 		FRemoteClientDelegate OnPreRemoteClientRemovedDelegate; 
 		
+		/**
+		 * Sends FConcertReplication_QueryReplicationInfo_Request in regular intervals.
+		 * Shared by all remote clients so all requests are bundled reducing the number of network requests. 
+		 */
+		FRegularQueryService QueryService;
+		/** Keeps a cache of object to owning clients. */
+		FGlobalAuthorityCache AuthorityCache;
+		
 		/** Manages SNotificationItems when submission to the server fails. */
 		FSubmissionNotifier SubmissionNotifier;
 		
diff --git a/Engine/Plugins/Developer/Concert/ConcertApp/MultiUserClient/Source/MultiUserClient/Private/Replication/MultiUserReplicationManager.cpp b/Engine/Plugins/Developer/Concert/ConcertApp/MultiUserClient/Source/MultiUserClient/Private/Replication/MultiUserReplicationManager.cpp
--- a/Engine/Plugins/Developer/Concert/ConcertApp/MultiUserClient/Source/MultiUserClient/Private/Replication/MultiUserReplicationManager.cpp
+++ b/Engine/Plugins/Developer/Concert/ConcertApp/MultiUserClient/Source/MultiUserClient/Private/Replication/MultiUserReplicationManager.cpp
@@ -75,7 +75,7 @@ namespace UE::MultiUserClient
 	{
 		// This clears the UI. The clients' IEditableObjectToPropertiesModels should no longer be referenced by anyone.
 		SetConnectionStateAndBroadcast(EMultiUserReplicationConnectionState::Disconnected);
-		// Keep in mind that FClientStreamRepository::GetLocalClientEditModel is referenced by the UI so call this after clearing the UI.
+		// Keep in mind the IEditableObjectToPropertiesModels were referenced by the UI so call this after clearing the UI.
 		ConnectedState.Reset();
 	}
 
diff --git a/Engine/Plugins/Developer/Concert/ConcertApp/MultiUserClient/Source/MultiUserClient/Private/Replication/Util/GlobalAuthorityCache.cpp b/Engine/Plugins/Developer/Concert/ConcertApp/MultiUserClient/Source/MultiUserClient/Private/Replication/Util/GlobalAuthorityCache.cpp
new file mode 100644
--- /dev/null
+++ b/Engine/Plugins/Developer/Concert/ConcertApp/MultiUserClient/Source/MultiUserClient/Private/Replication/Util/GlobalAuthorityCache.cpp
@@ -0,0 +1,129 @@
+// Copyright Epic Games, Inc. All Rights Reserved.
+
+#include "GlobalAuthorityCache.h"
+
+#include "Replication/Client/RemoteReplicationClient.h"
+#include "Replication/Client/ReplicationClientManager.h"
+
+namespace UE::MultiUserClient
+{
+	FGlobalAuthorityCache::FGlobalAuthorityCache(FReplicationClientManager& InClientManager)
+		: ClientManager(InClientManager)
+	{
+		RegisterForClientEvents(ClientManager.GetLocalClient());
+
+		ClientManager.OnPostRemoteClientAdded().AddRaw(this, &FGlobalAuthorityCache::OnPostRemoteClientAdded);
+		ClientManager.OnPreRemoteClientRemoved().AddRaw(this, &FGlobalAuthorityCache::OnPreRemoteClientRemoved);
+	}
+
+	FGlobalAuthorityCache::~FGlobalAuthorityCache()
+	{
+		// FGlobalAuthorityCache is owned by FRelicationClientManager so this is strictly not needed. But we'll follow RAII in case ownership changes.
+		UnregisterFromClientEvents(ClientManager.GetLocalClient());
+		for (const TNonNullPtr<FRemoteReplicationClient>& RemoteClient : ClientManager.GetRemoteClients())
+		{
+			UnregisterFromClientEvents(*RemoteClient);
+		}
+		
+		ClientManager.OnPostRemoteClientAdded().RemoveAll(this);
+		ClientManager.OnPreRemoteClientRemoved().RemoveAll(this);
+	}
+
+	void FGlobalAuthorityCache::ForEachClientWithAuthorityOverObject(const FSoftObjectPath& Object, TFunctionRef<EBreakBehavior(const FGuid& ClientId)> Callback) const
+	{
+		const TSet<FGuid>* Clients = OwnedObjectsToClients.Find(Object);
+		if (!Clients)
+		{
+			return;
+		}
+
+		for (const FGuid& ClientId : *Clients)
+		{
+			if (Callback(ClientId) == EBreakBehavior::Break)
+			{
+				break;
+			}
+		}
+	}
+
+	TArray<FGuid> FGlobalAuthorityCache::GetClientsWithAuthorityOverObject(const FSoftObjectPath& Object) const
+	{
+		TArray<FGuid> Result;
+		ForEachClientWithAuthorityOverObject(Object, [&Result](const FGuid& ClientId)
+		{
+			Result.Add(ClientId);
+			return EBreakBehavior::Continue;
+		});
+		return Result;
+	}
+
+	TOptional<FGuid> FGlobalAuthorityCache::GetClientWithAuthorityOverProperty(const FSoftObjectPath& Object, const FConcertPropertyChain& Property) const
+	{
+		TOptional<FGuid> Result;
+		ForEachClientWithAuthorityOverObject(Object, [this, &Object, &Property, &Result](const FGuid& ClientId)
+		{
+			const FReplicationClient* Client = ClientManager.FindClient(ClientId);
+			if (!ensureMsgf(Client, TEXT("OnPreRemoteClientRemoved should have updated OwnedObjectsToClients")))
+			{
+				return EBreakBehavior::Continue;
+			}
+			
+			const FReplicatedObjectInfo* ObjectInfo = Client->GetStreamSynchronizer().GetServerState().ReplicatedObjects.Find(Object);
+			const bool bHasObjectRegistered = ensureMsgf(ObjectInfo, TEXT("OnStreamChanged should have updated OwnedObjectsToClients"))
+				&& ObjectInfo->PropertySelection.ReplicatedProperties.Contains(Property);
+			const bool bHasAuthority = Client->GetAuthoritySynchronizer().HasAuthorityOver(Object);
+			if (bHasObjectRegistered && bHasAuthority)
+			{
+				Result = ClientId;
+				return EBreakBehavior::Break;
+			}
+			return EBreakBehavior::Continue;
+		});
+		return Result;
+	}
+
+	void FGlobalAuthorityCache::RegisterForClientEvents(const FReplicationClient& Client)
+	{
+		const FGuid& ClientEndpointId = Client.GetEndpointId();
+		Client.GetAuthoritySynchronizer().OnServerStateChanged().AddRaw(this, &FGlobalAuthorityCache::OnPostAuthorityChanged, ClientEndpointId);
+		Client.GetStreamSynchronizer().OnServerStateChanged().AddRaw(this, &FGlobalAuthorityCache::OnStreamChanged, ClientEndpointId);
+	}
+
+	void FGlobalAuthorityCache::UnregisterFromClientEvents(const FReplicationClient& Client) const
+	{
+		Client.GetAuthoritySynchronizer().OnServerStateChanged().RemoveAll(this);
+		Client.GetStreamSynchronizer().OnServerStateChanged().RemoveAll(this);
+	}
+
+	void FGlobalAuthorityCache::AddClient(const FGuid& ClientId)
+	{
+		const FReplicationClient* Client = ClientManager.FindClient(ClientId);
+		if (!ensure(Client))
+		{
+			return;
+		}
+		
+		const FObjectReplicationMap& ClientObjectMap = Client->GetStreamSynchronizer().GetServerState();
+		const IClientAuthoritySynchronizer& ClientAuthority = Client->GetAuthoritySynchronizer();
+		for (const TPair<FSoftObjectPath, FReplicatedObjectInfo>& StreamContents : ClientObjectMap.ReplicatedObjects)
+		{
+			const FSoftObjectPath& Object = StreamContents.Key;
+			if (ClientAuthority.HasAuthorityOver(Object))
+			{
+				OwnedObjectsToClients.FindOrAdd(Object).Add(ClientId);
+			}
+		}
+	}
+
+	void FGlobalAuthorityCache::RemoveClient(const FGuid& ClientId)
+	{
+		for (auto It = OwnedObjectsToClients.CreateIterator(); It; ++It)
+		{
+			It->Value.Remove(ClientId);
+			if (It->Value.IsEmpty())
+			{
+				It.RemoveCurrent();
+			}
+		}
+	}
+}
diff --git a/Engine/Plugins/Developer/Concert/ConcertApp/MultiUserClient/Source/MultiUserClient/Private/Replication/Util/GlobalAuthorityCache.h b/Engine/Plugins/Developer/Concert/ConcertApp/MultiUserClient/Source/MultiUserClient/Private/Replication/Util/GlobalAuthorityCache.h
new file mode 100644
--- /dev/null
+++ b/Engine/Plugins/Developer/Concert/ConcertApp/MultiUserClient/Source/MultiUserClient/Private/Replication/Util/GlobalAuthorityCache.h
@@ -0,0 +1,78 @@
+// Copyright Epic Games, Inc. All Rights Reserved.
+
+#pragma once
+
+#include "Replication/Client/RemoteReplicationClient.h"
+
+#include "Delegates/Delegate.h"
+#include "Templates/Function.h"
+
+struct FConcertPropertyChain;
+
+namespace UE::MultiUserClient
+{
+	class FRemoteReplicationClient;
+	class FReplicationClient;
+	class FReplicationClientManager;
+}
+
+namespace UE::MultiUserClient
+{
+	/**
+	 * Allows efficient look-up of which objects and properties are owned by which clients.
+	 * This class efficiently answers the question: "Which clients own this object?"
+	 */
+	class FGlobalAuthorityCache : public FNoncopyable
+	{
+	public:
+		
+		FGlobalAuthorityCache(FReplicationClientManager& InClientManager);
+		~FGlobalAuthorityCache();
+
+		/** Iterates every client that has authority over Object. */
+		void ForEachClientWithAuthorityOverObject(const FSoftObjectPath& Object, TFunctionRef<EBreakBehavior(const FGuid& ClientId)> Callback) const;
+		/** Util that uses ForEachClientWithAuthorityOverObject to make an array. */
+		TArray<FGuid> GetClientsWithAuthorityOverObject(const FSoftObjectPath& Object) const;
+
+		/** Gets the client that has authority over the given property, if there is any. */
+		TOptional<FGuid> GetClientWithAuthorityOverProperty(const FSoftObjectPath& Object, const FConcertPropertyChain& Property) const;
+
+		DECLARE_MULTICAST_DELEGATE_OneParam(FOnCacheChanged, const FGuid& ClientId);
+		/** Called when the cache changes for a specific client. */
+		FOnCacheChanged& OnCacheChanged() { return OnCacheChangedDelegate; }
+		
+	private:
+
+		/** Used to obtain the clients and their states */
+		FReplicationClientManager& ClientManager;
+		
+		/** Maps objects that are owned to the clients that own them */
+		TMap<FSoftObjectPath, TSet<FGuid>> OwnedObjectsToClients;
+		
+		/** Called when the cache changes for a specific client. */
+		FOnCacheChanged OnCacheChangedDelegate;
+
+		/** Registers for authority and stream changes */
+		void RegisterForClientEvents(const FReplicationClient& Client);
+		void UnregisterFromClientEvents(const FReplicationClient& Client) const;
+
+		/** Adds the client to OwnedObjectsToClients */
+		void AddClient(const FGuid& ClientId);
+		void RemoveClient(const FGuid& ClientId);
+
+		// Respond to remote client registration
+		void OnPostRemoteClientAdded(FRemoteReplicationClient& RemoteClient) { RegisterForClientEvents(RemoteClient); }
+		void OnPreRemoteClientRemoved(FRemoteReplicationClient& RemoteClient) const { UnregisterFromClientEvents(RemoteClient); }
+
+		// Rebuild client when their authority changes
+		void OnPostAuthorityChanged(const FGuid ClientId) { RebuildClient(ClientId); }
+		void OnStreamChanged(const FGuid ClientId){ RebuildClient(ClientId); }
+		void RebuildClient(const FGuid ClientId)
+		{
+			RemoveClient(ClientId);
+			AddClient(ClientId);
+			OnCacheChangedDelegate.Broadcast(ClientId);
+		}
+	};
+}
+
diff --git a/Engine/Plugins/Developer/Concert/ConcertApp/MultiUserClient/Source/MultiUserClient/Private/Replication/Util/RegularQueryService.h b/Engine/Plugins/Developer/Concert/ConcertApp/MultiUserClient/Source/MultiUserClient/Private/Replication/Util/RegularQueryService.h
--- a/Engine/Plugins/Developer/Concert/ConcertApp/MultiUserClient/Source/MultiUserClient/Private/Replication/Util/RegularQueryService.h
+++ b/Engine/Plugins/Developer/Concert/ConcertApp/MultiUserClient/Source/MultiUserClient/Private/Replication/Util/RegularQueryService.h
@@ -2,7 +2,6 @@
 
 #pragma once
 
-#include "IConcertSessionHandler.h"
 #include "Replication/Messages/ClientQuery.h"
 #include "Containers/Ticker.h"
 #include "Replication/IToken.h"
diff --git a/Engine/Plugins/Developer/Concert/ConcertApp/MultiUserClient/Source/MultiUserClient/Private/Widgets/ActiveSession/Replication/Client/Single/SOwnerClientList.cpp b/Engine/Plugins/Developer/Concert/ConcertApp/MultiUserClient/Source/MultiUserClient/Private/Widgets/ActiveSession/Replication/Client/Single/SOwnerClientList.cpp
new file mode 100644
--- /dev/null
+++ b/Engine/Plugins/Developer/Concert/ConcertApp/MultiUserClient/Source/MultiUserClient/Private/Widgets/ActiveSession/Replication/Client/Single/SOwnerClientList.cpp
@@ -0,0 +1,59 @@
+// Copyright Epic Games, Inc. All Rights Reserved.
+
+#include "SOwnerClientList.h"
+
+#include "Replication/Util/GlobalAuthorityCache.h"
+#include "Widgets/ClientName/SHorizontalClientList.h"
+
+#include "Widgets/Layout/SBox.h"
+#include "Widgets/Text/STextBlock.h"
+
+#define LOCTEXT_NAMESPACE "SOwnerClientList"
+
+namespace UE::MultiUserClient
+{
+	void SOwnerClientList::Construct(const FArguments& InArgs, TSharedRef<IConcertClient> InClient, FGlobalAuthorityCache& InAuthorityCache)
+	{
+		AuthorityCache = &InAuthorityCache;
+		GetClientListDelegate = InArgs._GetClientList;
+		check(GetClientListDelegate.IsBound());
+		
+		AuthorityCache->OnCacheChanged().AddSP(this, &SOwnerClientList::OnClientChanged);
+		
+		ChildSlot
+		[
+			SAssignNew(ClientList, ConcertClientSharedSlate::SHorizontalClientList, MoveTemp(InClient))
+			.HighlightText(InArgs._HighlightText)
+			.EmptyListSlot()
+			[
+				SNew(SBox)
+				.HAlign(HAlign_Left)
+				.Padding(5.f, 0.f, 0.f, 0.f)
+				[
+					SNew(STextBlock)
+					.Text(LOCTEXT("NotApplicable", "-"))
+				]
+			]
+		];
+
+		RefreshList();
+	}
+
+	SOwnerClientList::~SOwnerClientList()
+	{
+		AuthorityCache->OnCacheChanged().RemoveAll(this);
+	}
+
+	void SOwnerClientList::OnClientChanged(const FGuid& ChangedClientId)
+	{
+		RefreshList();
+	}
+	
+	void SOwnerClientList::RefreshList()
+	{
+		const TArray<FGuid> Clients = GetClientListDelegate.Execute(*AuthorityCache);
+		ClientList->RefreshList(Clients);
+	}
+}
+
+#undef LOCTEXT_NAMESPACE
\ No newline at end of file
diff --git a/Engine/Plugins/Developer/Concert/ConcertApp/MultiUserClient/Source/MultiUserClient/Private/Widgets/ActiveSession/Replication/Client/Single/SOwnerClientList.h b/Engine/Plugins/Developer/Concert/ConcertApp/MultiUserClient/Source/MultiUserClient/Private/Widgets/ActiveSession/Replication/Client/Single/SOwnerClientList.h
new file mode 100644
--- /dev/null
+++ b/Engine/Plugins/Developer/Concert/ConcertApp/MultiUserClient/Source/MultiUserClient/Private/Widgets/ActiveSession/Replication/Client/Single/SOwnerClientList.h
@@ -0,0 +1,56 @@
+// Copyright Epic Games, Inc. All Rights Reserved.
+
+#pragma once
+
+#include "Widgets/DeclarativeSyntaxSupport.h"
+#include "Widgets/SCompoundWidget.h"
+
+namespace UE::ConcertClientSharedSlate
+{
+	class SHorizontalClientList;
+}
+
+class IConcertClient;
+
+namespace UE::MultiUserClient
+{
+	class FGlobalAuthorityCache;
+	
+	/**
+	 * Displays a list of authoritive clients in the SReplicationClientView columns for top-level, subobject and properties;
+	 * the passed in FGetClientList determines which is displayed
+	 */
+	class SOwnerClientList : public SCompoundWidget
+	{
+	public:
+
+		DECLARE_DELEGATE_RetVal_OneParam(TArray<FGuid>, FGetClientList, const FGlobalAuthorityCache&);
+		
+		SLATE_BEGIN_ARGS(SOwnerClientList)
+		{}
+			/** Obtains the client list to display */
+			SLATE_EVENT(FGetClientList, GetClientList)
+			/** Used for highlighting in the text */
+			SLATE_ATTRIBUTE(FText, HighlightText)
+		SLATE_END_ARGS()
+
+		void Construct(const FArguments& InArgs, TSharedRef<IConcertClient> InClient, FGlobalAuthorityCache& InAuthorityCache);
+		virtual ~SOwnerClientList();
+
+	private:
+		
+		/** Tells us when the authority change */
+		FGlobalAuthorityCache* AuthorityCache;
+
+		/** Displays the clients */
+		TSharedPtr<ConcertClientSharedSlate::SHorizontalClientList> ClientList;
+
+		/** Used to get the client list from FGlobalAuthorityCache */
+		FGetClientList GetClientListDelegate;
+
+		void OnClientChanged(const FGuid& ChangedClientId);
+		void RefreshList();
+	};
+}
+
+
diff --git a/Engine/Plugins/Developer/Concert/ConcertApp/MultiUserClient/Source/MultiUserClient/Private/Widgets/ActiveSession/Replication/Client/Single/SReplicationClientView.cpp b/Engine/Plugins/Developer/Concert/ConcertApp/MultiUserClient/Source/MultiUserClient/Private/Widgets/ActiveSession/Replication/Client/Single/SReplicationClientView.cpp
--- a/Engine/Plugins/Developer/Concert/ConcertApp/MultiUserClient/Source/MultiUserClient/Private/Widgets/ActiveSession/Replication/Client/Single/SReplicationClientView.cpp
+++ b/Engine/Plugins/Developer/Concert/ConcertApp/MultiUserClient/Source/MultiUserClient/Private/Widgets/ActiveSession/Replication/Client/Single/SReplicationClientView.cpp
@@ -2,13 +2,14 @@
 
 #include "SReplicationClientView.h"
 
-#include "StreamEditorColumns.h"
+#include "SingleClientColumns.h"
 #include "Replication/Client/ReplicationClient.h"
 #include "Replication/Editor/Model/IEditableObjectToPropertiesModel.h"
 #include "Replication/Editor/Model/Object/EditorObjectSelectionSourceModel.h"
 #include "Replication/Editor/Model/Property/SelectPropertyFromUClassModel.h"
 #include "Replication/Editor/View/IReplicationStreamEditor.h"
 #include "Replication/ReplicationWidgetFactories.h"
+#include "Replication/Client/ReplicationClientManager.h"
 #include "Replication/Submission/ISubmissionWorkflow.h"
 #include "Widgets/ActiveSession/Replication/Client/Single/SSingleClientToolbar.h"
 
@@ -18,7 +19,7 @@
 
 namespace UE::MultiUserClient
 {
-	void SReplicationClientView::Construct(const FArguments& InArgs)
+	void SReplicationClientView::Construct(const FArguments& InArgs, const TSharedRef<IConcertClient>& InClient, FReplicationClientManager& InClientManager)
 	{
 		GetReplicationClientAttribute = InArgs._GetReplicationClient;
 		FReplicationClient* ReplicationClient = GetReplicationClientAttribute.Get();
@@ -27,12 +28,19 @@ namespace UE::MultiUserClient
 		ConcertClientSharedSlate::IObjectToPropertiesModel& PropertyModel = *ReplicationClient->GetClientEditModel();
 		FAuthorityChangeTracker& AuthorityTracker = ReplicationClient->GetAuthorityDiffer();
 		ISubmissionWorkflow& SubmissionWorkflow = ReplicationClient->GetSubmissionWorkflow();
+		FGlobalAuthorityCache& AuthorityCache = InClientManager.GetAuthorityCache();
+		TAttribute<const ConcertClientSharedSlate::IReplicationStreamViewer*> GetReplicationViewerAttribute =
+			TAttribute<const ConcertClientSharedSlate::IReplicationStreamViewer*>::CreateLambda([this](){ return EditorView.Get(); });
 
 		// Add checkboxes in front of top level and subobject rows for changing authority
 		using namespace ConcertClientSharedSlate;
 		const FCreateSubobjectViewParams SubobjectViewParams
 		{
-			.AdditionalColumns = { StreamEditorColumns::ToggleSubobjectAuthority(AuthorityTracker, SubmissionWorkflow) }
+			.AdditionalColumns =
+			{
+				SingleClientColumns::ToggleSubobjectAuthority(AuthorityTracker, SubmissionWorkflow),
+				SingleClientColumns::OwnerOfSubobject(InClient, AuthorityCache)
+			}
 		};
 		const FCreateEditorParams ReplicationEditorCreationParams
 		{
@@ -40,7 +48,15 @@ namespace UE::MultiUserClient
 			.ObjectSource = MakeShared<FEditorObjectSelectionSourceModel>(),
 			.PropertySource = MakeShared<FSelectPropertyFromUClassModel>(),
 			.SubobjectView = CreateUnrealEditorSubobjectView(SubobjectViewParams),
-			.AdditionalObjectColumns = { StreamEditorColumns::ToggleTopLevelAuthority(PropertyModel, AuthorityTracker, SubmissionWorkflow) },
+			.AdditionalObjectColumns =
+			{
+				SingleClientColumns::ToggleTopLevelAuthority(PropertyModel, AuthorityTracker, SubmissionWorkflow),
+				SingleClientColumns::OwnerOfTopLevelObject(InClient, AuthorityCache, PropertyModel)
+			},
+			.AdditionalPropertyColumns =
+			{
+				SingleClientColumns::OwnerOfProperty(InClient, AuthorityCache, GetReplicationViewerAttribute)
+			},
 			.IsEditingEnabled = TAttribute<bool>::CreateLambda([&SubmissionWorkflow](){ return SubmissionWorkflow.GetUploadability() != EChangeUploadability::NotImplemented; }),
 			.EditingDisabledToolTipText = LOCTEXT("Editing.NotImplemented", "Editing remote clients is not implemented. You can only edit the local client.")
 		};
diff --git a/Engine/Plugins/Developer/Concert/ConcertApp/MultiUserClient/Source/MultiUserClient/Private/Widgets/ActiveSession/Replication/Client/Single/SReplicationClientView.h b/Engine/Plugins/Developer/Concert/ConcertApp/MultiUserClient/Source/MultiUserClient/Private/Widgets/ActiveSession/Replication/Client/Single/SReplicationClientView.h
--- a/Engine/Plugins/Developer/Concert/ConcertApp/MultiUserClient/Source/MultiUserClient/Private/Widgets/ActiveSession/Replication/Client/Single/SReplicationClientView.h
+++ b/Engine/Plugins/Developer/Concert/ConcertApp/MultiUserClient/Source/MultiUserClient/Private/Widgets/ActiveSession/Replication/Client/Single/SReplicationClientView.h
@@ -5,6 +5,8 @@
 #include "Widgets/DeclarativeSyntaxSupport.h"
 #include "Widgets/SCompoundWidget.h"
 
+class IConcertClient;
+
 namespace UE::ConcertClientSharedSlate
 {
 	class IReplicationStreamEditor;
@@ -13,6 +15,7 @@ namespace UE::ConcertClientSharedSlate
 namespace UE::MultiUserClient
 {
 	class FReplicationClient;
+	class FReplicationClientManager;
 	
 	/** Displays the contents of a client. */
 	class SReplicationClientView : public SCompoundWidget
@@ -21,13 +24,13 @@ namespace UE::MultiUserClient
 
 		SLATE_BEGIN_ARGS(SReplicationClientView)
 		{}
-			/** The client to depict. Should always return true. If the client is destroyed, so should this widget be. */
+			/** The client to depict. Should always return something valid. If the client is destroyed, so should this widget be. */
 			SLATE_ATTRIBUTE(FReplicationClient*, GetReplicationClient)
 			/** Dedicated space for a widget with which to change the view. */
 			SLATE_NAMED_SLOT(FArguments, ViewSelectionArea)
 		SLATE_END_ARGS()
 
-		void Construct(const FArguments& InArgs);
+		void Construct(const FArguments& InArgs, const TSharedRef<IConcertClient>& InClient, FReplicationClientManager& InClientManager);
 
 	private:
 		
diff --git a/Engine/Plugins/Developer/Concert/ConcertApp/MultiUserClient/Source/MultiUserClient/Private/Widgets/ActiveSession/Replication/Client/Single/SingleClientColumns.h b/Engine/Plugins/Developer/Concert/ConcertApp/MultiUserClient/Source/MultiUserClient/Private/Widgets/ActiveSession/Replication/Client/Single/SingleClientColumns.h
new file mode 100644
--- /dev/null
+++ b/Engine/Plugins/Developer/Concert/ConcertApp/MultiUserClient/Source/MultiUserClient/Private/Widgets/ActiveSession/Replication/Client/Single/SingleClientColumns.h
@@ -0,0 +1,120 @@
+// Copyright Epic Games, Inc. All Rights Reserved.
+
+#pragma once
+
+#include "Replication/Editor/View/PredefinedReplicationColumns.h"
+#include "Widgets/DeclarativeSyntaxSupport.h"
+#include "Widgets/SCompoundWidget.h"
+
+class IConcertClient;
+
+namespace UE::ConcertClientSharedSlate
+{
+	class IReplicationStreamViewer;
+	class IObjectToPropertiesModel;
+}
+
+namespace UE::MultiUserClient
+{
+	class FAuthorityChangeTracker;
+	class FGlobalAuthorityCache;
+	class IClientAuthoritySynchronizer;
+	class ISubmissionWorkflow;
+}
+
+namespace UE::MultiUserClient::SingleClientColumns
+{
+	/** @see UE::ConcertClientSharedSlate::ReplicationColumns::Subobject::ETopLevelColumnOrder */
+	enum class ETopLevelObjectColumnOrder
+	{
+		ToggleAuthority = 0,
+		Owner = 40
+	};
+	/** @see UE::ConcertClientSharedSlate::ReplicationColumns::Subobject::ESubobjectColumnOrder */
+	enum class ESubobjectColumnOrder
+	{
+		ToggleAuthority = 0,
+		Owner = 20
+	};
+	/** @see UE::ConcertClientSharedSlate::ReplicationColumns::Property::EReplicationPropertyColumnOrder */
+	enum class EPropertyColumnOrder
+	{
+		Owner = 50
+	};
+
+
+	/********** Toggle authority **********/
+	extern const FName ToggleTopLevelAuthorityColumnId;
+	extern const FName ToggleSubobjectAuthorityColumnId;
+
+	/**
+	 * Checkbox placed in the top-level view.
+	 * It gives / removes authority for the top-level object and all of its subobjects.
+	 *
+	 * @param ClientStreamModel Used to discover the client's registered subobjects
+	 * @param ChangeTracker Used to determine the checkbox state and change authority.
+	 * @param SubmissionWorkflow Used to determine whether changing authority is at all enabled
+	 * @return Column that can be placed in the table
+	 */
+	ConcertClientSharedSlate::ReplicationColumns::FReplicationTopLevelObjectColumn ToggleTopLevelAuthority(
+		ConcertClientSharedSlate::IObjectToPropertiesModel& ClientStreamModel,
+		FAuthorityChangeTracker& ChangeTracker,
+		ISubmissionWorkflow& SubmissionWorkflow
+		);
+
+	/**
+	 * Checkbox placed in the subobject view
+	 * It gives / removes authority for the object it is placed next to.
+	 * 
+	 * @param AuthoritySynchronizer Determines whether the client has authority or not.
+	 * @param ChangeTracker Used to determine the checkbox state and change authority.
+	 * @param AuthoritySynchronizer Used to determine whether changing authority is at all enabled
+	 * @return Column that can be placed in the table
+	 */
+	ConcertClientSharedSlate::ReplicationColumns::FReplicationSubobjectObjectColumn ToggleSubobjectAuthority(
+		FAuthorityChangeTracker& ChangeTracker,
+		ISubmissionWorkflow& SubmissionWorkflow
+		);
+
+	
+	/********** Owner **********/
+	extern const FName OwnerOfSubobjectColumnId;
+	extern const FName OwnerOfPropertyColumnId;
+	
+	/**
+	 * Displays the owner of a top-level object.
+	 * @param InClient The local Concert client used to look up other client display info
+	 * @param InAuthorityCache Used to determine which client owns the property
+	 * @param InObjectModel Used to get subobjects for a top-level object
+	 * @return Column that can be placed in the table
+	 */
+	ConcertClientSharedSlate::ReplicationColumns::FReplicationTopLevelObjectColumn OwnerOfTopLevelObject(
+		const TSharedRef<IConcertClient>& InClient,
+		FGlobalAuthorityCache& InAuthorityCache,
+		ConcertClientSharedSlate::IObjectToPropertiesModel& InObjectModel
+		);
+
+	/**
+	 * Displays the owner of a subobject.
+	 * @param InClient The local Concert client used to look up other client display info
+	 * @param InAuthorityCache Used to determine which client owns the property
+	 * @return Column that can be placed in the table
+	 */
+	ConcertClientSharedSlate::ReplicationColumns::FReplicationSubobjectObjectColumn OwnerOfSubobject(
+		const TSharedRef<IConcertClient>& InClient,
+		FGlobalAuthorityCache& InAuthorityCache
+		);
+	
+	/**
+	 * Displays the owner of a property.
+	 * @param InClient The local Concert client used to look up other client display info
+	 * @param InAuthorityCache Used to determine which client owns the property
+	 * @param InViewerAttribute Used to determine which objects the property box is displaying
+	 * @return Column that can be placed in the table
+	 */
+	ConcertClientSharedSlate::ReplicationColumns::FReplicationPropertyColumn OwnerOfProperty(
+		const TSharedRef<IConcertClient>& InClient,
+		FGlobalAuthorityCache& InAuthorityCache,
+		const TAttribute<const ConcertClientSharedSlate::IReplicationStreamViewer*>& InViewerAttribute
+		);
+}
diff --git a/Engine/Plugins/Developer/Concert/ConcertApp/MultiUserClient/Source/MultiUserClient/Private/Widgets/ActiveSession/Replication/Client/Single/SingleClientColumns_Owner.cpp b/Engine/Plugins/Developer/Concert/ConcertApp/MultiUserClient/Source/MultiUserClient/Private/Widgets/ActiveSession/Replication/Client/Single/SingleClientColumns_Owner.cpp
new file mode 100644
--- /dev/null
+++ b/Engine/Plugins/Developer/Concert/ConcertApp/MultiUserClient/Source/MultiUserClient/Private/Widgets/ActiveSession/Replication/Client/Single/SingleClientColumns_Owner.cpp
@@ -0,0 +1,184 @@
+// Copyright Epic Games, Inc. All Rights Reserved.
+
+#include "SingleClientColumns.h"
+
+#include "IConcertClient.h"
+#include "SOwnerClientList.h"
+#include "Replication/Editor/View/IReplicationStreamViewer.h"
+#include "Replication/Editor/View/ReplicationColumnsUtils.h"
+#include "Replication/Util/GlobalAuthorityCache.h"
+#include "Widgets/ClientName/SClientName.h"
+
+#define LOCTEXT_NAMESPACE "SingleClientColumns.Owner"
+
+namespace UE::MultiUserClient::SingleClientColumns
+{
+	const FName OwnerOfSubobjectColumnId = TEXT("OwnerOfSubobjectColumn");
+	const FName OwnerOfPropertyColumnId = TEXT("OwnerOfPropertyColumn");
+
+	namespace Private::Shared
+	{
+		constexpr float DefaultWidth = 200.f;
+		
+		static void PopulateSearchTerms(const TSharedRef<IConcertClient>& InClient, const TArray<FGuid>& Clients, TArray<FString>& InOutSearchStrings)
+		{
+			for (const FGuid& ClientId : Clients)
+			{
+				FConcertSessionClientInfo ClientInfo;
+				const TSharedPtr<IConcertClientSession> Session = InClient->GetCurrentSession();
+				const bool bIsLocalClient = ensure(Session) && Session->GetSessionClientEndpointId() == ClientId;
+				if (bIsLocalClient)
+				{
+					InOutSearchStrings.Add(ConcertClientSharedSlate::SClientName::GetDisplayText(Session->GetLocalClientInfo(), bIsLocalClient).ToString());
+				}
+				else if (ensure(Session) && Session->FindSessionClient(ClientId, ClientInfo))
+				{
+					InOutSearchStrings.Add(ConcertClientSharedSlate::SClientName::GetDisplayText(ClientInfo.ClientInfo, bIsLocalClient).ToString());
+				}
+			}
+		}
+	}
+
+	namespace Private::TopLevel
+	{
+		static TArray<FGuid> GetOwnersOfTopLevelObject(
+			const FGlobalAuthorityCache& InAuthorityCache,
+			ConcertClientSharedSlate::IObjectToPropertiesModel& InObjectModel,
+			const FSoftObjectPath& InTopLevelObject
+			)
+		{
+			TArray<FGuid> Owners = InAuthorityCache.GetClientsWithAuthorityOverObject(InTopLevelObject);
+			InObjectModel.ForEachSubobject(InTopLevelObject, [&InAuthorityCache, &Owners](const FSoftObjectPath& Subobject)
+			{
+				for (const FGuid& Client : InAuthorityCache.GetClientsWithAuthorityOverObject(Subobject))
+				{
+					Owners.AddUnique(Client);
+				}
+				return EBreakBehavior::Continue;
+			});
+			return Owners;
+		}
+	}
+	
+	ConcertClientSharedSlate::ReplicationColumns::FReplicationTopLevelObjectColumn OwnerOfTopLevelObject(
+		const TSharedRef<IConcertClient>& InClient,
+		FGlobalAuthorityCache& InAuthorityCache,
+		ConcertClientSharedSlate::IObjectToPropertiesModel& InObjectModel
+		)
+	{
+		using FColumnType = ConcertClientSharedSlate::ReplicationColumns::FReplicationSubobjectObjectColumn;
+		return FColumnType(
+			FColumnType::FArguments()
+				.GenerateWidgetColumn_Lambda([InClient, &InObjectModel, &InAuthorityCache](const FColumnType::FBuildArgs& InArgs)
+				{
+					return SNew(SOwnerClientList, InClient, InAuthorityCache)
+						.GetClientList_Lambda([&InAuthorityCache, &InObjectModel, ObjectPath = InArgs.RowData.GetObjectPath()](const FGlobalAuthorityCache&)
+						{
+							return Private::TopLevel::GetOwnersOfTopLevelObject(InAuthorityCache, InObjectModel, ObjectPath);
+						})
+						.HighlightText_Lambda([HighlightText = InArgs.HighlightText](){ return *HighlightText.Get(); });
+				})
+				.PopulateSearchItems_Lambda([InClient, &InAuthorityCache, &InObjectModel](const ConcertClientSharedSlate::FReplicatedObjectData& InArgs, TArray<FString>& InOutSearchStrings)
+				{
+					Private::Shared::PopulateSearchTerms(
+						InClient,
+						Private::TopLevel::GetOwnersOfTopLevelObject(InAuthorityCache, InObjectModel, InArgs.GetObjectPath()),
+						InOutSearchStrings
+						);
+				})
+				.ColumnSortOrder(static_cast<int32>(ETopLevelObjectColumnOrder::Owner)),
+			SHeaderRow::Column(OwnerOfSubobjectColumnId)
+				.DefaultLabel(LOCTEXT("Subobject.Owner", "Owner"))
+				.FillSized(Private::Shared::DefaultWidth)
+			);
+	}
+	
+	ConcertClientSharedSlate::ReplicationColumns::FReplicationSubobjectObjectColumn OwnerOfSubobject(
+		const TSharedRef<IConcertClient>& InClient,
+		FGlobalAuthorityCache& InAuthorityCache
+		)
+	{
+		using FColumnType = ConcertClientSharedSlate::ReplicationColumns::FReplicationSubobjectObjectColumn;
+		return FColumnType(
+			FColumnType::FArguments()
+				.GenerateWidgetColumn_Lambda([InClient, &InAuthorityCache](const FColumnType::FBuildArgs& InArgs)
+				{
+					return SNew(SOwnerClientList, InClient, InAuthorityCache)
+						.GetClientList_Lambda([&InAuthorityCache, ObjectPath = InArgs.RowData.GetObjectPath()](const FGlobalAuthorityCache&)
+						{
+							return InAuthorityCache.GetClientsWithAuthorityOverObject(ObjectPath);
+						})
+						.HighlightText_Lambda([HighlightText = InArgs.HighlightText](){ return *HighlightText.Get(); });
+				})
+				.PopulateSearchItems_Lambda([InClient, &InAuthorityCache](const ConcertClientSharedSlate::FReplicatedObjectData& InArgs, TArray<FString>& InOutSearchStrings)
+				{
+					Private::Shared::PopulateSearchTerms(
+						InClient,
+						InAuthorityCache.GetClientsWithAuthorityOverObject(InArgs.GetObjectPath()),
+						InOutSearchStrings
+						);
+				})
+				.ColumnSortOrder(static_cast<int32>(ESubobjectColumnOrder::Owner)),
+			SHeaderRow::Column(OwnerOfSubobjectColumnId)
+				.DefaultLabel(LOCTEXT("Subobject.Owner", "Owner"))
+				.FillSized(Private::Shared::DefaultWidth)
+			);
+	}
+	
+	namespace Private::Property
+	{
+		static TArray<FGuid> GetPropertyOwners(
+			const FGlobalAuthorityCache& InAuthorityCache,
+			const ConcertClientSharedSlate::IReplicationStreamViewer& InViewer,
+			const FConcertPropertyChain& Property
+			)
+		{
+			TArray<FGuid> Result;
+			const TArray<FSoftObjectPath> SelectedObjects = InViewer.GetObjectsBeingPropertyEdited();
+			for (const FSoftObjectPath& SelectedObject : SelectedObjects)
+			{
+				const TOptional<FGuid> Owner = InAuthorityCache.GetClientWithAuthorityOverProperty(SelectedObject, Property);
+				if (Owner)
+				{
+					Result.AddUnique(*Owner);
+				}
+			}
+			return Result;
+		}
+	}
+
+	ConcertClientSharedSlate::ReplicationColumns::FReplicationPropertyColumn OwnerOfProperty(
+		const TSharedRef<IConcertClient>& InClient,
+		FGlobalAuthorityCache& InAuthorityCache,
+		const TAttribute<const ConcertClientSharedSlate::IReplicationStreamViewer*>& InViewer
+		)
+	{
+		using FColumnType = ConcertClientSharedSlate::ReplicationColumns::FReplicationPropertyColumn;
+		return FColumnType(
+			FColumnType::FArguments()
+				.GenerateWidgetColumn_Lambda([InClient, &InAuthorityCache, InViewer](const FColumnType::FBuildArgs& InArgs)
+				{
+					return SNew(SOwnerClientList, InClient, InAuthorityCache)
+						.GetClientList_Lambda([&InAuthorityCache, InViewer, Property = InArgs.RowData.GetProperty()](const FGlobalAuthorityCache&)
+						{
+							return Private::Property::GetPropertyOwners(InAuthorityCache, *InViewer.Get(), Property);
+						})
+						.HighlightText_Lambda([HighlightText = InArgs.HighlightText](){ return *HighlightText.Get(); });
+				})
+				.PopulateSearchItems_Lambda([InClient, &InAuthorityCache, InViewer](const ConcertClientSharedSlate::FReplicatedPropertyData& InArgs, TArray<FString>& InOutSearchStrings)
+				{
+					Private::Shared::PopulateSearchTerms(
+						InClient,
+						Private::Property::GetPropertyOwners(InAuthorityCache, *InViewer.Get(), InArgs.GetProperty()),
+						InOutSearchStrings
+						);
+				})
+				.ColumnSortOrder(static_cast<int32>(EPropertyColumnOrder::Owner)),
+			SHeaderRow::Column(OwnerOfPropertyColumnId)
+				.DefaultLabel(LOCTEXT("Property.Owner", "Owner"))
+				.FillSized(Private::Shared::DefaultWidth)
+			);
+	}
+}
+
+#undef LOCTEXT_NAMESPACE
\ No newline at end of file
diff --git a/Engine/Plugins/Developer/Concert/ConcertApp/MultiUserClient/Source/MultiUserClient/Private/Widgets/ActiveSession/Replication/Client/Single/SingleClientColumns_ToggleAuthority.cpp b/Engine/Plugins/Developer/Concert/ConcertApp/MultiUserClient/Source/MultiUserClient/Private/Widgets/ActiveSession/Replication/Client/Single/SingleClientColumns_ToggleAuthority.cpp
new file mode 100644
--- /dev/null
+++ b/Engine/Plugins/Developer/Concert/ConcertApp/MultiUserClient/Source/MultiUserClient/Private/Widgets/ActiveSession/Replication/Client/Single/SingleClientColumns_ToggleAuthority.cpp
@@ -0,0 +1,220 @@
+// Copyright Epic Games, Inc. All Rights Reserved.
+
+#include "SingleClientColumns.h"
+
+#include "IConcertClient.h"
+#include "SOwnerClientList.h"
+#include "Replication/Authority/AuthorityChangeTracker.h"
+#include "Replication/Authority/EAuthorityMutability.h"
+#include "Replication/Authority/IClientAuthoritySynchronizer.h"
+#include "Replication/Editor/Model/IObjectToPropertiesModel.h"
+#include "Replication/Editor/View/IReplicationStreamViewer.h"
+#include "Replication/Editor/View/ReplicationColumnsUtils.h"
+#include "Replication/Submission/ISubmissionWorkflow.h"
+#include "Replication/Util/GlobalAuthorityCache.h"
+#include "Widgets/ClientName/SClientName.h"
+
+#define LOCTEXT_NAMESPACE "SingleClientColumns.ToggleAuthority"
+
+namespace UE::MultiUserClient::SingleClientColumns
+{
+	const FName ToggleTopLevelAuthorityColumnId = TEXT("ToggleTopLevelAuthorityColumn");
+	const FName ToggleSubobjectAuthorityColumnId = TEXT("ToggleSubobjectAuthorityColumn");
+
+	namespace Private::ToggleTopLevelAuthority
+	{
+		static ECheckBoxState GetCheckboxState(
+			const ConcertClientSharedSlate::FReplicatedObjectData& ObjectData,
+			ConcertClientSharedSlate::IObjectToPropertiesModel* ClientStreamModel,
+			FAuthorityChangeTracker* ChangeTracker
+			)
+		{
+			const FSoftObjectPath& ObjectPath = ObjectData.GetObjectPath();
+
+			// Init the checkbox state to the top level object, if the top level object has properties associated with it...
+			const bool bCanSetTopLevelState = ChangeTracker->CanSetAuthorityFor(ObjectPath)
+				|| ChangeTracker->GetChangeAuthorityMutability(ObjectPath) == EAuthorityMutability::NotSupported;
+			TOptional<ECheckBoxState> CheckBoxState = bCanSetTopLevelState
+				? ChangeTracker->GetAuthorityStateAfterApplied(ObjectPath) ? ECheckBoxState::Checked : ECheckBoxState::Unchecked
+				: TOptional<ECheckBoxState>{};
+
+			// ... and then determine whether all the other subobjects share that state (Check / Unchecked) or not (Undetermined)
+			ClientStreamModel->ForEachSubobject(ObjectPath, [&CheckBoxState, &ChangeTracker](const FSoftObjectPath& Child)
+			{
+				if (!ChangeTracker->CanSetAuthorityFor(Child)
+					&& ChangeTracker->GetChangeAuthorityMutability(Child) != EAuthorityMutability::NotSupported)
+				{
+					return EBreakBehavior::Continue;
+				}
+
+				// Init checkbox state
+				const bool bCurrentAuthorityState = ChangeTracker->GetAuthorityStateAfterApplied(Child);
+				const bool bIsFirstObject = !CheckBoxState.IsSet(); 
+				if (bIsFirstObject)
+				{
+					CheckBoxState = bCurrentAuthorityState ? ECheckBoxState::Checked : ECheckBoxState::Unchecked;
+					return EBreakBehavior::Continue;
+				}
+
+				// Return the checkbox state that all subobjects share, or Undetermined if mixed
+				const bool bMustBeUndetermined = (*CheckBoxState == ECheckBoxState::Unchecked && bCurrentAuthorityState)
+					|| (*CheckBoxState == ECheckBoxState::Checked && !bCurrentAuthorityState);
+				if (bMustBeUndetermined)
+				{
+					CheckBoxState = ECheckBoxState::Undetermined;
+					return EBreakBehavior::Break;
+				}
+				return EBreakBehavior::Continue;
+			});
+						
+			return CheckBoxState
+				? *CheckBoxState
+				// If this case happens, there are no objects to check. In that case, it visually "looks better" to be unchecked
+				: ECheckBoxState::Unchecked;
+		}
+
+		static void OnCheckboxStateChanged(
+			bool bIsChecked,
+			const ConcertClientSharedSlate::FReplicatedObjectData& ObjectData,
+			ConcertClientSharedSlate::IObjectToPropertiesModel* ClientStreamModel,
+			FAuthorityChangeTracker* ChangeTracker
+			)
+		{
+			const FSoftObjectPath& ObjectPath = ObjectData.GetObjectPath();
+			TArray<FSoftObjectPath> Paths { ObjectPath };
+			ClientStreamModel->ForEachSubobject(ObjectPath, [&Paths](const FSoftObjectPath& Child)
+			{
+				Paths.Add(Child);
+				return EBreakBehavior::Continue;
+			});
+			
+			ChangeTracker->SetAuthorityIfAllowed(Paths, bIsChecked);
+		}
+
+		static bool IsEnabled(
+			const ConcertClientSharedSlate::FReplicatedObjectData& ObjectData,
+			ConcertClientSharedSlate::IObjectToPropertiesModel* ClientStreamModel,
+			FAuthorityChangeTracker* ChangeTracker
+			)
+		{
+			const FSoftObjectPath& TopLevelObjectPath = ObjectData.GetObjectPath();
+			const bool bCanChangeTopLevel = ChangeTracker->CanSetAuthorityFor(TopLevelObjectPath);
+			if (bCanChangeTopLevel)
+			{
+				return true;
+			}
+						
+			bool bCanChangeAnySubobject = false;
+			ClientStreamModel->ForEachSubobject(TopLevelObjectPath, [&ChangeTracker, &bCanChangeAnySubobject](const FSoftObjectPath& Child)
+			{
+				bCanChangeAnySubobject |= ChangeTracker->CanSetAuthorityFor(Child);
+				return bCanChangeAnySubobject ? EBreakBehavior::Break : EBreakBehavior::Continue;
+			});
+			return bCanChangeAnySubobject;
+		}
+	}
+	
+	ConcertClientSharedSlate::ReplicationColumns::FReplicationTopLevelObjectColumn ToggleTopLevelAuthority(
+		ConcertClientSharedSlate::IObjectToPropertiesModel& ClientStreamModel,
+		FAuthorityChangeTracker& ChangeTracker,
+		ISubmissionWorkflow& SubmissionWorkflow
+		)
+	{
+		using namespace ConcertClientSharedSlate;
+		using namespace Private::ToggleTopLevelAuthority;
+		using FTopLevelColumnDelegates = TReplicationColumnDelegates<FReplicatedObjectData>;
+		
+		return MakeCheckboxColumn<FReplicatedObjectData>(
+			ToggleTopLevelAuthorityColumnId,
+				FTopLevelColumnDelegates(
+					FTopLevelColumnDelegates::FGetColumnCheckboxState::CreateStatic(&GetCheckboxState, &ClientStreamModel, &ChangeTracker),
+					FTopLevelColumnDelegates::FOnColumnCheckboxChanged::CreateStatic(&OnCheckboxStateChanged, &ClientStreamModel, &ChangeTracker),
+					FTopLevelColumnDelegates::FGetToolTipText::CreateLambda([&ClientStreamModel, &ChangeTracker, &SubmissionWorkflow](const FReplicatedObjectData& ObjectData)
+					{
+						
+						if (SubmissionWorkflow.GetUploadability() == EChangeUploadability::NotImplemented)
+						{
+							const ECheckBoxState CheckBoxState = GetCheckboxState(ObjectData, &ClientStreamModel, &ChangeTracker);
+							const FText CheckBoxText = [CheckBoxState]()
+							{
+								switch (CheckBoxState)
+								{
+								case ECheckBoxState::Unchecked: return LOCTEXT("TopLevel.ChangeAuthority.ToolTip.NotSupported.Unchecked", "None of the subobjects are being replicated.");
+								case ECheckBoxState::Checked: return LOCTEXT("TopLevel.ChangeAuthority.ToolTip.NotSupported.Checked", "All of the subobjects are being replicated.");
+								case ECheckBoxState::Undetermined: return LOCTEXT("TopLevel.ChangeAuthority.ToolTip.NotSupported.Undetermined", "Some of the subobjects are being replicated.");
+								default: checkNoEntry(); return FText::GetEmpty();
+								}
+							}();
+							
+							return FText::Format(
+								LOCTEXT("TopLevel.ChangeAuthority.ToolTip.NotSupportedFmt", "{0}\nEditing remote clients is not implemented. You can only edit the local client."),
+								CheckBoxText
+								);
+						}
+						
+						return IsEnabled(ObjectData, &ClientStreamModel, &ChangeTracker)
+							? LOCTEXT("TopLevel.ChangeAuthority.Allowed", "Toggles whether this object and its subobjects will replicate the assigned properties.")
+							: LOCTEXT("TopLevel.ChangeAuthority.Disallowed", "Toggles whether this object and its subobjects will replicate the assigned properties.\nDisabled because none of the subobjects can be replicated with their current configuration.");
+					}),
+					FTopLevelColumnDelegates::FIsEnabled::CreateStatic(&IsEnabled, &ClientStreamModel, &ChangeTracker)
+				),
+			FText::GetEmpty(),
+			static_cast<int32>(ETopLevelObjectColumnOrder::ToggleAuthority)
+		);
+	}
+
+	ConcertClientSharedSlate::ReplicationColumns::FReplicationSubobjectObjectColumn ToggleSubobjectAuthority(
+		FAuthorityChangeTracker& ChangeTracker,
+		ISubmissionWorkflow& SubmissionWorkflow
+		)
+	{
+		using namespace ConcertClientSharedSlate;
+		using FSubobjectColumnDelegates = TReplicationColumnDelegates<FReplicatedObjectData>;
+		return MakeCheckboxColumn<FReplicatedObjectData>(
+			ToggleTopLevelAuthorityColumnId,
+				FSubobjectColumnDelegates(
+					FSubobjectColumnDelegates::FGetColumnCheckboxState::CreateLambda(
+					[&ChangeTracker](const FReplicatedObjectData& ObjectData)
+					{
+						const bool bHasAuthority = ChangeTracker.GetAuthorityStateAfterApplied(ObjectData.GetObjectPath());
+						return bHasAuthority ? ECheckBoxState::Checked : ECheckBoxState::Unchecked;
+					}),
+					FSubobjectColumnDelegates::FOnColumnCheckboxChanged::CreateLambda(
+					[&ChangeTracker](bool bIsChecked, const FReplicatedObjectData& ObjectData)
+					{
+						ChangeTracker.SetAuthorityIfAllowed({ ObjectData.GetObjectPath() }, bIsChecked);
+					}),
+					FSubobjectColumnDelegates::FGetToolTipText::CreateLambda([&ChangeTracker, &SubmissionWorkflow](const FReplicatedObjectData& ObjectData)
+					{
+						const bool bHasAuthority = ChangeTracker.GetAuthorityStateAfterApplied(ObjectData.GetObjectPath());
+						const FText NotSupportedText = FText::Format(
+							LOCTEXT("Subobject.ChangeAuthority.ToolTip.NotSupportedFmt", "{0}\nEditing remote clients is not implemented. You can only edit the local client."),
+							bHasAuthority ? LOCTEXT("Subobject.ChangeAuthority.Replicating.True", "This object is being replicated.") : LOCTEXT("Subobject.ChangeAuthority.Replicating.False", "This object is not currently being replicated.")
+							);
+						if (SubmissionWorkflow.GetUploadability() == EChangeUploadability::NotImplemented)
+						{
+							return NotSupportedText;
+						}
+						
+						switch (ChangeTracker.GetChangeAuthorityMutability(ObjectData.GetObjectPath()))
+						{
+						case EAuthorityMutability::Allowed: return LOCTEXT("Subobject.ChangeAuthority.ToolTip.Allowed", "Toggles whether this object should be replicated.");
+						case EAuthorityMutability::NoProperties: return LOCTEXT("Subobject.ChangeAuthority.ToolTip.NoProperties", "Toggles whether this object should be replicated.\nAssign properties to this object first.\n");
+						// TODO DP UE-198356: Make the tooltip display info about the conflicting client
+						case EAuthorityMutability::ClientConflict: return LOCTEXT("Subobject.ChangeAuthority.ToolTip.ClientConflict", "Toggles whether this object should be replicated.\nAnother client already is already replicating some of the assigned properties.");
+						case EAuthorityMutability::NotSupported: return NotSupportedText;
+						default: checkNoEntry(); return FText::GetEmpty();
+						}
+					}),
+					FSubobjectColumnDelegates::FIsEnabled::CreateLambda([&ChangeTracker](const FReplicatedObjectData& ObjectData)
+					{
+						return ChangeTracker.CanSetAuthorityFor(ObjectData.GetObjectPath());
+					})
+				),
+			FText::GetEmpty(),
+			static_cast<int32>(ESubobjectColumnOrder::ToggleAuthority)
+		);
+	}
+}
+
+#undef LOCTEXT_NAMESPACE
\ No newline at end of file
diff --git a/Engine/Plugins/Developer/Concert/ConcertApp/MultiUserClient/Source/MultiUserClient/Private/Widgets/ActiveSession/Replication/Client/Single/StreamEditorColumns.cpp b/Engine/Plugins/Developer/Concert/ConcertApp/MultiUserClient/Source/MultiUserClient/Private/Widgets/ActiveSession/Replication/Client/Single/StreamEditorColumns.cpp
deleted file mode 100644
diff --git a/Engine/Plugins/Developer/Concert/ConcertApp/MultiUserClient/Source/MultiUserClient/Private/Widgets/ActiveSession/Replication/Client/Single/StreamEditorColumns.h b/Engine/Plugins/Developer/Concert/ConcertApp/MultiUserClient/Source/MultiUserClient/Private/Widgets/ActiveSession/Replication/Client/Single/StreamEditorColumns.h
deleted file mode 100644
diff --git a/Engine/Plugins/Developer/Concert/ConcertApp/MultiUserClient/Source/MultiUserClient/Private/Widgets/ActiveSession/Replication/Joined/SReplicationJoinedView.cpp b/Engine/Plugins/Developer/Concert/ConcertApp/MultiUserClient/Source/MultiUserClient/Private/Widgets/ActiveSession/Replication/Joined/SReplicationJoinedView.cpp
--- a/Engine/Plugins/Developer/Concert/ConcertApp/MultiUserClient/Source/MultiUserClient/Private/Widgets/ActiveSession/Replication/Joined/SReplicationJoinedView.cpp
+++ b/Engine/Plugins/Developer/Concert/ConcertApp/MultiUserClient/Source/MultiUserClient/Private/Widgets/ActiveSession/Replication/Joined/SReplicationJoinedView.cpp
@@ -2,17 +2,14 @@
 
 #include "SReplicationJoinedView.h"
 
-#include "ConcertLogGlobal.h"
 #include "IConcertSyncClient.h"
 #include "Replication/MultiUserReplicationManager.h"
 #include "SSelectClientViewComboButton.h"
 #include "Widgets/ActiveSession/Replication/Client/Single/SReplicationClientView.h"
 
-#include "Framework/Notifications/NotificationManager.h"
 #include "Replication/Client/RemoteReplicationClient.h"
 #include "Widgets/Layout/SBox.h"
 #include "Widgets/Layout/SWidgetSwitcher.h"
-#include "Widgets/Notifications/SNotificationList.h"
 #include "Widgets/SBoxPanel.h"
 #include "Widgets/Text/STextBlock.h"
 
@@ -34,7 +31,7 @@ namespace UE::MultiUserClient
 			SAssignNew(ClientViewSwitcher, SWidgetSwitcher)
 			+SWidgetSwitcher::Slot()
 			[
-				SNew(SReplicationClientView)
+				SNew(SReplicationClientView, InClient->GetConcertClient(), *ReplicationManager->GetClientManager())
 				.GetReplicationClient_Lambda([this](){ return &ReplicationManager->GetClientManager()->GetLocalClient(); })
 				.ViewSelectionArea()
 				[
@@ -97,7 +94,7 @@ namespace UE::MultiUserClient
 		TMap<FGuid, int32> OldRemoteClientToWidgetSwitcherIndex = MoveTemp(RemoteClientToWidgetSwitcherIndex);
         for (const TNonNullPtr<FRemoteReplicationClient>& RemoteClient : ReplicationManager->GetClientManager()->GetRemoteClients())
         {
-        	const FGuid& EndpointId = RemoteClient->GetRemoteEndpointId();
+        	const FGuid& EndpointId = RemoteClient->GetEndpointId();
         	
         	const int32* ExistingClientWidgetIndex = OldRemoteClientToWidgetSwitcherIndex.Find(EndpointId);
         	if (ExistingClientWidgetIndex)
@@ -111,8 +108,8 @@ namespace UE::MultiUserClient
         	{
         		ClientViewSwitcher->AddSlot()
         		[
-        			SNew(SReplicationClientView)
-        			.GetReplicationClient_Lambda([this, EndpointId = RemoteClient->GetRemoteEndpointId()]()
+        			SNew(SReplicationClientView, Client->GetConcertClient(), *ReplicationManager->GetClientManager())
+        			.GetReplicationClient_Lambda([this, EndpointId = RemoteClient->GetEndpointId()]()
         			{
         				// It is unsafe to simply capture RemoteClient because the containing TArray may reallocate its location
         				return ReplicationManager->GetClientManager()->FindRemoteClient(EndpointId);
@@ -164,7 +161,7 @@ namespace UE::MultiUserClient
 			{
 				const TArray<TNonNullPtr<FRemoteReplicationClient>> RemoteClients = ReplicationManager->GetClientManager()->GetRemoteClients();
 				TArray<FGuid> Result;
-				Algo::Transform(RemoteClients, Result, [](const TNonNullPtr<FRemoteReplicationClient>& InClient){ return InClient->GetRemoteEndpointId(); });
+				Algo::Transform(RemoteClients, Result, [](const TNonNullPtr<FRemoteReplicationClient>& InClient){ return InClient->GetEndpointId(); });
 
 				const TSharedPtr<IConcertClientSession> CurrentSession = Client->GetConcertClient()->GetCurrentSession();
 				Result.Sort([&CurrentSession](const FGuid& Left, const FGuid& Right)
diff --git a/Engine/Plugins/Developer/Concert/ConcertUI/ConcertClientSharedSlate/Source/ConcertClientSharedSlate/Private/Replication/Editor/View/ObjectViewer/SReplicationStreamViewer.cpp b/Engine/Plugins/Developer/Concert/ConcertUI/ConcertClientSharedSlate/Source/ConcertClientSharedSlate/Private/Replication/Editor/View/ObjectViewer/SReplicationStreamViewer.cpp
--- a/Engine/Plugins/Developer/Concert/ConcertUI/ConcertClientSharedSlate/Source/ConcertClientSharedSlate/Private/Replication/Editor/View/ObjectViewer/SReplicationStreamViewer.cpp
+++ b/Engine/Plugins/Developer/Concert/ConcertUI/ConcertClientSharedSlate/Source/ConcertClientSharedSlate/Private/Replication/Editor/View/ObjectViewer/SReplicationStreamViewer.cpp
@@ -152,8 +152,7 @@ namespace UE::ConcertClientSharedSlate
 	{
 		TArray Columns
 		{
-			ReplicationColumns::TopLevel::IconColumn(PropertiesModel.ToSharedRef()),
-			ReplicationColumns::TopLevel::LabelColumn(),
+			ReplicationColumns::TopLevel::LabelColumn(PropertiesModel.ToSharedRef()),
 			ReplicationColumns::TopLevel::TypeColumn(PropertiesModel.ToSharedRef())
 		};
 		Columns.Append(InArgs._AdditionalObjectColumns);
diff --git a/Engine/Plugins/Developer/Concert/ConcertUI/ConcertClientSharedSlate/Source/ConcertClientSharedSlate/Private/Replication/Editor/View/ObjectViewer/Tree/SReplicationColumnRow.h b/Engine/Plugins/Developer/Concert/ConcertUI/ConcertClientSharedSlate/Source/ConcertClientSharedSlate/Private/Replication/Editor/View/ObjectViewer/Tree/SReplicationColumnRow.h
--- a/Engine/Plugins/Developer/Concert/ConcertUI/ConcertClientSharedSlate/Source/ConcertClientSharedSlate/Private/Replication/Editor/View/ObjectViewer/Tree/SReplicationColumnRow.h
+++ b/Engine/Plugins/Developer/Concert/ConcertUI/ConcertClientSharedSlate/Source/ConcertClientSharedSlate/Private/Replication/Editor/View/ObjectViewer/Tree/SReplicationColumnRow.h
@@ -86,7 +86,12 @@ namespace UE::ConcertClientSharedSlate
 			const bool bNeedsExpanderArrow = ColumnName == ExpandableColumnLabel;
 			if (!bNeedsExpanderArrow)
 			{
-				return ColumnWidget;
+				return SNew(SBox)
+					.MinDesiredHeight(RowHeight)
+					.VAlign(VAlign_Center)
+					[
+						ColumnWidget
+					];
 			}
 			
 			return SNew(SBox)
@@ -104,6 +109,7 @@ namespace UE::ConcertClientSharedSlate
 
 					+SHorizontalBox::Slot()
 					.FillWidth(1.0f)
+					.VAlign(VAlign_Center)
 					[
 						ColumnWidget
 					]
diff --git a/Engine/Plugins/Developer/Concert/ConcertUI/ConcertClientSharedSlate/Source/ConcertClientSharedSlate/Private/Replication/Editor/View/ObjectViewer/Tree/SelectionViewerColumns.cpp b/Engine/Plugins/Developer/Concert/ConcertUI/ConcertClientSharedSlate/Source/ConcertClientSharedSlate/Private/Replication/Editor/View/ObjectViewer/Tree/SelectionViewerColumns.cpp
--- a/Engine/Plugins/Developer/Concert/ConcertUI/ConcertClientSharedSlate/Source/ConcertClientSharedSlate/Private/Replication/Editor/View/ObjectViewer/Tree/SelectionViewerColumns.cpp
+++ b/Engine/Plugins/Developer/Concert/ConcertUI/ConcertClientSharedSlate/Source/ConcertClientSharedSlate/Private/Replication/Editor/View/ObjectViewer/Tree/SelectionViewerColumns.cpp
@@ -27,39 +27,33 @@ namespace UE::ConcertClientSharedSlate::ReplicationColumns::TopLevel
 	const FName IconColumnId = TEXT("IconColumn");
 	const FName LabelColumnId = TEXT("LabelColumn");
 	const FName TypeColumnId = TEXT("TypeColumn");
-
-	FReplicationTopLevelObjectColumn IconColumn(TSharedRef<IObjectToPropertiesModel> Model, const float ColumnWidth)
+	
+	FReplicationTopLevelObjectColumn LabelColumn(TSharedRef<IObjectToPropertiesModel> Model)
 	{
 		return FReplicationTopLevelObjectColumn(
 			FReplicationTopLevelObjectColumn::FArguments()
-				.GenerateWidgetColumn_Lambda([Model](const FReplicationTopLevelObjectColumn::FBuildArgs& Args)
+				.GenerateWidgetColumn_Lambda([Model = MoveTemp(Model)](const FReplicationTopLevelObjectColumn::FBuildArgs& Args)
 				{
 					return SNew(SHorizontalBox)
 						+SHorizontalBox::Slot()
+						.AutoWidth()
 						.HAlign(HAlign_Center)
 						.VAlign(VAlign_Center)
 						[
 							SNew(SImage)
 							.Image(DisplayUtils::GetObjectIcon(*Model, Args.RowData.GetObjectPath()).GetOptionalIcon())
+						]
+					
+						+SHorizontalBox::Slot()
+						.HAlign(HAlign_Left)
+						.VAlign(VAlign_Center)
+						.Padding(6.f, 0.f, 0.f, 0.f)
+						[
+							SNew(STextBlock)
+							.HighlightText(TAttribute<FText>::CreateLambda([HighlightText = Args.HighlightText](){ return *HighlightText; }))
+							.Text(DisplayUtils::GetObjectDisplayText(Args.RowData.GetObjectPath()))
 						];
 				})
-				.ColumnSortOrder(static_cast<int32>(ETopLevelColumnOrder::Icon)),
-			SHeaderRow::Column(IconColumnId)
-				.DefaultLabel(FText::GetEmpty())
-				.FixedWidth(ColumnWidth)
-			);
-	}
-	
-	FReplicationTopLevelObjectColumn LabelColumn()
-	{
-		return FReplicationTopLevelObjectColumn(
-			FReplicationTopLevelObjectColumn::FArguments()
-				.GenerateWidgetColumn_Lambda([](const FReplicationTopLevelObjectColumn::FBuildArgs& Args)
-				{
-					return SNew(STextBlock)
-						.HighlightText(TAttribute<FText>::CreateLambda([HighlightText = Args.HighlightText](){ return *HighlightText; }))
-						.Text(DisplayUtils::GetObjectDisplayText(Args.RowData.GetObjectPath()));
-				})
 				.PopulateSearchItems_Lambda([](const FReplicatedObjectData& ObjectData, TArray<FString>& InOutSearchStrings)
 				{
 					InOutSearchStrings.Add(DisplayUtils::GetObjectDisplayText(ObjectData.GetObjectPath()).ToString());
@@ -170,7 +164,7 @@ namespace UE::ConcertClientSharedSlate::ReplicationColumns::Property
 				.ColumnSortOrder(static_cast<int32>(EReplicationPropertyColumnOrder::Label)),
 			SHeaderRow::Column(LabelColumnId)
 				.DefaultLabel(LOCTEXT("LabelColumnLabel", "Label"))
-				.FillSized(500.f)
+				.FillWidth(1.f)
 			);
 	}
 	
@@ -198,7 +192,7 @@ namespace UE::ConcertClientSharedSlate::ReplicationColumns::Property
 				.ColumnSortOrder(static_cast<int32>(EReplicationPropertyColumnOrder::Type)),
 			SHeaderRow::Column(TypeColumnId)
 				.DefaultLabel(LOCTEXT("TypeColumnLabel", "Type"))
-				.FillWidth(1.f)
+				.FillWidth(0.5f)
 			);
 	}
 
diff --git a/Engine/Plugins/Developer/Concert/ConcertUI/ConcertClientSharedSlate/Source/ConcertClientSharedSlate/Private/Replication/Editor/View/ObjectViewer/Tree/SelectionViewerColumns.h b/Engine/Plugins/Developer/Concert/ConcertUI/ConcertClientSharedSlate/Source/ConcertClientSharedSlate/Private/Replication/Editor/View/ObjectViewer/Tree/SelectionViewerColumns.h
--- a/Engine/Plugins/Developer/Concert/ConcertUI/ConcertClientSharedSlate/Source/ConcertClientSharedSlate/Private/Replication/Editor/View/ObjectViewer/Tree/SelectionViewerColumns.h
+++ b/Engine/Plugins/Developer/Concert/ConcertUI/ConcertClientSharedSlate/Source/ConcertClientSharedSlate/Private/Replication/Editor/View/ObjectViewer/Tree/SelectionViewerColumns.h
@@ -22,16 +22,13 @@ namespace UE::ConcertClientSharedSlate::ReplicationColumns::TopLevel
 
 	enum class ETopLevelColumnOrder : int32
 	{
-		/** Displays the class icon */
-		Icon = 10,
 		/** Label of the object */
 		Label = 20,
 		/** Class of the object */
 		Type = 30,
 	};
 
-	FReplicationTopLevelObjectColumn IconColumn(TSharedRef<IObjectToPropertiesModel> Model, const float ColumnWidth = 20.f);
-	FReplicationTopLevelObjectColumn LabelColumn();
+	FReplicationTopLevelObjectColumn LabelColumn(TSharedRef<IObjectToPropertiesModel> Model);
 	FReplicationTopLevelObjectColumn TypeColumn(TSharedRef<IObjectToPropertiesModel> Model);
 }
 
diff --git a/Engine/Plugins/Developer/Concert/ConcertUI/ConcertClientSharedSlate/Source/ConcertClientSharedSlate/Private/Widgets/ClientName/SClientName.cpp b/Engine/Plugins/Developer/Concert/ConcertUI/ConcertClientSharedSlate/Source/ConcertClientSharedSlate/Private/Widgets/ClientName/SClientName.cpp
--- a/Engine/Plugins/Developer/Concert/ConcertUI/ConcertClientSharedSlate/Source/ConcertClientSharedSlate/Private/Widgets/ClientName/SClientName.cpp
+++ b/Engine/Plugins/Developer/Concert/ConcertUI/ConcertClientSharedSlate/Source/ConcertClientSharedSlate/Private/Widgets/ClientName/SClientName.cpp
@@ -2,8 +2,8 @@
 
 #include "Widgets/ClientName/SClientName.h"
 
-#include "ConcertFrontendUtils.h"
-
+#include "Styling/AppStyle.h"
+#include "Widgets/Layout/SBorder.h"
 #include "Widgets/Text/STextBlock.h"
 
 #define LOCTEXT_NAMESPACE "SClientName"
@@ -19,26 +19,36 @@ namespace UE::ConcertClientSharedSlate
 		
 		ChildSlot
 		[
-			ConcertFrontendUtils::CreateDisplayName(
-				TAttribute<FText>::CreateSP(this, &SClientName::GetClientDisplayName)
-				)
+			SNew(SBorder)
+			.BorderImage(FAppStyle::Get().GetBrush("NoBorder"))
+			.ColorAndOpacity(FLinearColor(0.75f, 0.75f, 0.75f))
+			[
+				SNew(STextBlock)
+				.Font(InArgs._Font)
+				.Text(TAttribute<FText>::CreateSP(this, &SClientName::GetClientDisplayName))
+				.HighlightText(InArgs._HighlightText)
+			]
 		];
 	}
 
-	FText SClientName::GetClientDisplayName() const
+	FText SClientName::GetDisplayText(const FConcertClientInfo& Info, bool bDisplayAsLocalClient)
 	{
-		const FConcertClientInfo* ClientInfo = ClientInfoAttribute.Get();
-		check(ClientInfo);
-		
-		if (DisplayAsLocalClientAttribute.Get())
+		if (bDisplayAsLocalClient)
 		{
 			return FText::Format(
 				LOCTEXT("ClientDisplayNameFmt", "{0} (me)"),
-				FText::FromString(ClientInfo->DisplayName)
+				FText::FromString(Info.DisplayName)
 				);
 		}
 		
-		return FText::FromString(ClientInfo->DisplayName);
+		return FText::FromString(Info.DisplayName);
+	}
+
+	FText SClientName::GetClientDisplayName() const
+	{
+		const FConcertClientInfo* ClientInfo = ClientInfoAttribute.Get();
+		check(ClientInfo);
+		return GetDisplayText(*ClientInfo, DisplayAsLocalClientAttribute.Get());
 	}
 }
 
diff --git a/Engine/Plugins/Developer/Concert/ConcertUI/ConcertClientSharedSlate/Source/ConcertClientSharedSlate/Private/Widgets/ClientName/SHorizontalClientList.cpp b/Engine/Plugins/Developer/Concert/ConcertUI/ConcertClientSharedSlate/Source/ConcertClientSharedSlate/Private/Widgets/ClientName/SHorizontalClientList.cpp
new file mode 100644
--- /dev/null
+++ b/Engine/Plugins/Developer/Concert/ConcertUI/ConcertClientSharedSlate/Source/ConcertClientSharedSlate/Private/Widgets/ClientName/SHorizontalClientList.cpp
@@ -0,0 +1,129 @@
+// Copyright Epic Games, Inc. All Rights Reserved.
+
+#include "Widgets/ClientName/SHorizontalClientList.h"
+
+#include "IConcertClient.h"
+
+#include "Widgets/ClientName/SLocalClientName.h"
+#include "Widgets/ClientName/SRemoteClientName.h"
+#include "Widgets/Layout/SScrollBox.h"
+#include "Widgets/Layout/SWidgetSwitcher.h"
+#include "Widgets/Text/STextBlock.h"
+
+#define LOCTEXT_NAMESPACE "SHorizontalClientList"
+
+namespace UE::ConcertClientSharedSlate
+{
+	bool SHorizontalClientList::SortLocalClientFirstThenAlphabetical(const FConcertSessionClientInfo& Left, const FConcertSessionClientInfo& Right, TSharedRef<IConcertClient> Client)
+	{
+		// If one of the compare clients is local, always return that the local client is smaller.
+		const auto IsLocalClient = [&Client](const FConcertSessionClientInfo& Info)
+		{
+			return Client->GetCurrentSession() && Client->GetCurrentSession()->GetSessionClientEndpointId() == Info.ClientEndpointId;
+		};
+		const bool bLeftIsLocalClient = IsLocalClient(Left);
+		const bool bRightIsLocalClient = IsLocalClient(Right);
+		return bLeftIsLocalClient
+			|| (!bRightIsLocalClient && Left.ClientInfo.DisplayName < Right.ClientInfo.DisplayName);
+	}
+
+	void SHorizontalClientList::Construct(const FArguments& InArgs, TSharedRef<IConcertClient> InClient)
+	{
+		LocalConcertClient = MoveTemp(InClient);
+		SortPredicateDelegate = InArgs._SortPredicate.IsBound()
+			? InArgs._SortPredicate
+			: FSortPredicate::CreateStatic(&SHorizontalClientList::SortLocalClientFirstThenAlphabetical, InClient);
+		HighlightTextAttribute = InArgs._HighlightText;
+		NameFont = InArgs._Font;
+		
+		ChildSlot
+		[
+			SAssignNew(WidgetSwitcher, SWidgetSwitcher)
+			.WidgetIndex(0)
+			+SWidgetSwitcher::Slot()
+			[
+				InArgs._EmptyListSlot.Widget
+			]
+			+SWidgetSwitcher::Slot()
+			[
+				SAssignNew(ScrollBox, SScrollBox)
+				.Orientation(Orient_Horizontal)
+			]
+		];
+	}
+
+	void SHorizontalClientList::RefreshList(const TConstArrayView<FGuid>& Clients)
+	{
+		ScrollBox->ClearChildren();
+
+		if (Clients.IsEmpty())
+		{
+			WidgetSwitcher->SetActiveWidgetIndex(0);
+			return;
+		}
+		WidgetSwitcher->SetActiveWidgetIndex(1);
+
+		const TSharedPtr<IConcertClientSession> ClientSession = LocalConcertClient->GetCurrentSession();
+		if (!ensureMsgf(ClientSession, TEXT("This widget does not work if the local client is not in a session!")))
+		{
+			return;
+		}
+
+		// Prefetch the client info to avoid many FindSessionClient during Sort()
+		TArray<FConcertSessionClientInfo> ClientsToDisplay;
+		for (const FGuid& Client : Clients)
+		{
+			FConcertSessionClientInfo Info;
+			if (ClientSession->FindSessionClient(Client, Info))
+			{
+				ClientsToDisplay.Add(Info);
+			}
+			else if (Client == ClientSession->GetSessionClientEndpointId())
+			{
+				ClientsToDisplay.Emplace(Client, ClientSession->GetLocalClientInfo());
+			}
+		}
+
+		ClientsToDisplay.Sort([this](const FConcertSessionClientInfo& Left, const FConcertSessionClientInfo& Right)
+		{
+			return SortPredicateDelegate.Execute(Left, Right);
+		});
+		bool bIsFirst = true;
+		for (const FConcertSessionClientInfo& Info : ClientsToDisplay)
+		{
+			if (!bIsFirst)
+			{
+				ScrollBox->AddSlot()
+				[
+					SNew(STextBlock)
+					.Text(LOCTEXT("Comma", ", "))
+					.Font(NameFont)
+				];
+			}
+			
+			if (Info.ClientEndpointId == ClientSession->GetSessionClientEndpointId())
+			{
+				ScrollBox->AddSlot()
+				[
+					SNew(SLocalClientName, LocalConcertClient.ToSharedRef())
+					.HighlightText(HighlightTextAttribute)
+					.Font(NameFont)
+				];
+			}
+			else
+			{
+				ScrollBox->AddSlot()
+				[
+					SNew(SRemoteClientName, LocalConcertClient.ToSharedRef())
+					.ClientEndpointId(Info.ClientEndpointId)
+					.HighlightText(HighlightTextAttribute)
+					.Font(NameFont)
+				];
+			}
+
+			bIsFirst = false;
+		}
+	}
+}
+
+#undef LOCTEXT_NAMESPACE
\ No newline at end of file
diff --git a/Engine/Plugins/Developer/Concert/ConcertUI/ConcertClientSharedSlate/Source/ConcertClientSharedSlate/Private/Widgets/ClientName/SLocalClientName.cpp b/Engine/Plugins/Developer/Concert/ConcertUI/ConcertClientSharedSlate/Source/ConcertClientSharedSlate/Private/Widgets/ClientName/SLocalClientName.cpp
--- a/Engine/Plugins/Developer/Concert/ConcertUI/ConcertClientSharedSlate/Source/ConcertClientSharedSlate/Private/Widgets/ClientName/SLocalClientName.cpp
+++ b/Engine/Plugins/Developer/Concert/ConcertUI/ConcertClientSharedSlate/Source/ConcertClientSharedSlate/Private/Widgets/ClientName/SLocalClientName.cpp
@@ -18,6 +18,8 @@ namespace UE::ConcertClientSharedSlate
 			SNew(SClientName)
 			.ClientInfo_Lambda([InClient](){ return &InClient->GetClientInfo(); })
 			.DisplayAsLocalClient(true)
+			.HighlightText(InArgs._HighlightText)
+			.Font(InArgs._Font)
 		];
 	}
 }
diff --git a/Engine/Plugins/Developer/Concert/ConcertUI/ConcertClientSharedSlate/Source/ConcertClientSharedSlate/Private/Widgets/ClientName/SRemoteClientName.cpp b/Engine/Plugins/Developer/Concert/ConcertUI/ConcertClientSharedSlate/Source/ConcertClientSharedSlate/Private/Widgets/ClientName/SRemoteClientName.cpp
--- a/Engine/Plugins/Developer/Concert/ConcertUI/ConcertClientSharedSlate/Source/ConcertClientSharedSlate/Private/Widgets/ClientName/SRemoteClientName.cpp
+++ b/Engine/Plugins/Developer/Concert/ConcertUI/ConcertClientSharedSlate/Source/ConcertClientSharedSlate/Private/Widgets/ClientName/SRemoteClientName.cpp
@@ -20,6 +20,8 @@ namespace UE::ConcertClientSharedSlate
 		[
 			SNew(SClientName)
 			.ClientInfo(this, &SRemoteClientName::GetClientInfo)
+			.HighlightText(InArgs._HighlightText)
+			.Font(InArgs._Font)
 		];
 	}
 
diff --git a/Engine/Plugins/Developer/Concert/ConcertUI/ConcertClientSharedSlate/Source/ConcertClientSharedSlate/Public/Replication/Editor/View/IReplicationStreamViewer.h b/Engine/Plugins/Developer/Concert/ConcertUI/ConcertClientSharedSlate/Source/ConcertClientSharedSlate/Public/Replication/Editor/View/IReplicationStreamViewer.h
--- a/Engine/Plugins/Developer/Concert/ConcertUI/ConcertClientSharedSlate/Source/ConcertClientSharedSlate/Public/Replication/Editor/View/IReplicationStreamViewer.h
+++ b/Engine/Plugins/Developer/Concert/ConcertUI/ConcertClientSharedSlate/Source/ConcertClientSharedSlate/Public/Replication/Editor/View/IReplicationStreamViewer.h
@@ -21,7 +21,7 @@ namespace UE::ConcertClientSharedSlate
 		/** @return The objects selected in the top of the viewer; these are the objects that IReplicationSubobjectView bases its view of. */
 		virtual TArray<FSoftObjectPath> GetSelectedTopLevelObjects() const = 0;
 		/**
-		 * @return The objects for which the properties are being edited.
+		 * @return The objects for which the properties are being edited / displayed.
 		 * If there is an IReplicationSubobjectView, this is IReplicationSubobjectView::GetSelectedObjects.
 		 * Otherwise it is IReplicationStreamViewer::GetSelectedTopLevelObjects.
 		 */
diff --git a/Engine/Plugins/Developer/Concert/ConcertUI/ConcertClientSharedSlate/Source/ConcertClientSharedSlate/Public/Replication/Editor/View/ReplicationColumn.h b/Engine/Plugins/Developer/Concert/ConcertUI/ConcertClientSharedSlate/Source/ConcertClientSharedSlate/Public/Replication/Editor/View/ReplicationColumn.h
--- a/Engine/Plugins/Developer/Concert/ConcertUI/ConcertClientSharedSlate/Source/ConcertClientSharedSlate/Public/Replication/Editor/View/ReplicationColumn.h
+++ b/Engine/Plugins/Developer/Concert/ConcertUI/ConcertClientSharedSlate/Source/ConcertClientSharedSlate/Public/Replication/Editor/View/ReplicationColumn.h
@@ -90,7 +90,10 @@ namespace UE::ConcertClientSharedSlate
 					})
 					.PopulateSearchItems_Lambda([TransformOperation, PopulateSearchString](const TOtherColumnType& InOtherRowData, TArray<FString>& InOutSearchStrings)
 					{
-						PopulateSearchString.Execute(TransformOperation(InOtherRowData), InOutSearchStrings);
+						if (ensure(PopulateSearchString.IsBound()))
+						{
+							PopulateSearchString.Execute(TransformOperation(InOtherRowData), InOutSearchStrings);
+						}
 					})
 					.ColumnSortOrder(ColumnSortOrderValue),
 				*this
diff --git a/Engine/Plugins/Developer/Concert/ConcertUI/ConcertClientSharedSlate/Source/ConcertClientSharedSlate/Public/Replication/Editor/View/ReplicationColumnsUtils.h b/Engine/Plugins/Developer/Concert/ConcertUI/ConcertClientSharedSlate/Source/ConcertClientSharedSlate/Public/Replication/Editor/View/ReplicationColumnsUtils.h
--- a/Engine/Plugins/Developer/Concert/ConcertUI/ConcertClientSharedSlate/Source/ConcertClientSharedSlate/Public/Replication/Editor/View/ReplicationColumnsUtils.h
+++ b/Engine/Plugins/Developer/Concert/ConcertUI/ConcertClientSharedSlate/Source/ConcertClientSharedSlate/Public/Replication/Editor/View/ReplicationColumnsUtils.h
@@ -72,6 +72,7 @@ namespace UE::ConcertClientSharedSlate
 		check(Delegates.GetCheckboxStateDelegate.IsBound() && Delegates.OnCheckboxChangedDelegate.IsBound());
 		return TReplicationColumn<TListItemType>(
 			typename TReplicationColumn<TListItemType>::FArguments()
+				.PopulateSearchItems_Lambda([](const auto&, auto&){})
 				.GenerateWidgetColumn_Lambda(
 					[Delegates](const typename TReplicationColumn<TListItemType>::FBuildArgs& Args)
 					{
diff --git a/Engine/Plugins/Developer/Concert/ConcertUI/ConcertClientSharedSlate/Source/ConcertClientSharedSlate/Public/Widgets/ClientName/SClientName.h b/Engine/Plugins/Developer/Concert/ConcertUI/ConcertClientSharedSlate/Source/ConcertClientSharedSlate/Public/Widgets/ClientName/SClientName.h
--- a/Engine/Plugins/Developer/Concert/ConcertUI/ConcertClientSharedSlate/Source/ConcertClientSharedSlate/Public/Widgets/ClientName/SClientName.h
+++ b/Engine/Plugins/Developer/Concert/ConcertUI/ConcertClientSharedSlate/Source/ConcertClientSharedSlate/Public/Widgets/ClientName/SClientName.h
@@ -3,6 +3,7 @@
 #pragma once
 
 #include "ConcertMessageData.h"
+#include "Styling/AppStyle.h"
 #include "Widgets/DeclarativeSyntaxSupport.h"
 #include "Widgets/SCompoundWidget.h"
 
@@ -17,15 +18,23 @@ namespace UE::ConcertClientSharedSlate
 
 		SLATE_BEGIN_ARGS(SClientName)
 			: _DisplayAsLocalClient(false)
+			, _Font(FAppStyle::Get().GetFontStyle("BoldFont"))
 		{}
 			/** The client info to display. */
 			SLATE_ATTRIBUTE(const FConcertClientInfo*, ClientInfo)
 			/** Whether visually indicate that this is a local client (appends "(me)" if true). */
 			SLATE_ATTRIBUTE(bool, DisplayAsLocalClient)
+			/** Used for highlighting in the text */
+			SLATE_ATTRIBUTE(FText, HighlightText)
+			/** The font to use for the name */
+			SLATE_ARGUMENT(FSlateFontInfo, Font)
 		SLATE_END_ARGS()
 
 		void Construct(const FArguments& InArgs);
 
+		/** @return The display that would be used given the settings. */
+		static FText GetDisplayText(const FConcertClientInfo& Info, bool bDisplayAsLocalClient);
+
 	private:
 
 		/** The client info to display. */
diff --git a/Engine/Plugins/Developer/Concert/ConcertUI/ConcertClientSharedSlate/Source/ConcertClientSharedSlate/Public/Widgets/ClientName/SHorizontalClientList.h b/Engine/Plugins/Developer/Concert/ConcertUI/ConcertClientSharedSlate/Source/ConcertClientSharedSlate/Public/Widgets/ClientName/SHorizontalClientList.h
new file mode 100644
--- /dev/null
+++ b/Engine/Plugins/Developer/Concert/ConcertUI/ConcertClientSharedSlate/Source/ConcertClientSharedSlate/Public/Widgets/ClientName/SHorizontalClientList.h
@@ -0,0 +1,61 @@
+// Copyright Epic Games, Inc. All Rights Reserved.
+
+#pragma once
+
+#include "Styling/AppStyle.h"
+#include "Widgets/DeclarativeSyntaxSupport.h"
+#include "Widgets/SCompoundWidget.h"
+
+class IConcertClient;
+class SScrollBox;
+class SWidgetSwitcher;
+
+struct FConcertSessionClientInfo;
+struct FGuid;
+
+namespace UE::ConcertClientSharedSlate
+{
+	/** Aligns client widgets from left to right. If there is not enough space, a horizontal scroll bar cuts of the list. */
+	class CONCERTCLIENTSHAREDSLATE_API SHorizontalClientList : public SCompoundWidget
+	{
+	public:
+
+		DECLARE_DELEGATE_RetVal_TwoParams(bool, FSortPredicate, const FConcertSessionClientInfo& Left, const FConcertSessionClientInfo& FConcertSessionClientInfo);
+		static bool SortLocalClientFirstThenAlphabetical(const FConcertSessionClientInfo& Left, const FConcertSessionClientInfo& Right, TSharedRef<IConcertClient> Client);
+		
+		SLATE_BEGIN_ARGS(SHorizontalClientList)
+			: _Font(FAppStyle::Get().GetFontStyle("NormalFont"))
+		{}
+			/** Defaults to placing the local client first (if contained) and sorting alphabetically otherwise. */
+			SLATE_EVENT(FSortPredicate, SortPredicate)
+			/** Used for highlighting in the text */
+			SLATE_ATTRIBUTE(FText, HighlightText)
+			/** The widget to display when the list is empty */
+			SLATE_NAMED_SLOT(FArguments, EmptyListSlot)
+			/** The font to use for the names */
+			SLATE_ARGUMENT(FSlateFontInfo, Font)
+		SLATE_END_ARGS()
+
+		void Construct(const FArguments& InArgs, TSharedRef<IConcertClient> InClient);
+
+		/** Refreshes the list. */
+		void RefreshList(const TConstArrayView<FGuid>& Clients);
+
+	private:
+
+		/** Clients used to look up client names by ID. */
+		TSharedPtr<IConcertClient> LocalConcertClient;
+
+		/** Sorts the client list */
+		FSortPredicate SortPredicateDelegate;
+		/** Used for highlighting in the text */
+		TAttribute<FText> HighlightTextAttribute;
+		/** The font to use for the names */
+		FSlateFontInfo NameFont;
+
+		/** Displays the ScrollBox when there are clients and the EmptyListSlot otherwise. */
+		TSharedPtr<SWidgetSwitcher> WidgetSwitcher;
+		/** Contains the children. */
+		TSharedPtr<SScrollBox> ScrollBox;
+	};
+}
\ No newline at end of file
diff --git a/Engine/Plugins/Developer/Concert/ConcertUI/ConcertClientSharedSlate/Source/ConcertClientSharedSlate/Public/Widgets/ClientName/SLocalClientName.h b/Engine/Plugins/Developer/Concert/ConcertUI/ConcertClientSharedSlate/Source/ConcertClientSharedSlate/Public/Widgets/ClientName/SLocalClientName.h
--- a/Engine/Plugins/Developer/Concert/ConcertUI/ConcertClientSharedSlate/Source/ConcertClientSharedSlate/Public/Widgets/ClientName/SLocalClientName.h
+++ b/Engine/Plugins/Developer/Concert/ConcertUI/ConcertClientSharedSlate/Source/ConcertClientSharedSlate/Public/Widgets/ClientName/SLocalClientName.h
@@ -2,7 +2,7 @@
 
 #pragma once
 
-#include "ConcertMessageData.h"
+#include "Styling/AppStyle.h"
 #include "Widgets/DeclarativeSyntaxSupport.h"
 #include "Widgets/SCompoundWidget.h"
 
@@ -21,7 +21,12 @@ namespace UE::ConcertClientSharedSlate
 	public:
 
 		SLATE_BEGIN_ARGS(SLocalClientName)
+			: _Font(FAppStyle::Get().GetFontStyle("BoldFont"))
 		{}
+			/** Used for highlighting in the text */
+			SLATE_ATTRIBUTE(FText, HighlightText)
+			/** The font to use for the name */
+			SLATE_ARGUMENT(FSlateFontInfo, Font)
 		SLATE_END_ARGS()
 
 		void Construct(const FArguments& InArgs, TSharedRef<IConcertClient> InClient);
diff --git a/Engine/Plugins/Developer/Concert/ConcertUI/ConcertClientSharedSlate/Source/ConcertClientSharedSlate/Public/Widgets/ClientName/SRemoteClientName.h b/Engine/Plugins/Developer/Concert/ConcertUI/ConcertClientSharedSlate/Source/ConcertClientSharedSlate/Public/Widgets/ClientName/SRemoteClientName.h
--- a/Engine/Plugins/Developer/Concert/ConcertUI/ConcertClientSharedSlate/Source/ConcertClientSharedSlate/Public/Widgets/ClientName/SRemoteClientName.h
+++ b/Engine/Plugins/Developer/Concert/ConcertUI/ConcertClientSharedSlate/Source/ConcertClientSharedSlate/Public/Widgets/ClientName/SRemoteClientName.h
@@ -3,6 +3,7 @@
 #pragma once
 
 #include "ConcertMessageData.h"
+#include "Styling/AppStyle.h"
 #include "Widgets/DeclarativeSyntaxSupport.h"
 #include "Widgets/SCompoundWidget.h"
 
@@ -24,8 +25,13 @@ namespace UE::ConcertClientSharedSlate
 	public:
 
 		SLATE_BEGIN_ARGS(SRemoteClientName)
+			: _Font(FAppStyle::Get().GetFontStyle("BoldFont"))
 		{}
 			SLATE_ATTRIBUTE(FGuid, ClientEndpointId)
+			/** Used for highlighting in the text */
+			SLATE_ATTRIBUTE(FText, HighlightText)
+			/** The font to use for the name */
+			SLATE_ARGUMENT(FSlateFontInfo, Font)
 		SLATE_END_ARGS()
 
 		void Construct(const FArguments& InArgs, TSharedRef<IConcertClient> InClient);
