Add status bar showing the number of objects being replicated

diff --git a/Engine/Plugins/Developer/Concert/ConcertApp/MultiUserClient/Source/MultiUserClient/MultiUserClient.Build.cs b/Engine/Plugins/Developer/Concert/ConcertApp/MultiUserClient/Source/MultiUserClient/MultiUserClient.Build.cs
--- a/Engine/Plugins/Developer/Concert/ConcertApp/MultiUserClient/Source/MultiUserClient/MultiUserClient.Build.cs
+++ b/Engine/Plugins/Developer/Concert/ConcertApp/MultiUserClient/Source/MultiUserClient/MultiUserClient.Build.cs
@@ -23,6 +23,7 @@ namespace UnrealBuildTool.Rules
 					"Core",
 					"DesktopPlatform",
 					"EditorStyle",
+					"Engine",
 					"InputCore",
 					"MessageLog",
 					"Projects",
diff --git a/Engine/Plugins/Developer/Concert/ConcertApp/MultiUserClient/Source/MultiUserClient/Private/Widgets/ActiveSession/Replication/Client/SReplicationStatus.cpp b/Engine/Plugins/Developer/Concert/ConcertApp/MultiUserClient/Source/MultiUserClient/Private/Widgets/ActiveSession/Replication/Client/SReplicationStatus.cpp
new file mode 100644
--- /dev/null
+++ b/Engine/Plugins/Developer/Concert/ConcertApp/MultiUserClient/Source/MultiUserClient/Private/Widgets/ActiveSession/Replication/Client/SReplicationStatus.cpp
@@ -0,0 +1,127 @@
+// Copyright Epic Games, Inc. All Rights Reserved.
+
+#include "SReplicationStatus.h"
+
+#include "Replication/Util/GlobalAuthorityCache.h"
+
+#include "Algo/AnyOf.h"
+#include "Styling/AppStyle.h"
+#include "Widgets/SBoxPanel.h"
+#include "Widgets/Text/STextBlock.h"
+
+#define LOCTEXT_NAMESPACE "SReplicationStatus"
+
+namespace UE::MultiUserClient
+{
+	namespace Private
+	{
+		static TOptional<FSoftObjectPath> GetOwningActorOf(const FSoftObjectPath& SoftObjectPath)
+		{
+			// Example of an actor called floor
+			// SoftObjectPath = { AssetPath = {PackageName = "/Game/Maps/SyncBoxLevel", AssetName = "SyncBoxLevel"}, SubPathString = "PersistentLevel.Floor" } }
+			const FString& SubPathString = SoftObjectPath.GetSubPathString();
+
+			constexpr int32 PersistentLevelStringLength = 16; // "PersistentLevel." has 16 characters
+			const bool bIsWorldObject = SubPathString.Contains(TEXT("PersistentLevel."), ESearchCase::CaseSensitive);
+			if (!bIsWorldObject)
+			{
+				// Not a path to a world object
+				return {};
+			}
+
+			// Start search after the . behind PersistentLevel
+			const int32 StartSearch = PersistentLevelStringLength + 1;
+			const int32 IndexOfDotAfterActorName = SubPathString.Find(TEXT("."), ESearchCase::CaseSensitive, ESearchDir::FromStart, StartSearch);
+			if (IndexOfDotAfterActorName == INDEX_NONE)
+			{
+				// SoftObjectPath points to an actor
+				return SoftObjectPath;
+			}
+
+			const int32 NumToChopOffRight = SubPathString.Len() - IndexOfDotAfterActorName;
+			const FString NewSubstring = SubPathString.LeftChop(NumToChopOffRight);
+			const FSoftObjectPath PathToOwningActor(SoftObjectPath.GetAssetPath(), NewSubstring);
+			return PathToOwningActor;
+		}
+	}
+	
+	void SReplicationStatus::Construct(const FArguments& InArgs, const ConcertClientSharedSlate::IObjectToPropertiesModel& InObjectModel, FGlobalAuthorityCache& InAuthorityCache)
+	{
+		ObjectModel = &InObjectModel;
+		AuthorityCache = &InAuthorityCache;
+		DisplayedClientsAttribute = InArgs._DisplayedClients;
+		check(DisplayedClientsAttribute.IsBound() || DisplayedClientsAttribute.IsSet());
+
+		AuthorityCache->OnCacheChanged().AddSP(this, &SReplicationStatus::OnAuthorityCacheChanged);
+		ChildSlot
+		[
+			SNew(SHorizontalBox)
+			+SHorizontalBox::Slot()
+			.AutoWidth()
+			[
+				SAssignNew(ActorsText, STextBlock)
+				.Text(LOCTEXT("Replicating", "Replicating "))
+			]
+			+SHorizontalBox::Slot()
+			.AutoWidth()
+			[
+				SAssignNew(ObjectsText, STextBlock)
+				.Font(FAppStyle::Get().GetFontStyle("BoldFont"))
+			]
+			+SHorizontalBox::Slot()
+			.AutoWidth()
+			[
+				SAssignNew(ActorsText, STextBlock)
+				.Text(LOCTEXT("For", " for "))
+			]
+			+SHorizontalBox::Slot()
+			.AutoWidth()
+			[
+				SAssignNew(ActorsText, STextBlock)
+				.Font(FAppStyle::Get().GetFontStyle("BoldFont"))
+			]
+		];
+		
+		RefreshStatusText();
+	}
+
+	SReplicationStatus::~SReplicationStatus()
+	{
+		AuthorityCache->OnCacheChanged().RemoveAll(this);
+	}
+
+	void SReplicationStatus::RefreshStatusText()
+	{
+		const TSet<FGuid> DisplayedClients = DisplayedClientsAttribute.Get();
+
+		TSet<FSoftObjectPath> ReplicatedActors;
+		TSet<FSoftObjectPath> ReplicatedObjects;
+		ObjectModel->ForEachReplicatedObject([this, &DisplayedClients, &ReplicatedActors, &ReplicatedObjects](const FSoftObjectPath& Path)
+		{
+			const bool bIsReplicated = Algo::AnyOf(AuthorityCache->GetClientsWithAuthorityOverObject(Path), [&DisplayedClients](const FGuid& ClientId)
+			{
+				return DisplayedClients.Contains(ClientId);
+			});
+			if (bIsReplicated)
+			{
+				const TOptional<FSoftObjectPath> PathToOwningActor = Private::GetOwningActorOf(Path);
+				if (PathToOwningActor)
+				{
+					ReplicatedActors.Add(*PathToOwningActor);
+				}
+				
+				ReplicatedObjects.Add(Path);
+			}
+			return EBreakBehavior::Continue;
+		});
+
+		ObjectsText->SetText(FText::Format(LOCTEXT("ObjectsTextFmt", "{0} {0}|plural(one=Object,other=Objects)"),
+			ReplicatedObjects.Num()
+			));
+		ActorsText->SetText(FText::Format(LOCTEXT("ActorsTextFmt", "{0} {0}|plural(one=Actor,other=Actors)"),
+			ReplicatedActors.Num()
+			));
+	}
+}
+
+#undef LOCTEXT_NAMESPACE
\ No newline at end of file
diff --git a/Engine/Plugins/Developer/Concert/ConcertApp/MultiUserClient/Source/MultiUserClient/Private/Widgets/ActiveSession/Replication/Client/SReplicationStatus.h b/Engine/Plugins/Developer/Concert/ConcertApp/MultiUserClient/Source/MultiUserClient/Private/Widgets/ActiveSession/Replication/Client/SReplicationStatus.h
new file mode 100644
--- /dev/null
+++ b/Engine/Plugins/Developer/Concert/ConcertApp/MultiUserClient/Source/MultiUserClient/Private/Widgets/ActiveSession/Replication/Client/SReplicationStatus.h
@@ -0,0 +1,58 @@
+// Copyright Epic Games, Inc. All Rights Reserved.
+
+#pragma once
+
+#include "Widgets/DeclarativeSyntaxSupport.h"
+#include "Widgets/SCompoundWidget.h"
+
+class STextBlock;
+
+namespace UE::ConcertClientSharedSlate
+{
+	class IObjectToPropertiesModel;
+}
+
+namespace UE::MultiUserClient
+{
+	class FGlobalAuthorityCache;
+
+	/**
+	 * Displays a text "Replicating x Objects for y Actors".
+	 *
+	 * This view sums up all objects being replicated and shows the distinct actors being replicated.
+	 * Example: You are replicating AActor::ActorGuid and USceneComponent::RelativeLocation on an actor called Floor.
+	 * Result: "Replicating 2 Objects for 1 Actor".
+	 */
+	class SReplicationStatus : public SCompoundWidget
+	{
+	public:
+
+		SLATE_BEGIN_ARGS(SReplicationStatus)
+		{}
+			/** The clients to show statistics for */
+			SLATE_ATTRIBUTE(TSet<FGuid>, DisplayedClients)
+		SLATE_END_ARGS()
+
+		void Construct(const FArguments& InArgs, const ConcertClientSharedSlate::IObjectToPropertiesModel& InObjectModel, FGlobalAuthorityCache& InAuthorityCache);
+		virtual ~SReplicationStatus() override;
+
+	private:
+
+		/** Used to look up registered objects and subobjects. */
+		const ConcertClientSharedSlate::IObjectToPropertiesModel* ObjectModel;
+		/** Used to get authority state of objects and informs us when authority changes. */
+		FGlobalAuthorityCache* AuthorityCache;
+
+		/** The clients to show statistics for */
+		TAttribute<TSet<FGuid>> DisplayedClientsAttribute;
+
+		/** Updated when authority changes. Displays subobjects in bold. */
+		TSharedPtr<STextBlock> ObjectsText;
+		/** Updated when authority changes. Displays actors in bold. */
+		TSharedPtr<STextBlock> ActorsText;
+
+		void OnAuthorityCacheChanged(const FGuid& ClientId) { RefreshStatusText(); }
+		void RefreshStatusText();
+	};
+}
+
diff --git a/Engine/Plugins/Developer/Concert/ConcertApp/MultiUserClient/Source/MultiUserClient/Private/Widgets/ActiveSession/Replication/Client/Single/SReplicationClientView.cpp b/Engine/Plugins/Developer/Concert/ConcertApp/MultiUserClient/Source/MultiUserClient/Private/Widgets/ActiveSession/Replication/Client/Single/SReplicationClientView.cpp
--- a/Engine/Plugins/Developer/Concert/ConcertApp/MultiUserClient/Source/MultiUserClient/Private/Widgets/ActiveSession/Replication/Client/Single/SReplicationClientView.cpp
+++ b/Engine/Plugins/Developer/Concert/ConcertApp/MultiUserClient/Source/MultiUserClient/Private/Widgets/ActiveSession/Replication/Client/Single/SReplicationClientView.cpp
@@ -76,8 +76,9 @@ namespace UE::MultiUserClient
 			.AutoHeight()
 			.Padding(2.f)
 			[
-				SNew(SSingleClientToolbar)
+				SNew(SSingleClientToolbar, PropertyModel, AuthorityCache)
 				.ViewSelectionArea() [ InArgs._ViewSelectionArea.Widget ]
+				.DisplayedClients(TSet{ ReplicationClient->GetEndpointId() })
 			]
 
 			// Editor
diff --git a/Engine/Plugins/Developer/Concert/ConcertApp/MultiUserClient/Source/MultiUserClient/Private/Widgets/ActiveSession/Replication/Client/Single/SSingleClientToolbar.cpp b/Engine/Plugins/Developer/Concert/ConcertApp/MultiUserClient/Source/MultiUserClient/Private/Widgets/ActiveSession/Replication/Client/Single/SSingleClientToolbar.cpp
--- a/Engine/Plugins/Developer/Concert/ConcertApp/MultiUserClient/Source/MultiUserClient/Private/Widgets/ActiveSession/Replication/Client/Single/SSingleClientToolbar.cpp
+++ b/Engine/Plugins/Developer/Concert/ConcertApp/MultiUserClient/Source/MultiUserClient/Private/Widgets/ActiveSession/Replication/Client/Single/SSingleClientToolbar.cpp
@@ -2,12 +2,14 @@
 
 #include "SSingleClientToolbar.h"
 
+#include "Widgets/ActiveSession/Replication/Client/SReplicationStatus.h"
+
 #include "Widgets/Layout/SSpacer.h"
 #include "Widgets/SBoxPanel.h"
 
 namespace UE::MultiUserClient
 {
-	void SSingleClientToolbar::Construct(const FArguments& InArgs)
+	void SSingleClientToolbar::Construct(const FArguments& InArgs, const ConcertClientSharedSlate::IObjectToPropertiesModel& InObjectModel, FGlobalAuthorityCache& InAuthorityCache)
 	{
 		ChildSlot
 		[
@@ -19,6 +21,21 @@ namespace UE::MultiUserClient
 			[
 				InArgs._ViewSelectionArea.Widget
 			]
+			
+			+SHorizontalBox::Slot()
+			.FillWidth(1.f)
+			[
+				SNew(SSpacer)
+			]
+
+			+SHorizontalBox::Slot()
+			.AutoWidth()
+			.VAlign(VAlign_Center)
+			.Padding(0.f, 0.f, 5.f, 0.f)
+			[
+				SNew(SReplicationStatus, InObjectModel, InAuthorityCache)
+				.DisplayedClients(InArgs._DisplayedClients)
+			]
 		];
 	}
 }
\ No newline at end of file
diff --git a/Engine/Plugins/Developer/Concert/ConcertApp/MultiUserClient/Source/MultiUserClient/Private/Widgets/ActiveSession/Replication/Client/Single/SSingleClientToolbar.h b/Engine/Plugins/Developer/Concert/ConcertApp/MultiUserClient/Source/MultiUserClient/Private/Widgets/ActiveSession/Replication/Client/Single/SSingleClientToolbar.h
--- a/Engine/Plugins/Developer/Concert/ConcertApp/MultiUserClient/Source/MultiUserClient/Private/Widgets/ActiveSession/Replication/Client/Single/SSingleClientToolbar.h
+++ b/Engine/Plugins/Developer/Concert/ConcertApp/MultiUserClient/Source/MultiUserClient/Private/Widgets/ActiveSession/Replication/Client/Single/SSingleClientToolbar.h
@@ -5,8 +5,14 @@
 #include "Widgets/DeclarativeSyntaxSupport.h"
 #include "Widgets/SCompoundWidget.h"
 
+namespace UE::ConcertClientSharedSlate
+{
+	class IObjectToPropertiesModel;
+}
+
 namespace UE::MultiUserClient
 {
+	class FGlobalAuthorityCache;
 	class FStreamChangeTracker;
 	class ISubmissionWorkflow;
 
@@ -21,9 +27,12 @@ namespace UE::MultiUserClient
 		{}
 			/** Dedicated space for a widget with which to change the view. */
 			SLATE_NAMED_SLOT(FArguments, ViewSelectionArea)
+		
+			/** The clients to show statistics for */
+			SLATE_ATTRIBUTE(TSet<FGuid>, DisplayedClients)
 		SLATE_END_ARGS()
 
-		void Construct(const FArguments& InArgs);
+		void Construct(const FArguments& InArgs, const ConcertClientSharedSlate::IObjectToPropertiesModel& InObjectModel, FGlobalAuthorityCache& InAuthorityCache);
 	};
 }
 
