Show warning icon when conflicts occur

diff --git a/Engine/Plugins/Developer/Concert/ConcertApp/MultiUserClient/Source/MultiUserClient/Private/Replication/Authority/AuthoritySynchronizer_Base.cpp b/Engine/Plugins/Developer/Concert/ConcertApp/MultiUserClient/Source/MultiUserClient/Private/Replication/Authority/AuthoritySynchronizer_Base.cpp
--- a/Engine/Plugins/Developer/Concert/ConcertApp/MultiUserClient/Source/MultiUserClient/Private/Replication/Authority/AuthoritySynchronizer_Base.cpp
+++ b/Engine/Plugins/Developer/Concert/ConcertApp/MultiUserClient/Source/MultiUserClient/Private/Replication/Authority/AuthoritySynchronizer_Base.cpp
@@ -23,7 +23,6 @@ namespace UE::MultiUserClient
 			return EAuthorityMutability::NoProperties;
 		}
 
-		// TODO DP UE-198356 check whether this would cause a conflict with other clients
 		return EAuthorityMutability::Allowed;
 	}
 }
diff --git a/Engine/Plugins/Developer/Concert/ConcertApp/MultiUserClient/Source/MultiUserClient/Private/Replication/Authority/AuthoritySynchronizer_LocalClient.cpp b/Engine/Plugins/Developer/Concert/ConcertApp/MultiUserClient/Source/MultiUserClient/Private/Replication/Authority/AuthoritySynchronizer_LocalClient.cpp
--- a/Engine/Plugins/Developer/Concert/ConcertApp/MultiUserClient/Source/MultiUserClient/Private/Replication/Authority/AuthoritySynchronizer_LocalClient.cpp
+++ b/Engine/Plugins/Developer/Concert/ConcertApp/MultiUserClient/Source/MultiUserClient/Private/Replication/Authority/AuthoritySynchronizer_LocalClient.cpp
@@ -35,6 +35,23 @@ namespace UE::MultiUserClient
 		return ensure(ReplicationManager) && !ReplicationManager->GetClientOwnedStreamsForObject(ObjectPath).IsEmpty();
 	}
 
+	bool FAuthoritySynchronizer_LocalClient::HasAnyAuthority() const
+	{
+		const IConcertClientReplicationManager* ReplicationManager = Client->GetReplicationManager();
+		if (!ensure(ReplicationManager))
+		{
+			return false;
+		}
+		
+		bool bReached = false;
+		ReplicationManager->ForEachClientOwnedObject([&bReached](const FSoftObjectPath& Object, TSet<FGuid>&& OwningStreams)
+		{
+			bReached = true;
+			return EBreakBehavior::Break;
+		});
+		return bReached;
+	}
+
 	void FAuthoritySynchronizer_LocalClient::OnAuthorityChanged() const
 	{
 		OnServerStateChangedDelegate.Broadcast();
diff --git a/Engine/Plugins/Developer/Concert/ConcertApp/MultiUserClient/Source/MultiUserClient/Private/Replication/Authority/AuthoritySynchronizer_LocalClient.h b/Engine/Plugins/Developer/Concert/ConcertApp/MultiUserClient/Source/MultiUserClient/Private/Replication/Authority/AuthoritySynchronizer_LocalClient.h
--- a/Engine/Plugins/Developer/Concert/ConcertApp/MultiUserClient/Source/MultiUserClient/Private/Replication/Authority/AuthoritySynchronizer_LocalClient.h
+++ b/Engine/Plugins/Developer/Concert/ConcertApp/MultiUserClient/Source/MultiUserClient/Private/Replication/Authority/AuthoritySynchronizer_LocalClient.h
@@ -18,6 +18,7 @@ namespace UE::MultiUserClient
 
 		//~ Begin IClientAuthoritySynchronizer Interface
 		virtual bool HasAuthorityOver(const FSoftObjectPath& ObjectPath) const override;
+		virtual bool HasAnyAuthority() const override;
 		//~ End IClientAuthoritySynchronizer Interface
 
 	private:
diff --git a/Engine/Plugins/Developer/Concert/ConcertApp/MultiUserClient/Source/MultiUserClient/Private/Replication/Authority/AuthoritySynchronizer_RemoteClient.cpp b/Engine/Plugins/Developer/Concert/ConcertApp/MultiUserClient/Source/MultiUserClient/Private/Replication/Authority/AuthoritySynchronizer_RemoteClient.cpp
--- a/Engine/Plugins/Developer/Concert/ConcertApp/MultiUserClient/Source/MultiUserClient/Private/Replication/Authority/AuthoritySynchronizer_RemoteClient.cpp
+++ b/Engine/Plugins/Developer/Concert/ConcertApp/MultiUserClient/Source/MultiUserClient/Private/Replication/Authority/AuthoritySynchronizer_RemoteClient.cpp
@@ -32,6 +32,11 @@ namespace UE::MultiUserClient
 		return EAuthorityMutability::NotSupported;
 	}
 
+	bool FAuthoritySynchronizer_RemoteClient::HasAnyAuthority() const
+	{
+		return !LastServerState.IsEmpty();
+	}
+
 	void FAuthoritySynchronizer_RemoteClient::HandleAuthorityQuery(const TArray<FReplicationAuthorityInfo>& PerStreamAuthority)
 	{
 		TSet<FSoftObjectPath> OldServerState = MoveTemp(LastServerState);
diff --git a/Engine/Plugins/Developer/Concert/ConcertApp/MultiUserClient/Source/MultiUserClient/Private/Replication/Authority/AuthoritySynchronizer_RemoteClient.h b/Engine/Plugins/Developer/Concert/ConcertApp/MultiUserClient/Source/MultiUserClient/Private/Replication/Authority/AuthoritySynchronizer_RemoteClient.h
--- a/Engine/Plugins/Developer/Concert/ConcertApp/MultiUserClient/Source/MultiUserClient/Private/Replication/Authority/AuthoritySynchronizer_RemoteClient.h
+++ b/Engine/Plugins/Developer/Concert/ConcertApp/MultiUserClient/Source/MultiUserClient/Private/Replication/Authority/AuthoritySynchronizer_RemoteClient.h
@@ -23,6 +23,7 @@ namespace UE::MultiUserClient
 
 		//~ Begin IClientAuthoritySynchronizer Interface
 		virtual EAuthorityMutability GetChangeAuthorityMutability(const FSoftObjectPath& ObjectPath) const override;
+		virtual bool HasAnyAuthority() const override;
 		virtual bool HasAuthorityOver(const FSoftObjectPath& ObjectPath) const override { return LastServerState.Contains(ObjectPath); }
 		//~ End IClientAuthoritySynchronizer Interface
 
diff --git a/Engine/Plugins/Developer/Concert/ConcertApp/MultiUserClient/Source/MultiUserClient/Private/Replication/Authority/EAuthorityMutability.h b/Engine/Plugins/Developer/Concert/ConcertApp/MultiUserClient/Source/MultiUserClient/Private/Replication/Authority/EAuthorityMutability.h
--- a/Engine/Plugins/Developer/Concert/ConcertApp/MultiUserClient/Source/MultiUserClient/Private/Replication/Authority/EAuthorityMutability.h
+++ b/Engine/Plugins/Developer/Concert/ConcertApp/MultiUserClient/Source/MultiUserClient/Private/Replication/Authority/EAuthorityMutability.h
@@ -13,9 +13,6 @@ namespace UE::MultiUserClient
 		/** The object has no registered properties. Cannot take authority. */
 		NoProperties,
 		
-		// TODO DP UE-198356: Add function to query which properties and clients are conflicting
-		/** Another client has authority over some of the object's registered properties already. Cannot take authority. */
-		ClientConflict,
 		/** The operation is not implemented. */
 		NotSupported
 	};
diff --git a/Engine/Plugins/Developer/Concert/ConcertApp/MultiUserClient/Source/MultiUserClient/Private/Replication/Authority/IClientAuthoritySynchronizer.h b/Engine/Plugins/Developer/Concert/ConcertApp/MultiUserClient/Source/MultiUserClient/Private/Replication/Authority/IClientAuthoritySynchronizer.h
--- a/Engine/Plugins/Developer/Concert/ConcertApp/MultiUserClient/Source/MultiUserClient/Private/Replication/Authority/IClientAuthoritySynchronizer.h
+++ b/Engine/Plugins/Developer/Concert/ConcertApp/MultiUserClient/Source/MultiUserClient/Private/Replication/Authority/IClientAuthoritySynchronizer.h
@@ -20,6 +20,9 @@ namespace UE::MultiUserClient
 		/** @return Whether the ObjectPath's authority can be changed. */
 		bool CanChangeAuthority(const FSoftObjectPath& ObjectPath) const { return GetChangeAuthorityMutability(ObjectPath) == EAuthorityMutability::Allowed; }
 
+		/** @return Whether this client is sending anything at all. */
+		virtual bool HasAnyAuthority() const = 0;
+
 		/** @return Whether the local instance thinks this client has authority over ObjectPath. */
 		virtual bool HasAuthorityOver(const FSoftObjectPath& ObjectPath) const = 0;
 
diff --git a/Engine/Plugins/Developer/Concert/ConcertApp/MultiUserClient/Source/MultiUserClient/Private/Replication/Util/GlobalAuthorityCache.cpp b/Engine/Plugins/Developer/Concert/ConcertApp/MultiUserClient/Source/MultiUserClient/Private/Replication/Util/GlobalAuthorityCache.cpp
--- a/Engine/Plugins/Developer/Concert/ConcertApp/MultiUserClient/Source/MultiUserClient/Private/Replication/Util/GlobalAuthorityCache.cpp
+++ b/Engine/Plugins/Developer/Concert/ConcertApp/MultiUserClient/Source/MultiUserClient/Private/Replication/Util/GlobalAuthorityCache.cpp
@@ -2,6 +2,7 @@
 
 #include "GlobalAuthorityCache.h"
 
+#include "Replication/AuthorityConflictSharedUtils.h"
 #include "Replication/Client/RemoteReplicationClient.h"
 #include "Replication/Client/ReplicationClientManager.h"
 
@@ -57,6 +58,52 @@ namespace UE::MultiUserClient
 		return Result;
 	}
 
+	FGlobalAuthorityCache::ECanTakeAuthority FGlobalAuthorityCache::CanClientTakeAuthority(const FSoftObjectPath& Object, const FGuid& ClientId, FProcessPropertyConflict ProcessConflict) const
+	{
+		const FReplicationClient* Client = ClientManager.FindClient(ClientId);
+		if (!ensure(Client))
+		{
+			return ECanTakeAuthority::NotApplicable;
+		}
+		
+		const FReplicatedObjectInfo* ObjectInfo = Client->GetStreamSynchronizer().GetServerState().ReplicatedObjects.Find(Object);
+		if (!ObjectInfo)
+		{
+			// Nothing to take authority over
+			return ECanTakeAuthority::NotApplicable;
+		}
+		
+		using namespace ConcertSyncCore::Replication::AuthorityConflictUtils;
+		const EAuthorityConflict Conflict = EnumerateAuthorityConflicts(
+			ClientId,
+			Object,
+			ObjectInfo->PropertySelection.ReplicatedProperties,
+			*this,
+			[&ProcessConflict](const FGuid& ClientId, const FGuid&, const FConcertPropertyChain& ConflictingProperty)
+			{
+				return ProcessConflict(ClientId, ConflictingProperty);
+			});
+		return Conflict == EAuthorityConflict::Allowed ? ECanTakeAuthority::Allowed : ECanTakeAuthority::Conflict;
+	}
+
+	bool FGlobalAuthorityCache::CanClientAddProperty(const FSoftObjectPath& Object, const FGuid& ClientId, const FConcertPropertyChain& Chain) const
+	{
+		const FReplicationClient* Client = ClientManager.FindClient(ClientId);
+		if (!ensure(Client))
+		{
+			return false;
+		}
+		
+		using namespace ConcertSyncCore::Replication::AuthorityConflictUtils;
+		const EAuthorityConflict Conflict = EnumerateAuthorityConflicts(
+			ClientId,
+			Object,
+			{ Chain },
+			*this
+			);
+		return Conflict == EAuthorityConflict::Allowed;
+	}
+
 	TOptional<FGuid> FGlobalAuthorityCache::GetClientWithAuthorityOverProperty(const FSoftObjectPath& Object, const FConcertPropertyChain& Property) const
 	{
 		TOptional<FGuid> Result;
@@ -126,4 +173,43 @@ namespace UE::MultiUserClient
 			}
 		}
 	}
+
+	void FGlobalAuthorityCache::ForEachStream(const FGuid& ClientEndpointId, TFunctionRef<EBreakBehavior(const FGuid& StreamId, const FObjectReplicationMap& ReplicationMap)> Callback) const
+	{
+		const FReplicationClient* Client = ClientManager.FindClient(ClientEndpointId);
+		if (!ensure(Client))
+		{
+			return;
+		}
+
+		const IClientStreamSynchronizer& StreamSynchronizer = Client->GetStreamSynchronizer();
+		Callback(StreamSynchronizer.GetStreamId(), StreamSynchronizer.GetServerState());
+	}
+
+	void FGlobalAuthorityCache::ForEachSendingClient(TFunctionRef<EBreakBehavior(const FGuid& ClientEndpointId)> Callback) const
+	{
+		auto ProcessClient = [&Callback](const FReplicationClient& Client)
+		{
+			return Client.GetAuthoritySynchronizer().HasAnyAuthority()
+				? Callback(Client.GetEndpointId())
+				:EBreakBehavior::Continue;
+		};
+
+		if (ProcessClient(ClientManager.GetLocalClient()) == EBreakBehavior::Break)
+		{
+			return;
+		}
+		for (const TNonNullPtr<FRemoteReplicationClient>& RemoteClient : ClientManager.GetRemoteClients())
+		{
+			if (ProcessClient(*RemoteClient) == EBreakBehavior::Break)
+			{
+				break;
+			}
+		}
+	}
+
+	bool FGlobalAuthorityCache::HasAuthority(const FGuid& ClientId, const FGuid& StreamId, const FSoftObjectPath& ObjectPath) const
+	{
+		return GetClientsWithAuthorityOverObject(ObjectPath).Contains(ClientId);
+	}
 }
diff --git a/Engine/Plugins/Developer/Concert/ConcertApp/MultiUserClient/Source/MultiUserClient/Private/Replication/Util/GlobalAuthorityCache.h b/Engine/Plugins/Developer/Concert/ConcertApp/MultiUserClient/Source/MultiUserClient/Private/Replication/Util/GlobalAuthorityCache.h
--- a/Engine/Plugins/Developer/Concert/ConcertApp/MultiUserClient/Source/MultiUserClient/Private/Replication/Util/GlobalAuthorityCache.h
+++ b/Engine/Plugins/Developer/Concert/ConcertApp/MultiUserClient/Source/MultiUserClient/Private/Replication/Util/GlobalAuthorityCache.h
@@ -4,6 +4,8 @@
 
 #include "Replication/Client/RemoteReplicationClient.h"
 
+#include "Replication/AuthorityConflictSharedUtils.h"
+
 #include "Delegates/Delegate.h"
 #include "Templates/Function.h"
 
@@ -22,9 +24,11 @@ namespace UE::MultiUserClient
 	 * Allows efficient look-up of which objects and properties are owned by which clients.
 	 * This class efficiently answers the question: "Which clients own this object?"
 	 */
-	class FGlobalAuthorityCache : public FNoncopyable
+	class FGlobalAuthorityCache : public FNoncopyable, private ConcertSyncCore::Replication::AuthorityConflictUtils::IReplicationGroundTruth
 	{
 	public:
+
+		using FProcessPropertyConflict = TFunctionRef<EBreakBehavior(const FGuid& ConflictingClientId, const FConcertPropertyChain& Property)>;
 		
 		FGlobalAuthorityCache(FReplicationClientManager& InClientManager);
 		~FGlobalAuthorityCache();
@@ -34,6 +38,19 @@ namespace UE::MultiUserClient
 		/** Util that uses ForEachClientWithAuthorityOverObject to make an array. */
 		TArray<FGuid> GetClientsWithAuthorityOverObject(const FSoftObjectPath& Object) const;
 
+		enum class ECanTakeAuthority
+		{
+			/** The specified client has no properties, does not exist, etc. */
+			NotApplicable,
+			Conflict,
+			Allowed
+		};
+		
+		/** @return Whether the given client can take authority over the object without causing any conflicts. */
+		ECanTakeAuthority CanClientTakeAuthority(const FSoftObjectPath& Object, const FGuid& ClientId, FProcessPropertyConflict ProcessConflict = [](auto&, auto&){ return EBreakBehavior::Break; }) const;
+		/** @return Whether the given client add the given property to the object without causing any conflicts. */
+		bool CanClientAddProperty(const FSoftObjectPath& Object, const FGuid& ClientId, const FConcertPropertyChain& Chain) const;
+
 		/** Gets the client that has authority over the given property, if there is any. */
 		TOptional<FGuid> GetClientWithAuthorityOverProperty(const FSoftObjectPath& Object, const FConcertPropertyChain& Property) const;
 
@@ -73,6 +90,12 @@ namespace UE::MultiUserClient
 			AddClient(ClientId);
 			OnCacheChangedDelegate.Broadcast(ClientId);
 		}
+
+		//~ Begin IReplicationGroundTruth Interface
+		virtual void ForEachStream(const FGuid& ClientEndpointId, TFunctionRef<EBreakBehavior(const FGuid& StreamId, const FObjectReplicationMap& ReplicationMap)> Callback) const override;
+		virtual void ForEachSendingClient(TFunctionRef<EBreakBehavior(const FGuid& ClientEndpointId)> Callback) const override;
+		virtual bool HasAuthority(const FGuid& ClientId, const FGuid& StreamId, const FSoftObjectPath& ObjectPath) const override;
+		//~ End IReplicationGroundTruth Interface
 	};
 }
 
diff --git a/Engine/Plugins/Developer/Concert/ConcertApp/MultiUserClient/Source/MultiUserClient/Private/Widgets/ActiveSession/Replication/Client/ClientUtils.cpp b/Engine/Plugins/Developer/Concert/ConcertApp/MultiUserClient/Source/MultiUserClient/Private/Widgets/ActiveSession/Replication/Client/ClientUtils.cpp
new file mode 100644
--- /dev/null
+++ b/Engine/Plugins/Developer/Concert/ConcertApp/MultiUserClient/Source/MultiUserClient/Private/Widgets/ActiveSession/Replication/Client/ClientUtils.cpp
@@ -0,0 +1,28 @@
+// Copyright Epic Games, Inc. All Rights Reserved.
+
+#include "ClientUtils.h"
+
+#include "IConcertClient.h"
+#include "Widgets/ClientName/SClientName.h"
+
+namespace UE::MultiUserClient::ClientUtils
+{
+	FString GetClientDisplayName(const IConcertClient& InLocalClientInstance, const FGuid& InClientToGetName)
+	{
+		const TSharedPtr<IConcertClientSession> Session = InLocalClientInstance.GetCurrentSession();
+		const bool bIsLocalClient = ensure(Session) && Session->GetSessionClientEndpointId() == InClientToGetName;
+		if (bIsLocalClient)
+		{
+			return ConcertClientSharedSlate::SClientName::GetDisplayText(Session->GetLocalClientInfo(), bIsLocalClient).ToString();
+		}
+
+		FConcertSessionClientInfo ClientInfo;
+		if (ensure(Session) && Session->FindSessionClient(InClientToGetName, ClientInfo))
+		{
+			return ConcertClientSharedSlate::SClientName::GetDisplayText(ClientInfo.ClientInfo, bIsLocalClient).ToString();
+		}
+
+		ensureMsgf(false, TEXT("Bad args"));
+		return {};
+	}
+}
diff --git a/Engine/Plugins/Developer/Concert/ConcertApp/MultiUserClient/Source/MultiUserClient/Private/Widgets/ActiveSession/Replication/Client/ClientUtils.h b/Engine/Plugins/Developer/Concert/ConcertApp/MultiUserClient/Source/MultiUserClient/Private/Widgets/ActiveSession/Replication/Client/ClientUtils.h
new file mode 100644
--- /dev/null
+++ b/Engine/Plugins/Developer/Concert/ConcertApp/MultiUserClient/Source/MultiUserClient/Private/Widgets/ActiveSession/Replication/Client/ClientUtils.h
@@ -0,0 +1,19 @@
+// Copyright Epic Games, Inc. All Rights Reserved.
+
+#pragma once
+
+class FString;
+class IConcertClient;
+struct FGuid;
+
+namespace UE::MultiUserClient::ClientUtils
+{
+	/**
+	 * Gets the display name for a client. Appends (me) if the client is local.
+	 * 
+	 * @param InLocalClientInstance Used to look up client display info
+	 * @param InClientToGetName The endpoint ID of the client whose name to get
+	 * @return The display name or empty
+	 */
+	FString GetClientDisplayName(const IConcertClient& InLocalClientInstance, const FGuid& InClientToGetName);
+}
diff --git a/Engine/Plugins/Developer/Concert/ConcertApp/MultiUserClient/Source/MultiUserClient/Private/Widgets/ActiveSession/Replication/Client/Single/SReplicationClientView.cpp b/Engine/Plugins/Developer/Concert/ConcertApp/MultiUserClient/Source/MultiUserClient/Private/Widgets/ActiveSession/Replication/Client/Single/SReplicationClientView.cpp
--- a/Engine/Plugins/Developer/Concert/ConcertApp/MultiUserClient/Source/MultiUserClient/Private/Widgets/ActiveSession/Replication/Client/Single/SReplicationClientView.cpp
+++ b/Engine/Plugins/Developer/Concert/ConcertApp/MultiUserClient/Source/MultiUserClient/Private/Widgets/ActiveSession/Replication/Client/Single/SReplicationClientView.cpp
@@ -43,6 +43,7 @@ namespace UE::MultiUserClient
 			.AdditionalColumns =
 			{
 				SingleClientColumns::ToggleSubobjectAuthority(AuthorityTracker, SubmissionWorkflow),
+				SingleClientColumns::ConflictWarningForSubobject(InClient, AuthorityCache, ReplicationClient->GetEndpointId()),
 				SingleClientColumns::OwnerOfSubobject(InClient, AuthorityCache)
 			}
 		};
@@ -59,7 +60,8 @@ namespace UE::MultiUserClient
 			},
 			.AdditionalPropertyColumns =
 			{
-				SingleClientColumns::OwnerOfProperty(InClient, AuthorityCache, MoveTemp(GetReplicationViewerAttribute))
+				SingleClientColumns::OwnerOfProperty(InClient, AuthorityCache, MoveTemp(GetReplicationViewerAttribute)),
+				SingleClientColumns::ConflictWarningForProperty(InClient, GetReplicationViewerAttribute, AuthorityCache, ReplicationClient->GetEndpointId())
 			},
 			.IsEditingEnabled = TAttribute<bool>::CreateLambda([&SubmissionWorkflow](){ return SubmissionWorkflow.GetUploadability() != EChangeUploadability::NotImplemented; }),
 			.EditingDisabledToolTipText = LOCTEXT("Editing.NotImplemented", "Editing remote clients is not implemented. You can only edit the local client."),
diff --git a/Engine/Plugins/Developer/Concert/ConcertApp/MultiUserClient/Source/MultiUserClient/Private/Widgets/ActiveSession/Replication/Client/Single/SWarningIcon.cpp b/Engine/Plugins/Developer/Concert/ConcertApp/MultiUserClient/Source/MultiUserClient/Private/Widgets/ActiveSession/Replication/Client/Single/SWarningIcon.cpp
new file mode 100644
--- /dev/null
+++ b/Engine/Plugins/Developer/Concert/ConcertApp/MultiUserClient/Source/MultiUserClient/Private/Widgets/ActiveSession/Replication/Client/Single/SWarningIcon.cpp
@@ -0,0 +1,25 @@
+// Copyright Epic Games, Inc. All Rights Reserved.
+
+#include "SWarningIcon.h"
+
+#include "Styling/AppStyle.h"
+#include "Widgets/Images/SImage.h"
+#include "Widgets/Layout/SScaleBox.h"
+
+namespace UE::MultiUserClient
+{
+	void SWarningIcon::Construct(const FArguments& InArgs)
+	{
+		ChildSlot
+		[
+			SNew(SScaleBox)
+			.Stretch(EStretch::ScaleToFit)
+			[
+				SNew(SImage)
+				.Image(FAppStyle::Get().GetBrush("Icons.WarningWithColor"))
+				.DesiredSizeOverride(FVector2D(16,16))
+			]
+		];
+	}
+}
+
diff --git a/Engine/Plugins/Developer/Concert/ConcertApp/MultiUserClient/Source/MultiUserClient/Private/Widgets/ActiveSession/Replication/Client/Single/SWarningIcon.h b/Engine/Plugins/Developer/Concert/ConcertApp/MultiUserClient/Source/MultiUserClient/Private/Widgets/ActiveSession/Replication/Client/Single/SWarningIcon.h
new file mode 100644
--- /dev/null
+++ b/Engine/Plugins/Developer/Concert/ConcertApp/MultiUserClient/Source/MultiUserClient/Private/Widgets/ActiveSession/Replication/Client/Single/SWarningIcon.h
@@ -0,0 +1,23 @@
+// Copyright Epic Games, Inc. All Rights Reserved.
+
+#pragma once
+
+#include "Widgets/DeclarativeSyntaxSupport.h"
+#include "Widgets/SCompoundWidget.h"
+
+namespace UE::MultiUserClient
+{
+	/** Displays a warning triangle */
+	class SWarningIcon : public SCompoundWidget
+	{
+	public:
+
+		SLATE_BEGIN_ARGS(SWarningIcon)
+		{}
+		SLATE_END_ARGS()
+
+		void Construct(const FArguments& InArgs);
+	};
+}
+
+
diff --git a/Engine/Plugins/Developer/Concert/ConcertApp/MultiUserClient/Source/MultiUserClient/Private/Widgets/ActiveSession/Replication/Client/Single/SingleClientColumns.h b/Engine/Plugins/Developer/Concert/ConcertApp/MultiUserClient/Source/MultiUserClient/Private/Widgets/ActiveSession/Replication/Client/Single/SingleClientColumns.h
--- a/Engine/Plugins/Developer/Concert/ConcertApp/MultiUserClient/Source/MultiUserClient/Private/Widgets/ActiveSession/Replication/Client/Single/SingleClientColumns.h
+++ b/Engine/Plugins/Developer/Concert/ConcertApp/MultiUserClient/Source/MultiUserClient/Private/Widgets/ActiveSession/Replication/Client/Single/SingleClientColumns.h
@@ -34,11 +34,13 @@ namespace UE::MultiUserClient::SingleClientColumns
 	enum class ESubobjectColumnOrder
 	{
 		ToggleAuthority = 0,
+		ConflictWarning = 5,
 		Owner = 20
 	};
 	/** @see UE::ConcertClientSharedSlate::ReplicationColumns::Property::EReplicationPropertyColumnOrder */
 	enum class EPropertyColumnOrder
 	{
+		ConflictWarning = 5,
 		Owner = 50
 	};
 
@@ -117,4 +119,42 @@ namespace UE::MultiUserClient::SingleClientColumns
 		FGlobalAuthorityCache& InAuthorityCache,
 		const TAttribute<const ConcertClientSharedSlate::IReplicationStreamViewer*>& InViewerAttribute
 		);
+
+	
+	/********** Conflict warning **********/
+	extern const FName ConflictWarningTopLevelObjectColumnId;
+	extern const FName ConflictWarningSubobjectColumnId;
+	extern const FName ConflictWarningPropertyColumnId;
+	
+	/**
+	 * Displays a warning symbol next to the checkbox if checking the checkbox would cause an authority conflict when submitted.
+	 * 
+	 * @param InClient The local Concert client used to look up other client display info
+	 * @param InAuthorityCache Used to determine which client owns the object
+	 * @param ClientId ID of the client for which the icon is being created
+	 * 
+	 * @return Column that can be placed in the table
+	 */
+	ConcertClientSharedSlate::ReplicationColumns::FReplicationSubobjectObjectColumn ConflictWarningForSubobject(
+		TSharedRef<IConcertClient> InClient,
+		FGlobalAuthorityCache& InAuthorityCache,
+		const FGuid& ClientId
+		);
+
+	/**
+	 * Displays a warning symbol next to the checkbox if checking the checkbox would cause an authority conflict when submitted.
+	 * 
+	 * @param InClient The local Concert client used to look up other client display info
+	 * @param InViewer Used to determine which objects the property box is displaying
+	 * @param InAuthorityCache Used to determine which client owns the property
+	 * @param ClientId ID of the client for which to check whether the property can be added to the stream
+	 * 
+	 * @return Column that can be placed in the table
+	 */
+	ConcertClientSharedSlate::ReplicationColumns::FReplicationPropertyColumn ConflictWarningForProperty(
+		TSharedRef<IConcertClient> InClient,
+		TAttribute<const ConcertClientSharedSlate::IReplicationStreamViewer*> InViewer,
+		FGlobalAuthorityCache& InAuthorityCache,
+		const FGuid& ClientId
+		);
 }
diff --git a/Engine/Plugins/Developer/Concert/ConcertApp/MultiUserClient/Source/MultiUserClient/Private/Widgets/ActiveSession/Replication/Client/Single/SingleClientColumns_ConflictWarning.cpp b/Engine/Plugins/Developer/Concert/ConcertApp/MultiUserClient/Source/MultiUserClient/Private/Widgets/ActiveSession/Replication/Client/Single/SingleClientColumns_ConflictWarning.cpp
new file mode 100644
--- /dev/null
+++ b/Engine/Plugins/Developer/Concert/ConcertApp/MultiUserClient/Source/MultiUserClient/Private/Widgets/ActiveSession/Replication/Client/Single/SingleClientColumns_ConflictWarning.cpp
@@ -0,0 +1,130 @@
+// Copyright Epic Games, Inc. All Rights Reserved.
+
+#include "SingleClientColumns.h"
+
+#include "IConcertClient.h"
+#include "Replication/Editor/Model/ReplicatedObjectData.h"
+#include "Replication/Editor/View/IReplicationStreamViewer.h"
+#include "Replication/Util/GlobalAuthorityCache.h"
+#include "SWarningIcon.h"
+
+#include "Algo/AllOf.h"
+#include "Widgets/ActiveSession/Replication/Client/ClientUtils.h"
+#include "Widgets/ClientName/SClientName.h"
+#include "Widgets/Images/SImage.h"
+#include "Widgets/Views/SHeaderRow.h"
+
+#define LOCTEXT_NAMESPACE "SingleClientColumns.ConflictWarning"
+
+namespace UE::MultiUserClient::SingleClientColumns
+{
+	const FName ConflictWarningTopLevelObjectColumnId = TEXT("ConflictWarningTopLevelObjectColumn");
+	const FName ConflictWarningSubobjectColumnId = TEXT("ConflictWarningSubobjectColumn");
+	const FName ConflictWarningPropertyColumnId = TEXT("ConflictWarningPropertyColumn");
+	
+	ConcertClientSharedSlate::ReplicationColumns::FReplicationSubobjectObjectColumn ConflictWarningForSubobject(
+		TSharedRef<IConcertClient> InClient,
+		FGlobalAuthorityCache& InAuthorityCache,
+		const FGuid& ClientId
+		)
+	{
+		using FColumnType = ConcertClientSharedSlate::ReplicationColumns::FReplicationSubobjectObjectColumn;
+		return FColumnType(
+			FColumnType::FArguments()
+				.GenerateWidgetColumn_Lambda([InClient = MoveTemp(InClient), &InAuthorityCache, ClientId](const FColumnType::FBuildArgs& InArgs) mutable
+				{
+					const auto GetVisibility = [&InAuthorityCache, ClientId, ObjectPath = InArgs.RowData.GetObjectPath()]()
+					{
+						const FGlobalAuthorityCache::ECanTakeAuthority TakeAuthorityResult = InAuthorityCache.CanClientTakeAuthority(ObjectPath, ClientId);
+						return TakeAuthorityResult == FGlobalAuthorityCache::ECanTakeAuthority::Conflict
+							? EVisibility::Visible
+							: EVisibility::Collapsed;
+					};
+					const auto GetToolTip = [InClient = MoveTemp(InClient), &InAuthorityCache, ClientId, ObjectPath = InArgs.RowData.GetObjectPath()]()
+					{
+						TSet<FString> ClientNames;
+						const FGlobalAuthorityCache::ECanTakeAuthority TakeAuthorityResult = InAuthorityCache.CanClientTakeAuthority(ObjectPath, ClientId,
+						[InClient, &ClientNames](const FGuid& ClientId, const FConcertPropertyChain& ConflictingProperty)
+						{
+							ClientNames.Add(ClientUtils::GetClientDisplayName(*InClient, ClientId));
+							return EBreakBehavior::Continue;
+						});
+						
+						return TakeAuthorityResult != FGlobalAuthorityCache::ECanTakeAuthority::Conflict
+							? FText::GetEmpty()
+							: FText::Format(
+								LOCTEXT("Subobject.WarningFmt", "{0} {1}|plural(one=is,other=are) replicating some of the assigned properties. This change will not be sent to the server."),
+								FText::FromString(FString::Join(ClientNames, TEXT(", "))),
+								ClientNames.Num()
+							);
+					};
+					
+					return SNew(SWarningIcon)
+						.Visibility_Lambda(GetVisibility)
+						.ToolTipText_Lambda(GetToolTip);
+				})
+				.PopulateSearchItems_Lambda([](const ConcertClientSharedSlate::FReplicatedObjectData&, TArray<FString>&){})
+				.ColumnSortOrder(static_cast<int32>(ESubobjectColumnOrder::ConflictWarning)),
+			SHeaderRow::Column(ConflictWarningSubobjectColumnId)
+				.DefaultLabel(FText::GetEmpty())
+				.FixedWidth(20.f)
+			);
+	}
+
+	ConcertClientSharedSlate::ReplicationColumns::FReplicationPropertyColumn ConflictWarningForProperty(
+		TSharedRef<IConcertClient> InClient,
+		TAttribute<const ConcertClientSharedSlate::IReplicationStreamViewer*> Viewer,
+		FGlobalAuthorityCache& InAuthorityCache,
+		const FGuid& ClientId
+		)
+	{
+		using FColumnType = ConcertClientSharedSlate::ReplicationColumns::FReplicationPropertyColumn;
+		return FColumnType(
+			FColumnType::FArguments()
+				.GenerateWidgetColumn_Lambda([InClient = MoveTemp(InClient), Viewer, &InAuthorityCache, ClientId](const FColumnType::FBuildArgs& InArgs) mutable
+				{
+					const auto GetVisibility = [&Viewer, &InAuthorityCache, ClientId, Property = InArgs.RowData.GetProperty()]()
+					{
+						const bool bCanAddProperty = Algo::AllOf(Viewer.Get()->GetObjectsBeingPropertyEdited(),
+							[&InAuthorityCache, ClientId, Property](const FSoftObjectPath& ObjectPath)
+							{
+								return InAuthorityCache.CanClientAddProperty(ObjectPath, ClientId, Property);
+							});
+						return bCanAddProperty
+							? EVisibility::Collapsed
+							: EVisibility::Visible;
+					};
+					const auto GetToolTip = [InClient = MoveTemp(InClient), Viewer, &InAuthorityCache, ClientId, Property = InArgs.RowData.GetProperty()]()
+					{
+						TSet<FString> Names;
+						for (const FSoftObjectPath& ObjectPath: Viewer.Get()->GetObjectsBeingPropertyEdited())
+						{
+							const TOptional<FGuid> Owner = InAuthorityCache.GetClientWithAuthorityOverProperty(ObjectPath, Property);
+							if (Owner || *Owner != ClientId)
+							{
+								Names.Add(ClientUtils::GetClientDisplayName(*InClient, *Owner));
+							}
+						}
+						
+						return Names.IsEmpty()
+							? FText::GetEmpty()
+							: FText::Format(
+								LOCTEXT("Property.WarningFmt", "Replicated by {0}. This property will not be submitted to the server."),
+								FText::FromString(FString::Join(Names, TEXT(", ")))
+								);
+					};
+					
+					return SNew(SWarningIcon)
+						.Visibility_Lambda(GetVisibility)
+						.ToolTipText_Lambda(GetToolTip);
+				})
+				.PopulateSearchItems_Lambda([](const ConcertClientSharedSlate::FReplicatedPropertyData&, TArray<FString>&){})
+				.ColumnSortOrder(static_cast<int32>(EPropertyColumnOrder::ConflictWarning)),
+			SHeaderRow::Column(ConflictWarningPropertyColumnId)
+				.DefaultLabel(FText::GetEmpty())
+				.FixedWidth(20.f)
+			);
+	}
+}
+
+#undef LOCTEXT_NAMESPACE
\ No newline at end of file
diff --git a/Engine/Plugins/Developer/Concert/ConcertApp/MultiUserClient/Source/MultiUserClient/Private/Widgets/ActiveSession/Replication/Client/Single/SingleClientColumns_Owner.cpp b/Engine/Plugins/Developer/Concert/ConcertApp/MultiUserClient/Source/MultiUserClient/Private/Widgets/ActiveSession/Replication/Client/Single/SingleClientColumns_Owner.cpp
--- a/Engine/Plugins/Developer/Concert/ConcertApp/MultiUserClient/Source/MultiUserClient/Private/Widgets/ActiveSession/Replication/Client/Single/SingleClientColumns_Owner.cpp
+++ b/Engine/Plugins/Developer/Concert/ConcertApp/MultiUserClient/Source/MultiUserClient/Private/Widgets/ActiveSession/Replication/Client/Single/SingleClientColumns_Owner.cpp
@@ -7,6 +7,7 @@
 #include "Replication/Editor/View/IReplicationStreamViewer.h"
 #include "Replication/Editor/View/ReplicationColumnsUtils.h"
 #include "Replication/Util/GlobalAuthorityCache.h"
+#include "Widgets/ActiveSession/Replication/Client/ClientUtils.h"
 #include "Widgets/ClientName/SClientName.h"
 
 #define LOCTEXT_NAMESPACE "SingleClientColumns.Owner"
@@ -24,17 +25,7 @@ namespace UE::MultiUserClient::SingleClientColumns
 		{
 			for (const FGuid& ClientId : Clients)
 			{
-				FConcertSessionClientInfo ClientInfo;
-				const TSharedPtr<IConcertClientSession> Session = InClient->GetCurrentSession();
-				const bool bIsLocalClient = ensure(Session) && Session->GetSessionClientEndpointId() == ClientId;
-				if (bIsLocalClient)
-				{
-					InOutSearchStrings.Add(ConcertClientSharedSlate::SClientName::GetDisplayText(Session->GetLocalClientInfo(), bIsLocalClient).ToString());
-				}
-				else if (ensure(Session) && Session->FindSessionClient(ClientId, ClientInfo))
-				{
-					InOutSearchStrings.Add(ConcertClientSharedSlate::SClientName::GetDisplayText(ClientInfo.ClientInfo, bIsLocalClient).ToString());
-				}
+				InOutSearchStrings.Add(ClientUtils::GetClientDisplayName(*InClient, ClientId));
 			}
 		}
 	}
diff --git a/Engine/Plugins/Developer/Concert/ConcertApp/MultiUserClient/Source/MultiUserClient/Private/Widgets/ActiveSession/Replication/Client/Single/SingleClientColumns_ToggleAuthority.cpp b/Engine/Plugins/Developer/Concert/ConcertApp/MultiUserClient/Source/MultiUserClient/Private/Widgets/ActiveSession/Replication/Client/Single/SingleClientColumns_ToggleAuthority.cpp
--- a/Engine/Plugins/Developer/Concert/ConcertApp/MultiUserClient/Source/MultiUserClient/Private/Widgets/ActiveSession/Replication/Client/Single/SingleClientColumns_ToggleAuthority.cpp
+++ b/Engine/Plugins/Developer/Concert/ConcertApp/MultiUserClient/Source/MultiUserClient/Private/Widgets/ActiveSession/Replication/Client/Single/SingleClientColumns_ToggleAuthority.cpp
@@ -200,8 +200,6 @@ namespace UE::MultiUserClient::SingleClientColumns
 						{
 						case EAuthorityMutability::Allowed: return LOCTEXT("Subobject.ChangeAuthority.ToolTip.Allowed", "Toggles whether this object should be replicated.");
 						case EAuthorityMutability::NoProperties: return LOCTEXT("Subobject.ChangeAuthority.ToolTip.NoProperties", "Toggles whether this object should be replicated.\nAssign properties to this object first.\n");
-						// TODO DP UE-198356: Make the tooltip display info about the conflicting client
-						case EAuthorityMutability::ClientConflict: return LOCTEXT("Subobject.ChangeAuthority.ToolTip.ClientConflict", "Toggles whether this object should be replicated.\nAnother client already is already replicating some of the assigned properties.");
 						case EAuthorityMutability::NotSupported: return NotSupportedText;
 						default: checkNoEntry(); return FText::GetEmpty();
 						}
diff --git a/Engine/Plugins/Developer/Concert/ConcertSync/ConcertSyncCore/Source/ConcertSyncCore/Private/Replication/AuthorityConflictSharedUtils.cpp b/Engine/Plugins/Developer/Concert/ConcertSync/ConcertSyncCore/Source/ConcertSyncCore/Private/Replication/AuthorityConflictSharedUtils.cpp
--- a/Engine/Plugins/Developer/Concert/ConcertSync/ConcertSyncCore/Source/ConcertSyncCore/Private/Replication/AuthorityConflictSharedUtils.cpp
+++ b/Engine/Plugins/Developer/Concert/ConcertSync/ConcertSyncCore/Source/ConcertSyncCore/Private/Replication/AuthorityConflictSharedUtils.cpp
@@ -2,7 +2,6 @@
 
 #include "Replication/AuthorityConflictSharedUtils.h"
 
-#include "Replication/Data/ObjectIds.h"
 #include "Replication/Data/ReplicationStreamDescription.h"
 
 namespace UE::ConcertSyncCore::Replication::AuthorityConflictUtils
@@ -24,17 +23,16 @@ namespace UE::ConcertSyncCore::Replication::AuthorityConflictUtils
 				}
 				
 				EBreakBehavior Result = EBreakBehavior::Continue;
-				GroundTruth.ForEachStream(ClientEndpointId, [&Object, &Callback, &ClientEndpointId, &GroundTruth, &Result](const FSharedReplicationStreamDescription& Stream) mutable
+				GroundTruth.ForEachStream(ClientEndpointId, [&Object, &Callback, &ClientEndpointId, &GroundTruth, &Result](const FGuid& StreamId, const FObjectReplicationMap& ReplicationMap) mutable
 				{
 					// If client has not claimed authority over this object in this stream, skip
-					const FGuid& StreamId = Stream.Identifier;
 					if (!GroundTruth.HasAuthority(ClientEndpointId, StreamId, Object))
 					{
 						return EBreakBehavior::Continue;
 					}
 
 					// This client is using the request object: report the potential conflict ...
-					const FObjectReplicationMap& ObjectReplicationMap = Stream.ReplicationMap;
+					const FObjectReplicationMap& ObjectReplicationMap = ReplicationMap;
 					const FReplicatedObjectInfo* ReplicationObjectInfo = ObjectReplicationMap.ReplicatedObjects.Find(Object);
 					if (ReplicationObjectInfo && Callback(ClientEndpointId, StreamId, ReplicationObjectInfo->PropertySelection) == EBreakBehavior::Break)
 					{
diff --git a/Engine/Plugins/Developer/Concert/ConcertSync/ConcertSyncCore/Source/ConcertSyncCore/Public/Replication/AuthorityConflictSharedUtils.h b/Engine/Plugins/Developer/Concert/ConcertSync/ConcertSyncCore/Source/ConcertSyncCore/Public/Replication/AuthorityConflictSharedUtils.h
--- a/Engine/Plugins/Developer/Concert/ConcertSync/ConcertSyncCore/Source/ConcertSyncCore/Public/Replication/AuthorityConflictSharedUtils.h
+++ b/Engine/Plugins/Developer/Concert/ConcertSync/ConcertSyncCore/Source/ConcertSyncCore/Public/Replication/AuthorityConflictSharedUtils.h
@@ -9,8 +9,8 @@
 
 struct FConcertPropertyChain;
 struct FConcertPropertySelection;
+struct FObjectReplicationMap;
 struct FReplicatedObjectId;
-struct FSharedReplicationStreamDescription;
 
 namespace UE::ConcertSyncCore::Replication::AuthorityConflictUtils
 {
@@ -30,7 +30,7 @@ namespace UE::ConcertSyncCore::Replication::AuthorityConflictUtils
 	public:
 
 		/** Provides a way to extract all streams registered to a given client. */
-		virtual void ForEachStream(const FGuid& ClientEndpointId, TFunctionRef<EBreakBehavior(const FSharedReplicationStreamDescription& Stream)> Callback) const = 0;
+		virtual void ForEachStream(const FGuid& ClientEndpointId, TFunctionRef<EBreakBehavior(const FGuid& StreamId, const FObjectReplicationMap& ReplicationMap)> Callback) const = 0;
 
 		/** Iterates through all clients have registered to send any data. */
 		virtual void ForEachSendingClient(TFunctionRef<EBreakBehavior(const FGuid& ClientEndpointId)> Callback) const = 0;
diff --git a/Engine/Plugins/Developer/Concert/ConcertSync/ConcertSyncServer/Source/ConcertSyncServer/Private/Replication/AuthorityManager.cpp b/Engine/Plugins/Developer/Concert/ConcertSync/ConcertSyncServer/Source/ConcertSyncServer/Private/Replication/AuthorityManager.cpp
--- a/Engine/Plugins/Developer/Concert/ConcertSync/ConcertSyncServer/Source/ConcertSyncServer/Private/Replication/AuthorityManager.cpp
+++ b/Engine/Plugins/Developer/Concert/ConcertSync/ConcertSyncServer/Source/ConcertSyncServer/Private/Replication/AuthorityManager.cpp
@@ -24,9 +24,12 @@ namespace UE::ConcertSyncServer::Replication
 				, Getters(Getters)
 			{}
 
-			virtual void ForEachStream(const FGuid& ClientEndpointId, TFunctionRef<EBreakBehavior(const FSharedReplicationStreamDescription& Stream)> Callback) const override
+			virtual void ForEachStream(const FGuid& ClientEndpointId, TFunctionRef<EBreakBehavior(const FGuid& StreamId, const FObjectReplicationMap& ReplicationMap)> Callback) const override
 			{
-				Getters.ForEachStream(ClientEndpointId, [&Callback](const FReplicationStreamDescription& Stream){ return Callback(Stream.BaseDescription); });
+				Getters.ForEachStream(ClientEndpointId, [&Callback](const FReplicationStreamDescription& Stream)
+				{
+					return Callback(Stream.BaseDescription.Identifier, Stream.BaseDescription.ReplicationMap);
+				});
 			}
 			
 			virtual void ForEachSendingClient(TFunctionRef<EBreakBehavior(const FGuid& ClientEndpointId)> Callback) const override
diff --git a/Engine/Plugins/Developer/Concert/ConcertSync/ConcertSyncTest/Source/ConcertSyncTest/Private/Replication/ReplicationAuthorityTests.cpp b/Engine/Plugins/Developer/Concert/ConcertSync/ConcertSyncTest/Source/ConcertSyncTest/Private/Replication/ReplicationAuthorityTests.cpp
--- a/Engine/Plugins/Developer/Concert/ConcertSync/ConcertSyncTest/Source/ConcertSyncTest/Private/Replication/ReplicationAuthorityTests.cpp
+++ b/Engine/Plugins/Developer/Concert/ConcertSync/ConcertSyncTest/Source/ConcertSyncTest/Private/Replication/ReplicationAuthorityTests.cpp
@@ -388,15 +388,15 @@ namespace UE::ConcertSyncTests::Replication::Authority
 				, bExistingClientHasAuthority(bExistingClientHasAuthority)
 			{}
 
-			virtual void ForEachStream(const FGuid& ClientEndpointId, TFunctionRef<EBreakBehavior(const FSharedReplicationStreamDescription& Stream)> Callback) const override
+			virtual void ForEachStream(const FGuid& ClientEndpointId, TFunctionRef<EBreakBehavior(const FGuid& StreamId, const FObjectReplicationMap& ReplicationMap)> Callback) const override
 			{
 				if (ClientEndpointId == RequestingClientId)
 				{
-					Callback(RequestingClientStream);
+					Callback(RequestingClientStream.Identifier, RequestingClientStream.ReplicationMap);
 				}
 				if (ClientEndpointId == ExistingClientId)
 				{
-					Callback(ExistingClientStream);
+					Callback(ExistingClientStream.Identifier, ExistingClientStream.ReplicationMap);
 				}
 			}
 			
diff --git a/Engine/Plugins/Developer/Concert/ConcertUI/ConcertClientSharedSlate/Source/ConcertClientSharedSlate/Private/Replication/Editor/View/SubobjectView/SSubobjectView.cpp b/Engine/Plugins/Developer/Concert/ConcertUI/ConcertClientSharedSlate/Source/ConcertClientSharedSlate/Private/Replication/Editor/View/SubobjectView/SSubobjectView.cpp
--- a/Engine/Plugins/Developer/Concert/ConcertUI/ConcertClientSharedSlate/Source/ConcertClientSharedSlate/Private/Replication/Editor/View/SubobjectView/SSubobjectView.cpp
+++ b/Engine/Plugins/Developer/Concert/ConcertUI/ConcertClientSharedSlate/Source/ConcertClientSharedSlate/Private/Replication/Editor/View/SubobjectView/SSubobjectView.cpp
@@ -216,7 +216,6 @@ namespace UE::ConcertClientSharedSlate
 		if (RowData.IsSeparator())
 		{
 			return SNew(SBox)
-				.Padding(1.f)
 				[
 					SNew(SSeparator)
 					.SeparatorImage(FAppStyle::Get().GetBrush("Menu.Separator"))
