Fix not sending authority requests newly added objects

diff --git a/Engine/Plugins/Developer/Concert/ConcertApp/MultiUserClient/Source/MultiUserClient/Private/Replication/Submission/SubmissionWorkflow_LocalClient.cpp b/Engine/Plugins/Developer/Concert/ConcertApp/MultiUserClient/Source/MultiUserClient/Private/Replication/Submission/SubmissionWorkflow_LocalClient.cpp
--- a/Engine/Plugins/Developer/Concert/ConcertApp/MultiUserClient/Source/MultiUserClient/Private/Replication/Submission/SubmissionWorkflow_LocalClient.cpp
+++ b/Engine/Plugins/Developer/Concert/ConcertApp/MultiUserClient/Source/MultiUserClient/Private/Replication/Submission/SubmissionWorkflow_LocalClient.cpp
@@ -57,13 +57,9 @@ namespace UE::MultiUserClient
 		{
 			return nullptr;
 		}
-
-		const FGuid LocalClientId = Client->GetConcertClient()->GetCurrentSession()->GetSessionClientEndpointId();
 		
 		// The authority request is pre-built now to avoid sending changes the local client makes while we're waiting for the latent server responses
 		FAuthorityChangeRequest AuthorityChangeRequest = AuthorityChangeTracker.BuildChangeRequest(GetLocalClientStreamId());
-		// Predict conflicts in case the remote client's streams have changed. Also: while the UI highlights "bad" requests, it does not correct it.
-		AuthorityCache.CleanseConflictsFrom(AuthorityChangeRequest, LocalClientId);
 		
 		const FStreamChangelist& Changelist = StreamChangeTracker.GetCachedDeltaChange();
 		const bool bIsChangelistEmpty = Changelist.ObjectsToPut.IsEmpty() && Changelist.ObjectsToRemove.IsEmpty();
@@ -82,21 +78,22 @@ namespace UE::MultiUserClient
 		}
 		else
 		{
-			FChangeStreamRequest Request = StreamSynchronizer.GetServerState().ReplicatedObjects.IsEmpty()
+			FChangeStreamRequest StreamRequest = StreamSynchronizer.GetServerState().ReplicatedObjects.IsEmpty()
 				? StreamRequestUtils::BuildChangeRequest_CreateNewStream(GetLocalClientStreamId(), Changelist)
 				: StreamRequestUtils::BuildChangeRequest_UpdateExistingStream(Changelist);
+			
 			// Predict conflicts in case the remote client's streams have changed. Also: while the UI highlights "bad" requests, it does not correct it.
-			AuthorityCache.CleanseConflictsFrom(AuthorityChangeRequest, LocalClientId);
+			AuthorityCache.CleanseConflictsFromStreamRequest(StreamRequest, GetLocalClientId());
 		
-			ReplicationManager->ChangeStream(Request)
-				.Next([this, DestructionDetection = LifetimeToken->AsWeak(), Request, AuthorityChangeRequest = MoveTemp(AuthorityChangeRequest)](FChangeStreamResponse&& Response)
+			ReplicationManager->ChangeStream(StreamRequest)
+				.Next([this, DestructionDetection = LifetimeToken->AsWeak(), StreamRequest, AuthorityChangeRequest = MoveTemp(AuthorityChangeRequest)](FChangeStreamResponse&& Response)
 				{
-					const FSubmitStreamChangesResponse SubmissionResult { EStreamSubmissionErrorCode::Success, { FCompletedChangeSubmission{Request, Response } } };
+					const FSubmitStreamChangesResponse SubmissionResult { EStreamSubmissionErrorCode::Success, { FCompletedChangeSubmission{StreamRequest, Response } } };
 					// The request might execute after we're destroyed, e.g. by leaving session while request is on the way.
 					// In that case, the Concert session triggers the OnSessionConnectionChanged which destroys us. Only after that, all the requests are timed out.
 					if (DestructionDetection.IsValid())
 					{
-						OnStreamChangeCompleted(Request, Response, AuthorityChangeRequest);
+						OnStreamChangeCompleted(StreamRequest, Response, AuthorityChangeRequest);
 					}
 
 					return SubmissionResult;
@@ -165,6 +162,9 @@ namespace UE::MultiUserClient
 			return;
 		}
 
+		// Predict conflicts in case the remote client's streams have changed. Also: while the UI highlights "bad" requests, it does not correct it.
+		AuthorityCache.CleanseConflictsFromAuthorityRequest(AuthorityChangeRequest, GetLocalClientId());
+		
 		FSubmitAuthorityChangesRequest Request{ EAuthoritySubmissionRequestErrorCode::Success, AuthorityChangeRequest };
 		Operation->EmplaceAuthorityRequestPromise(Request);
 		ReplicationManager->RequestAuthorityChange(MoveTemp(AuthorityChangeRequest))
@@ -184,4 +184,9 @@ namespace UE::MultiUserClient
 				}
 			});
 	}
+	
+	FGuid FSubmissionWorkflow_LocalClient::GetLocalClientId() const
+	{
+		return Client->GetConcertClient()->GetCurrentSession()->GetSessionClientEndpointId();
+	}
 }
diff --git a/Engine/Plugins/Developer/Concert/ConcertApp/MultiUserClient/Source/MultiUserClient/Private/Replication/Submission/SubmissionWorkflow_LocalClient.h b/Engine/Plugins/Developer/Concert/ConcertApp/MultiUserClient/Source/MultiUserClient/Private/Replication/Submission/SubmissionWorkflow_LocalClient.h
--- a/Engine/Plugins/Developer/Concert/ConcertApp/MultiUserClient/Source/MultiUserClient/Private/Replication/Submission/SubmissionWorkflow_LocalClient.h
+++ b/Engine/Plugins/Developer/Concert/ConcertApp/MultiUserClient/Source/MultiUserClient/Private/Replication/Submission/SubmissionWorkflow_LocalClient.h
@@ -73,6 +73,8 @@ namespace UE::MultiUserClient
 			ConcertSyncClient::Replication::FAuthorityChangeRequest AuthorityChangeRequest
 			);
 		void SendAuthorityChangeRequest(ConcertSyncClient::Replication::FAuthorityChangeRequest AuthorityChangeRequest);
+		
+		FGuid GetLocalClientId() const;
 	};
 }
 
diff --git a/Engine/Plugins/Developer/Concert/ConcertApp/MultiUserClient/Source/MultiUserClient/Private/Replication/Util/GlobalAuthorityCache.cpp b/Engine/Plugins/Developer/Concert/ConcertApp/MultiUserClient/Source/MultiUserClient/Private/Replication/Util/GlobalAuthorityCache.cpp
--- a/Engine/Plugins/Developer/Concert/ConcertApp/MultiUserClient/Source/MultiUserClient/Private/Replication/Util/GlobalAuthorityCache.cpp
+++ b/Engine/Plugins/Developer/Concert/ConcertApp/MultiUserClient/Source/MultiUserClient/Private/Replication/Util/GlobalAuthorityCache.cpp
@@ -119,14 +119,14 @@ namespace UE::MultiUserClient
 		return Result;
 	}
 
-	void FGlobalAuthorityCache::CleanseConflictsFrom(ConcertSyncClient::Replication::FAuthorityChangeRequest& Request, const FGuid& SendingClient) const
+	void FGlobalAuthorityCache::CleanseConflictsFromAuthorityRequest(ConcertSyncClient::Replication::FAuthorityChangeRequest& Request, const FGuid& SendingClient) const
 	{
-		ConcertSyncCore::Replication::AuthorityConflictUtils::CleanseConflictsFrom(Request, SendingClient, *this);
+		ConcertSyncCore::Replication::AuthorityConflictUtils::CleanseConflictsFromAuthorityRequest(Request, SendingClient, *this);
 	}
 
-	void FGlobalAuthorityCache::CleanseConflictsFrom(ConcertSyncClient::Replication::FChangeStreamRequest& Request, const FGuid& SendingClient) const
+	void FGlobalAuthorityCache::CleanseConflictsFromStreamRequest(ConcertSyncClient::Replication::FChangeStreamRequest& Request, const FGuid& SendingClient) const
 	{
-		ConcertSyncCore::Replication::AuthorityConflictUtils::CleanseConflictsFrom(Request, SendingClient, *this);
+		ConcertSyncCore::Replication::AuthorityConflictUtils::CleanseConflictsFromStreamRequest(Request, SendingClient, *this);
 	}
 
 	void FGlobalAuthorityCache::RegisterForClientEvents(const FReplicationClient& Client)
diff --git a/Engine/Plugins/Developer/Concert/ConcertApp/MultiUserClient/Source/MultiUserClient/Private/Replication/Util/GlobalAuthorityCache.h b/Engine/Plugins/Developer/Concert/ConcertApp/MultiUserClient/Source/MultiUserClient/Private/Replication/Util/GlobalAuthorityCache.h
--- a/Engine/Plugins/Developer/Concert/ConcertApp/MultiUserClient/Source/MultiUserClient/Private/Replication/Util/GlobalAuthorityCache.h
+++ b/Engine/Plugins/Developer/Concert/ConcertApp/MultiUserClient/Source/MultiUserClient/Private/Replication/Util/GlobalAuthorityCache.h
@@ -63,14 +63,14 @@ namespace UE::MultiUserClient
 		 * @param Request The request to clense
 		 * @param SendingClient The client that will send the request
 		 */
-		void CleanseConflictsFrom(ConcertSyncClient::Replication::FAuthorityChangeRequest& Request, const FGuid& SendingClient) const;
+		void CleanseConflictsFromAuthorityRequest(ConcertSyncClient::Replication::FAuthorityChangeRequest& Request, const FGuid& SendingClient) const;
 		/**
 		 * Removes entries from Request that would generate conflicts.
 		 * 
 		 * @param Request The request to clense
 		 * @param SendingClient The client that will send the request
 		 */
-		void CleanseConflictsFrom(ConcertSyncClient::Replication::FChangeStreamRequest& Request, const FGuid& SendingClient) const;
+		void CleanseConflictsFromStreamRequest(ConcertSyncClient::Replication::FChangeStreamRequest& Request, const FGuid& SendingClient) const;
 		
 		DECLARE_MULTICAST_DELEGATE_OneParam(FOnCacheChanged, const FGuid& ClientId);
 		/** Called when the cache changes for a specific client. */
diff --git a/Engine/Plugins/Developer/Concert/ConcertSync/ConcertSyncCore/Source/ConcertSyncCore/Private/Replication/AuthorityConflictSharedUtils.cpp b/Engine/Plugins/Developer/Concert/ConcertSync/ConcertSyncCore/Source/ConcertSyncCore/Private/Replication/AuthorityConflictSharedUtils.cpp
--- a/Engine/Plugins/Developer/Concert/ConcertSync/ConcertSyncCore/Source/ConcertSyncCore/Private/Replication/AuthorityConflictSharedUtils.cpp
+++ b/Engine/Plugins/Developer/Concert/ConcertSync/ConcertSyncCore/Source/ConcertSyncCore/Private/Replication/AuthorityConflictSharedUtils.cpp
@@ -84,7 +84,7 @@ namespace UE::ConcertSyncCore::Replication::AuthorityConflictUtils
 		return bFreeOfConflicts ? EAuthorityConflict::Allowed : EAuthorityConflict::Conflict;
 	}
 
-	void CleanseConflictsFrom(FConcertReplication_ChangeAuthority_Request& Request, const FGuid& SendingClient, const IReplicationGroundTruth& GroundTruth)
+	void CleanseConflictsFromAuthorityRequest(FConcertReplication_ChangeAuthority_Request& Request, const FGuid& SendingClient, const IReplicationGroundTruth& GroundTruth)
 	{
 		// Need to check whether TakeAuthority is taking authority over properties other clients are already replicating
 		GroundTruth.ForEachStream(SendingClient, [&SendingClient, &Request, &GroundTruth](const FGuid& StreamId, const FObjectReplicationMap& ReplicationMap)
@@ -120,7 +120,7 @@ namespace UE::ConcertSyncCore::Replication::AuthorityConflictUtils
 		});
 	}
 	
-	void CleanseConflictsFrom(FConcertReplication_ChangeStream_Request& Request, const FGuid& SendingClient, const IReplicationGroundTruth& GroundTruth)
+	void CleanseConflictsFromStreamRequest(FConcertReplication_ChangeStream_Request& Request, const FGuid& SendingClient, const IReplicationGroundTruth& GroundTruth)
 	{
 		// Need to check whether ObjectsToPut adds any properties that an existing client has authority over.
 		GroundTruth.ForEachStream(SendingClient, [&SendingClient, &Request, &GroundTruth](const FGuid& StreamId, const FObjectReplicationMap& ReplicationMap)
diff --git a/Engine/Plugins/Developer/Concert/ConcertSync/ConcertSyncCore/Source/ConcertSyncCore/Public/Replication/AuthorityConflictSharedUtils.h b/Engine/Plugins/Developer/Concert/ConcertSync/ConcertSyncCore/Source/ConcertSyncCore/Public/Replication/AuthorityConflictSharedUtils.h
--- a/Engine/Plugins/Developer/Concert/ConcertSync/ConcertSyncCore/Source/ConcertSyncCore/Public/Replication/AuthorityConflictSharedUtils.h
+++ b/Engine/Plugins/Developer/Concert/ConcertSync/ConcertSyncCore/Source/ConcertSyncCore/Public/Replication/AuthorityConflictSharedUtils.h
@@ -72,7 +72,7 @@ namespace UE::ConcertSyncCore::Replication::AuthorityConflictUtils
 	 * @param GroundTruth Provides information about clients. On the server, it is the definite state whereas on client machines it would be
 	 * what the local client thinks what the server state is.
 	 */
-	CONCERTSYNCCORE_API void CleanseConflictsFrom(FConcertReplication_ChangeAuthority_Request& Request, const FGuid& SendingClient, const IReplicationGroundTruth& GroundTruth);
+	CONCERTSYNCCORE_API void CleanseConflictsFromAuthorityRequest(FConcertReplication_ChangeAuthority_Request& Request, const FGuid& SendingClient, const IReplicationGroundTruth& GroundTruth);
 	
 	/**
 	 * Removes entries from Request that would generate conflicts.
@@ -82,5 +82,5 @@ namespace UE::ConcertSyncCore::Replication::AuthorityConflictUtils
 	 * @param GroundTruth Provides information about clients. On the server, it is the definite state whereas on client machines it would be
 	 * what the local client thinks what the server state is.
 	 */
-	CONCERTSYNCCORE_API void CleanseConflictsFrom(FConcertReplication_ChangeStream_Request& Request, const FGuid& SendingClient, const IReplicationGroundTruth& GroundTruth);
+	CONCERTSYNCCORE_API void CleanseConflictsFromStreamRequest(FConcertReplication_ChangeStream_Request& Request, const FGuid& SendingClient, const IReplicationGroundTruth& GroundTruth);
 }
