Authority checkboxes can now only be checked if it does not create any conflicts

diff --git a/Engine/Plugins/Developer/Concert/ConcertApp/MultiUserClient/Source/MultiUserClient/Private/Replication/Authority/AuthorityChangeTracker.cpp b/Engine/Plugins/Developer/Concert/ConcertApp/MultiUserClient/Source/MultiUserClient/Private/Replication/Authority/AuthorityChangeTracker.cpp
--- a/Engine/Plugins/Developer/Concert/ConcertApp/MultiUserClient/Source/MultiUserClient/Private/Replication/Authority/AuthorityChangeTracker.cpp
+++ b/Engine/Plugins/Developer/Concert/ConcertApp/MultiUserClient/Source/MultiUserClient/Private/Replication/Authority/AuthorityChangeTracker.cpp
@@ -3,21 +3,26 @@
 #include "AuthorityChangeTracker.h"
 
 #include "IClientAuthoritySynchronizer.h"
+#include "Replication/Util/GlobalAuthorityCache.h"
 
 namespace UE::MultiUserClient
 {
 	FAuthorityChangeTracker::FAuthorityChangeTracker(
-		IClientAuthoritySynchronizer& InAuthoritySynchronizer
+		const FGuid& InClientId,
+		const IClientAuthoritySynchronizer& InAuthoritySynchronizer,
+		FGlobalAuthorityCache& InAuthorityCache
 		)
-		: AuthoritySynchronizer(InAuthoritySynchronizer)
+		: ClientId(InClientId)
+		, AuthoritySynchronizer(InAuthoritySynchronizer)
+		, AuthorityCache(InAuthorityCache)
 	{
-		AuthoritySynchronizer.OnServerStateChanged().AddRaw(this, &FAuthorityChangeTracker::RefreshChanges);
+		AuthorityCache.OnCacheChanged().AddRaw(this, &FAuthorityChangeTracker::OnClientChanged);
 	}
 
 	FAuthorityChangeTracker::~FAuthorityChangeTracker()
 	{
 		// Strictly not needed because AuthoritySynchronizer will die at the same time as us but let's be nice
-		AuthoritySynchronizer.OnServerStateChanged().RemoveAll(this);
+		AuthorityCache.OnCacheChanged().RemoveAll(this);
 	}
 
 	void FAuthorityChangeTracker::SetAuthorityIfAllowed(TConstArrayView<FSoftObjectPath> ObjectPaths, bool bNewAuthorityState)
@@ -60,7 +65,8 @@ namespace UE::MultiUserClient
 
 	bool FAuthorityChangeTracker::CanSetAuthorityFor(const FSoftObjectPath& ObjectPath) const
 	{
-		return AuthoritySynchronizer.CanChangeAuthority(ObjectPath);
+		return AuthoritySynchronizer.CanChangeAuthority(ObjectPath)
+			&& AuthorityCache.CanClientTakeAuthorityAfterSubmission(ObjectPath, ClientId) != FGlobalAuthorityCache::ECanTakeAuthority::Conflict;
 	}
 
 	EAuthorityMutability FAuthorityChangeTracker::GetChangeAuthorityMutability(const FSoftObjectPath& ObjectPath) const
diff --git a/Engine/Plugins/Developer/Concert/ConcertApp/MultiUserClient/Source/MultiUserClient/Private/Replication/Authority/AuthorityChangeTracker.h b/Engine/Plugins/Developer/Concert/ConcertApp/MultiUserClient/Source/MultiUserClient/Private/Replication/Authority/AuthorityChangeTracker.h
--- a/Engine/Plugins/Developer/Concert/ConcertApp/MultiUserClient/Source/MultiUserClient/Private/Replication/Authority/AuthorityChangeTracker.h
+++ b/Engine/Plugins/Developer/Concert/ConcertApp/MultiUserClient/Source/MultiUserClient/Private/Replication/Authority/AuthorityChangeTracker.h
@@ -6,6 +6,11 @@
 #include "Containers/ContainersFwd.h"
 #include "Replication/IConcertClientReplicationManager.h"
 
+namespace UE::MultiUserClient
+{
+	class FGlobalAuthorityCache;
+}
+
 struct FSoftObjectPath;
 
 namespace UE::MultiUserClient
@@ -17,8 +22,8 @@ namespace UE::MultiUserClient
 	class FAuthorityChangeTracker
 	{
 	public:
-
-		FAuthorityChangeTracker(IClientAuthoritySynchronizer& InAuthoritySynchronizer);
+		
+		FAuthorityChangeTracker(const FGuid& InClientId, const IClientAuthoritySynchronizer& InAuthoritySynchronizer, FGlobalAuthorityCache& InAuthorityCache);
 		~FAuthorityChangeTracker();
 		
 		/**
@@ -49,14 +54,21 @@ namespace UE::MultiUserClient
 		
 	private:
 
-		/** Knows the current authority state of the client. */
-		IClientAuthoritySynchronizer& AuthoritySynchronizer;
+		/** Id of the client this change tracker is tracking. */
+		const FGuid ClientId;
+
+		/** Knows the current authority state of the client and determines whether we support changing this client's authority at all. */
+		const IClientAuthoritySynchronizer& AuthoritySynchronizer;
+		/** Used to determine whether other clients have authority over objects. */
+		FGlobalAuthorityCache& AuthorityCache;
 
 		/** Object to the authority state it should have. */
 		TMap<FSoftObjectPath, bool> NewAuthorityStates;
 
 		/** Called when NewAuthorityStates is updated. */
 		FOnAuthorityChangeMade FOnAuthorityChangeMadeDelegate;
+		
+		void OnClientChanged(const FGuid& Guid) { RefreshChanges(); }
 	};
 }
 
diff --git a/Engine/Plugins/Developer/Concert/ConcertApp/MultiUserClient/Source/MultiUserClient/Private/Replication/Client/LocalReplicationClient.cpp b/Engine/Plugins/Developer/Concert/ConcertApp/MultiUserClient/Source/MultiUserClient/Private/Replication/Client/LocalReplicationClient.cpp
--- a/Engine/Plugins/Developer/Concert/ConcertApp/MultiUserClient/Source/MultiUserClient/Private/Replication/Client/LocalReplicationClient.cpp
+++ b/Engine/Plugins/Developer/Concert/ConcertApp/MultiUserClient/Source/MultiUserClient/Private/Replication/Client/LocalReplicationClient.cpp
@@ -9,13 +9,14 @@
 namespace UE::MultiUserClient
 {
 	FLocalReplicationClient::FLocalReplicationClient(
+		FGlobalAuthorityCache& InAuthorityCache,
 		UMultiUserReplicationClientPreset& InSessionContent,
 		TUniquePtr<IClientStreamSynchronizer> InStreamSynchronizer,
-		TSharedRef<IConcertSyncClient> InClient,
-		const FGlobalAuthorityCache& InAuthorityCache
+		TSharedRef<IConcertSyncClient> InClient
 		)
 		: FReplicationClient(
 			InClient->GetConcertClient()->GetCurrentSession()->GetSessionClientEndpointId(),
+			InAuthorityCache,
 			InSessionContent,
 			MoveTemp(InStreamSynchronizer),
 		MakeUnique<FAuthoritySynchronizer_LocalClient>(
diff --git a/Engine/Plugins/Developer/Concert/ConcertApp/MultiUserClient/Source/MultiUserClient/Private/Replication/Client/LocalReplicationClient.h b/Engine/Plugins/Developer/Concert/ConcertApp/MultiUserClient/Source/MultiUserClient/Private/Replication/Client/LocalReplicationClient.h
--- a/Engine/Plugins/Developer/Concert/ConcertApp/MultiUserClient/Source/MultiUserClient/Private/Replication/Client/LocalReplicationClient.h
+++ b/Engine/Plugins/Developer/Concert/ConcertApp/MultiUserClient/Source/MultiUserClient/Private/Replication/Client/LocalReplicationClient.h
@@ -14,10 +14,10 @@ namespace UE::MultiUserClient
 	public:
 
 		FLocalReplicationClient(
+			FGlobalAuthorityCache& InAuthorityCache,
 			UMultiUserReplicationClientPreset& InSessionContent,
 			TUniquePtr<IClientStreamSynchronizer> InStreamSynchronizer,
-			TSharedRef<IConcertSyncClient> InClient,
-			const FGlobalAuthorityCache& InAuthorityCache
+			TSharedRef<IConcertSyncClient> InClient
 			);
 	};
 }
diff --git a/Engine/Plugins/Developer/Concert/ConcertApp/MultiUserClient/Source/MultiUserClient/Private/Replication/Client/RemoteReplicationClient.cpp b/Engine/Plugins/Developer/Concert/ConcertApp/MultiUserClient/Source/MultiUserClient/Private/Replication/Client/RemoteReplicationClient.cpp
--- a/Engine/Plugins/Developer/Concert/ConcertApp/MultiUserClient/Source/MultiUserClient/Private/Replication/Client/RemoteReplicationClient.cpp
+++ b/Engine/Plugins/Developer/Concert/ConcertApp/MultiUserClient/Source/MultiUserClient/Private/Replication/Client/RemoteReplicationClient.cpp
@@ -11,11 +11,13 @@ namespace UE::MultiUserClient
 {
 	FRemoteReplicationClient::FRemoteReplicationClient(
 		const FGuid& InConcertClientId,
+		FGlobalAuthorityCache& InAuthorityCache,
 		UMultiUserReplicationClientPreset& InSessionContent,
 		FRegularQueryService& QueryService
 		)
 		: FReplicationClient(
 			InConcertClientId,
+			InAuthorityCache,
 			InSessionContent,
 			MakeUnique<FStreamSynchronizer_RemoteClient>(InConcertClientId, QueryService),
 			MakeUnique<FAuthoritySynchronizer_RemoteClient>(
diff --git a/Engine/Plugins/Developer/Concert/ConcertApp/MultiUserClient/Source/MultiUserClient/Private/Replication/Client/RemoteReplicationClient.h b/Engine/Plugins/Developer/Concert/ConcertApp/MultiUserClient/Source/MultiUserClient/Private/Replication/Client/RemoteReplicationClient.h
--- a/Engine/Plugins/Developer/Concert/ConcertApp/MultiUserClient/Source/MultiUserClient/Private/Replication/Client/RemoteReplicationClient.h
+++ b/Engine/Plugins/Developer/Concert/ConcertApp/MultiUserClient/Source/MultiUserClient/Private/Replication/Client/RemoteReplicationClient.h
@@ -15,6 +15,7 @@ namespace UE::MultiUserClient
 
 		FRemoteReplicationClient(
 			const FGuid& InConcertClientId,
+			FGlobalAuthorityCache& InAuthorityCache,
 			UMultiUserReplicationClientPreset& InSessionContent,
 			FRegularQueryService& QueryService
 			);
diff --git a/Engine/Plugins/Developer/Concert/ConcertApp/MultiUserClient/Source/MultiUserClient/Private/Replication/Client/ReplicationClient.cpp b/Engine/Plugins/Developer/Concert/ConcertApp/MultiUserClient/Source/MultiUserClient/Private/Replication/Client/ReplicationClient.cpp
--- a/Engine/Plugins/Developer/Concert/ConcertApp/MultiUserClient/Source/MultiUserClient/Private/Replication/Client/ReplicationClient.cpp
+++ b/Engine/Plugins/Developer/Concert/ConcertApp/MultiUserClient/Source/MultiUserClient/Private/Replication/Client/ReplicationClient.cpp
@@ -11,6 +11,7 @@ namespace UE::MultiUserClient
 {
 	FReplicationClient::FReplicationClient(
 		const FGuid& EndpointId,
+		FGlobalAuthorityCache& InAuthorityCache,
 		UMultiUserReplicationClientPreset& InSessionContent,
 		TUniquePtr<IClientStreamSynchronizer> InStreamSynchronizer,
 		TUniquePtr<IClientAuthoritySynchronizer> InAuthoritySynchronizer,
@@ -29,7 +30,7 @@ namespace UE::MultiUserClient
 			ClientContentStorage->Stream->MakeReplicationMapGetterAttribute(),
 			FStreamChangeTracker::FOnModifyReplicationMap::CreateLambda([this](){ ClientContentStorage->Stream->Modify(); })
 			)
-		, LocalAuthorityDiffer(*AuthoritySynchronizer)
+		, LocalAuthorityDiffer(EndpointId, *AuthoritySynchronizer, InAuthorityCache)
 		, SubmissionWorkflow(MakeSubmissionWorkflowFunc(LocalClientStreamDiffer, LocalAuthorityDiffer, *StreamSynchronizer.Get()))
 		, AutoSubmissionPolicy(*SubmissionWorkflow.Get(), LocalClientEditModel.Get(), LocalAuthorityDiffer)
 	{
@@ -52,7 +53,7 @@ namespace UE::MultiUserClient
 		// This must be done before SetAuthorityIfAllowed because it uses the cache for checking whether the object has properties assigned
 		LocalClientStreamDiffer.RefreshChangesCache();
 		
-		// Better UX for user: automatically take authority for newly added objects
+		// Better UX for user: automatically take authority for newly added objects (but only if it is allowed and causes no conflicts)
 		TArray<FSoftObjectPath> ObjectPaths;
 		Algo::Transform(AddedObjects, ObjectPaths, [](const UObject* Object){ return FSoftObjectPath(Object); });
 		LocalAuthorityDiffer.SetAuthorityIfAllowed(ObjectPaths, true);
diff --git a/Engine/Plugins/Developer/Concert/ConcertApp/MultiUserClient/Source/MultiUserClient/Private/Replication/Client/ReplicationClient.h b/Engine/Plugins/Developer/Concert/ConcertApp/MultiUserClient/Source/MultiUserClient/Private/Replication/Client/ReplicationClient.h
--- a/Engine/Plugins/Developer/Concert/ConcertApp/MultiUserClient/Source/MultiUserClient/Private/Replication/Client/ReplicationClient.h
+++ b/Engine/Plugins/Developer/Concert/ConcertApp/MultiUserClient/Source/MultiUserClient/Private/Replication/Client/ReplicationClient.h
@@ -34,6 +34,7 @@ namespace UE::MultiUserClient
 
 		FReplicationClient(
 			const FGuid& EndpointId,
+			FGlobalAuthorityCache& InAuthorityCache,
 			UMultiUserReplicationClientPreset& InSessionContent,
 			TUniquePtr<IClientStreamSynchronizer> InStreamSynchronizer,
 			TUniquePtr<IClientAuthoritySynchronizer> InAuthoritySynchronizer,
diff --git a/Engine/Plugins/Developer/Concert/ConcertApp/MultiUserClient/Source/MultiUserClient/Private/Replication/Client/ReplicationClientManager.cpp b/Engine/Plugins/Developer/Concert/ConcertApp/MultiUserClient/Source/MultiUserClient/Private/Replication/Client/ReplicationClientManager.cpp
--- a/Engine/Plugins/Developer/Concert/ConcertApp/MultiUserClient/Source/MultiUserClient/Private/Replication/Client/ReplicationClientManager.cpp
+++ b/Engine/Plugins/Developer/Concert/ConcertApp/MultiUserClient/Source/MultiUserClient/Private/Replication/Client/ReplicationClientManager.cpp
@@ -24,7 +24,7 @@ namespace UE::MultiUserClient
 		, LocalClient([this, InClient]()
 		{
 			UMultiUserReplicationClientPreset* ClientPreset = SessionContent->AddClient();
-			return FLocalReplicationClient(*ClientPreset, MakeUnique<FStreamSynchronizer_LocalClient>(InClient, ClientPreset->Stream->StreamId), InClient, AuthorityCache);
+			return FLocalReplicationClient(AuthorityCache, *ClientPreset, MakeUnique<FStreamSynchronizer_LocalClient>(InClient, ClientPreset->Stream->StreamId), InClient);
 		}())
 		, SubmissionNotifier(*this)
 	{
@@ -110,7 +110,7 @@ namespace UE::MultiUserClient
 	
 	void FReplicationClientManager::CreateRemoteClient(const FGuid& ClientEndpointId, bool bBroadcastDelegate)
 	{
-		TUniquePtr<FRemoteReplicationClient> RemoteClientPtr = MakeUnique<FRemoteReplicationClient>(ClientEndpointId, *SessionContent->AddClient(), QueryService);
+		TUniquePtr<FRemoteReplicationClient> RemoteClientPtr = MakeUnique<FRemoteReplicationClient>(ClientEndpointId, AuthorityCache, *SessionContent->AddClient(), QueryService);
 		FRemoteReplicationClient& RemoteClient = *RemoteClientPtr;
 		RemoteClients.Emplace(
 			MoveTemp(RemoteClientPtr)
diff --git a/Engine/Plugins/Developer/Concert/ConcertApp/MultiUserClient/Source/MultiUserClient/Private/Replication/Stream/StreamChangeTracker.cpp b/Engine/Plugins/Developer/Concert/ConcertApp/MultiUserClient/Source/MultiUserClient/Private/Replication/Stream/StreamChangeTracker.cpp
--- a/Engine/Plugins/Developer/Concert/ConcertApp/MultiUserClient/Source/MultiUserClient/Private/Replication/Stream/StreamChangeTracker.cpp
+++ b/Engine/Plugins/Developer/Concert/ConcertApp/MultiUserClient/Source/MultiUserClient/Private/Replication/Stream/StreamChangeTracker.cpp
@@ -72,6 +72,20 @@ namespace UE::MultiUserClient
 		}
 	}
 
+	const FConcertPropertySelection* FStreamChangeTracker::GetPropertiesAfterSubmit(const FSoftObjectPath& ObjectPath) const
+	{
+		const FGuid StreamId = StreamSynchronizer.GetStreamId();
+		const FObjectInStreamID ObjectId { StreamId, ObjectPath };
+		const FConcertReplication_ChangeStream_PutObject* PutObject = CachedDeltaChange.ObjectsToPut.Find(ObjectId);
+		if (PutObject && !PutObject->Properties.ReplicatedProperties.IsEmpty())
+		{
+			return &PutObject->Properties;
+		}
+
+		const FReplicatedObjectInfo* ObjectInfo = StreamSynchronizer.GetServerState().ReplicatedObjects.Find(ObjectPath);
+		return ObjectInfo ? &ObjectInfo->PropertySelection : nullptr;
+	}
+
 	FStreamChangeTracker::EObjectChangeType FStreamChangeTracker::GetObjectChanges(const FSoftObjectPath& Object) const
 	{
 		const FGuid StreamId = StreamSynchronizer.GetStreamId();
@@ -90,12 +104,6 @@ namespace UE::MultiUserClient
 		const bool bObjectIsOnServer = StreamSynchronizer.GetServerState().ReplicatedObjects.Contains(Object);
 		return bObjectIsOnServer ? EObjectChangeType::PropertiesModified : EObjectChangeType::Added;
 	}
-
-	FStreamChangeTracker::EPropertyChangeType FStreamChangeTracker::GetPropertyChanges(const FSoftObjectPath& Object, const FConcertPropertyChain& PropertyChain) const
-	{
-		// TODO DP:
-		return EPropertyChangeType::NoChange;
-	}
 	
 	FStreamChangelist FStreamChangeTracker::DiffChanges(
 		const FGuid& StreamId,
diff --git a/Engine/Plugins/Developer/Concert/ConcertApp/MultiUserClient/Source/MultiUserClient/Private/Replication/Stream/StreamChangeTracker.h b/Engine/Plugins/Developer/Concert/ConcertApp/MultiUserClient/Source/MultiUserClient/Private/Replication/Stream/StreamChangeTracker.h
--- a/Engine/Plugins/Developer/Concert/ConcertApp/MultiUserClient/Source/MultiUserClient/Private/Replication/Stream/StreamChangeTracker.h
+++ b/Engine/Plugins/Developer/Concert/ConcertApp/MultiUserClient/Source/MultiUserClient/Private/Replication/Stream/StreamChangeTracker.h
@@ -70,8 +70,10 @@ namespace UE::MultiUserClient
 		/** @return Whether there are any changes that that can be submitted to the server (excludes those with warnings). */
 		bool HasChanges() const;
 
-		/** @return If change are submitted, is ObjectPath in the stream? */
+		/** @return If changes are submitted, is ObjectPath in the stream? */
 		bool DoesObjectHavePropertiesAfterSubmit(const FSoftObjectPath& ObjectPath) const;
+		/** @return Gets the final state the object will have after submission. */
+		const FConcertPropertySelection* GetPropertiesAfterSubmit(const FSoftObjectPath& ObjectPath) const;
 
 		enum class EObjectChangeType
 		{
@@ -82,15 +84,6 @@ namespace UE::MultiUserClient
 		};
 		/** @return Whether this object's local configuration differs from the server version. */
 		EObjectChangeType GetObjectChanges(const FSoftObjectPath& Object) const;
-		
-		enum class EPropertyChangeType
-		{
-			NoChange,
-			Added,
-			Removed
-		};
-		/** @return Gets how the given property config differs from what is on the server. */
-		EPropertyChangeType GetPropertyChanges(const FSoftObjectPath& Object, const FConcertPropertyChain& PropertyChain) const;
 
 		/** Called when RevertCachedChanges is called. The UI should be refreshed. */
 		DECLARE_MULTICAST_DELEGATE(FOnChangesReverted);
diff --git a/Engine/Plugins/Developer/Concert/ConcertApp/MultiUserClient/Source/MultiUserClient/Private/Replication/Util/GlobalAuthorityCache.cpp b/Engine/Plugins/Developer/Concert/ConcertApp/MultiUserClient/Source/MultiUserClient/Private/Replication/Util/GlobalAuthorityCache.cpp
--- a/Engine/Plugins/Developer/Concert/ConcertApp/MultiUserClient/Source/MultiUserClient/Private/Replication/Util/GlobalAuthorityCache.cpp
+++ b/Engine/Plugins/Developer/Concert/ConcertApp/MultiUserClient/Source/MultiUserClient/Private/Replication/Util/GlobalAuthorityCache.cpp
@@ -48,16 +48,17 @@ namespace UE::MultiUserClient
 		return Result;
 	}
 
-	FGlobalAuthorityCache::ECanTakeAuthority FGlobalAuthorityCache::CanClientTakeAuthority(const FSoftObjectPath& Object, const FGuid& ClientId, FProcessPropertyConflict ProcessConflict) const
+	FGlobalAuthorityCache::ECanTakeAuthority FGlobalAuthorityCache::CanClientTakeAuthorityAfterSubmission(const FSoftObjectPath& Object, const FGuid& ClientId, FProcessPropertyConflict ProcessConflict) const
 	{
 		const FReplicationClient* Client = ClientManager.FindClient(ClientId);
 		if (!ensure(Client))
 		{
 			return ECanTakeAuthority::NotApplicable;
 		}
-		
-		const FReplicatedObjectInfo* ObjectInfo = Client->GetStreamSynchronizer().GetServerState().ReplicatedObjects.Find(Object);
-		if (!ObjectInfo)
+
+		// Important: Get server state with local changes applied to it! CanClientTakeAuthority answers: "Can the client take authority after submitting?" 
+		const FConcertPropertySelection* PropertySelection = Client->GetStreamDiffer().GetPropertiesAfterSubmit(Object);
+		if (!PropertySelection)
 		{
 			// Nothing to take authority over
 			return ECanTakeAuthority::NotApplicable;
@@ -67,7 +68,7 @@ namespace UE::MultiUserClient
 		const EAuthorityConflict Conflict = EnumerateAuthorityConflicts(
 			ClientId,
 			Object,
-			ObjectInfo->PropertySelection.ReplicatedProperties,
+			PropertySelection->ReplicatedProperties,
 			*this,
 			[&ProcessConflict](const FGuid& ClientId, const FGuid&, const FConcertPropertyChain& ConflictingProperty)
 			{
diff --git a/Engine/Plugins/Developer/Concert/ConcertApp/MultiUserClient/Source/MultiUserClient/Private/Replication/Util/GlobalAuthorityCache.h b/Engine/Plugins/Developer/Concert/ConcertApp/MultiUserClient/Source/MultiUserClient/Private/Replication/Util/GlobalAuthorityCache.h
--- a/Engine/Plugins/Developer/Concert/ConcertApp/MultiUserClient/Source/MultiUserClient/Private/Replication/Util/GlobalAuthorityCache.h
+++ b/Engine/Plugins/Developer/Concert/ConcertApp/MultiUserClient/Source/MultiUserClient/Private/Replication/Util/GlobalAuthorityCache.h
@@ -49,8 +49,8 @@ namespace UE::MultiUserClient
 			Allowed
 		};
 		
-		/** @return Whether the given client can take authority over the object without causing any conflicts. */
-		ECanTakeAuthority CanClientTakeAuthority(const FSoftObjectPath& Object, const FGuid& ClientId, FProcessPropertyConflict ProcessConflict = [](auto&, auto&){ return EBreakBehavior::Break; }) const;
+		/** @return Whether the given client can take authority over the object without causing any conflicts. This also considers the changes made to the stream after submission. */
+		ECanTakeAuthority CanClientTakeAuthorityAfterSubmission(const FSoftObjectPath& Object, const FGuid& ClientId, FProcessPropertyConflict ProcessConflict = [](auto&, auto&){ return EBreakBehavior::Break; }) const;
 		/** @return Whether the given client add the given property to the object without causing any conflicts. */
 		bool CanClientAddProperty(const FSoftObjectPath& Object, const FGuid& ClientId, const FConcertPropertyChain& Chain) const;
 
diff --git a/Engine/Plugins/Developer/Concert/ConcertApp/MultiUserClient/Source/MultiUserClient/Private/Widgets/ActiveSession/Replication/Client/Single/SingleClientColumns_ConflictWarning.cpp b/Engine/Plugins/Developer/Concert/ConcertApp/MultiUserClient/Source/MultiUserClient/Private/Widgets/ActiveSession/Replication/Client/Single/SingleClientColumns_ConflictWarning.cpp
--- a/Engine/Plugins/Developer/Concert/ConcertApp/MultiUserClient/Source/MultiUserClient/Private/Widgets/ActiveSession/Replication/Client/Single/SingleClientColumns_ConflictWarning.cpp
+++ b/Engine/Plugins/Developer/Concert/ConcertApp/MultiUserClient/Source/MultiUserClient/Private/Widgets/ActiveSession/Replication/Client/Single/SingleClientColumns_ConflictWarning.cpp
@@ -35,7 +35,7 @@ namespace UE::MultiUserClient::SingleClientColumns
 				{
 					const auto GetVisibility = [&InAuthorityCache, ClientId, ObjectPath = InArgs.RowData.GetObjectPath()]()
 					{
-						const FGlobalAuthorityCache::ECanTakeAuthority TakeAuthorityResult = InAuthorityCache.CanClientTakeAuthority(ObjectPath, ClientId);
+						const FGlobalAuthorityCache::ECanTakeAuthority TakeAuthorityResult = InAuthorityCache.CanClientTakeAuthorityAfterSubmission(ObjectPath, ClientId);
 						return TakeAuthorityResult == FGlobalAuthorityCache::ECanTakeAuthority::Conflict
 							? EVisibility::Visible
 							: EVisibility::Collapsed;
@@ -43,7 +43,7 @@ namespace UE::MultiUserClient::SingleClientColumns
 					const auto GetToolTip = [InClient = MoveTemp(InClient), &InAuthorityCache, ClientId, ObjectPath = InArgs.RowData.GetObjectPath()]()
 					{
 						TSet<FString> ClientNames;
-						const FGlobalAuthorityCache::ECanTakeAuthority TakeAuthorityResult = InAuthorityCache.CanClientTakeAuthority(ObjectPath, ClientId,
+						const FGlobalAuthorityCache::ECanTakeAuthority TakeAuthorityResult = InAuthorityCache.CanClientTakeAuthorityAfterSubmission(ObjectPath, ClientId,
 						[InClient, &ClientNames](const FGuid& ClientId, const FConcertPropertyChain& ConflictingProperty)
 						{
 							ClientNames.Add(ClientUtils::GetClientDisplayName(*InClient, ClientId));
