When authority change is rejected, the corresponding checkboxes should be unchecked.

diff --git a/Engine/Plugins/Developer/Concert/ConcertApp/MultiUserClient/Source/MultiUserClient/Private/Replication/Authority/AuthorityChangeTracker.cpp b/Engine/Plugins/Developer/Concert/ConcertApp/MultiUserClient/Source/MultiUserClient/Private/Replication/Authority/AuthorityChangeTracker.cpp
--- a/Engine/Plugins/Developer/Concert/ConcertApp/MultiUserClient/Source/MultiUserClient/Private/Replication/Authority/AuthorityChangeTracker.cpp
+++ b/Engine/Plugins/Developer/Concert/ConcertApp/MultiUserClient/Source/MultiUserClient/Private/Replication/Authority/AuthorityChangeTracker.cpp
@@ -27,17 +27,31 @@ namespace UE::MultiUserClient
 
 	void FAuthorityChangeTracker::SetAuthorityIfAllowed(TConstArrayView<FSoftObjectPath> ObjectPaths, bool bNewAuthorityState)
 	{
+		bool bAddedAnything = false;
+		
 		for (const FSoftObjectPath& ObjectPath : ObjectPaths)
 		{
-			if (CanSetAuthorityFor(ObjectPath))
+			bool* bCurrentValue = NewAuthorityStates.Find(ObjectPath);
+			const bool bValueWasChanged = !bCurrentValue || *bCurrentValue != bNewAuthorityState;
+			
+			if (bValueWasChanged && CanSetAuthorityFor(ObjectPath))
 			{
+				bAddedAnything = true;
 				NewAuthorityStates.Add(ObjectPath, bNewAuthorityState);
 			}
 		}
 
-		if (HasChanges())
+		if (bAddedAnything)
+		{
+			OnAddedOwnedObjectsDelegate.Broadcast();
+		}
+	}
+
+	void FAuthorityChangeTracker::ClearAuthorityChange(TConstArrayView<FSoftObjectPath> ObjectPaths)
+	{
+		for (const FSoftObjectPath& ObjectPath : ObjectPaths)
 		{
-			FOnAuthorityChangeMadeDelegate.Broadcast();
+			NewAuthorityStates.Remove(ObjectPath);
 		}
 	}
 
diff --git a/Engine/Plugins/Developer/Concert/ConcertApp/MultiUserClient/Source/MultiUserClient/Private/Replication/Authority/AuthorityChangeTracker.h b/Engine/Plugins/Developer/Concert/ConcertApp/MultiUserClient/Source/MultiUserClient/Private/Replication/Authority/AuthorityChangeTracker.h
--- a/Engine/Plugins/Developer/Concert/ConcertApp/MultiUserClient/Source/MultiUserClient/Private/Replication/Authority/AuthorityChangeTracker.h
+++ b/Engine/Plugins/Developer/Concert/ConcertApp/MultiUserClient/Source/MultiUserClient/Private/Replication/Authority/AuthorityChangeTracker.h
@@ -26,10 +26,10 @@ namespace UE::MultiUserClient
 		FAuthorityChangeTracker(const FGuid& InClientId, const IClientAuthoritySynchronizer& InAuthoritySynchronizer, FGlobalAuthorityCache& InAuthorityCache);
 		~FAuthorityChangeTracker();
 		
-		/**
-		 * Marks that the authority should be changed to bNewAuthorityState.
-		 */
+		/** Marks that the authority should be changed to bNewAuthorityState. */
 		void SetAuthorityIfAllowed(TConstArrayView<FSoftObjectPath> ObjectPaths, bool bNewAuthorityState);
+		/** Reverts any local, unsubmitted changes made for the ObjectPaths. */
+		void ClearAuthorityChange(TConstArrayView<FSoftObjectPath> ObjectPaths);
 
 		/** Diffs NewAuthorityStates to the current authority states and removes entries. */
 		void RefreshChanges();
@@ -48,9 +48,9 @@ namespace UE::MultiUserClient
 		/** Builds a change request from the local changes. */
 		ConcertSyncClient::Replication::FAuthorityChangeRequest BuildChangeRequest(const FGuid& StreamId) const;
 
-		/** Called when NewAuthorityStates is updated. */
+		/** Called entries are added to NewAuthorityStates. NOT called when objects are removed. */
 		DECLARE_MULTICAST_DELEGATE(FOnAuthorityChangeMade);
-		FOnAuthorityChangeMade& OnAuthorityChangeMade() { return FOnAuthorityChangeMadeDelegate; }
+		FOnAuthorityChangeMade& OnAddedOwnedObjects() { return OnAddedOwnedObjectsDelegate; }
 		
 	private:
 
@@ -66,7 +66,7 @@ namespace UE::MultiUserClient
 		TMap<FSoftObjectPath, bool> NewAuthorityStates;
 
 		/** Called when NewAuthorityStates is updated. */
-		FOnAuthorityChangeMade FOnAuthorityChangeMadeDelegate;
+		FOnAuthorityChangeMade OnAddedOwnedObjectsDelegate;
 		
 		void OnClientChanged(const FGuid& Guid) { RefreshChanges(); }
 	};
diff --git a/Engine/Plugins/Developer/Concert/ConcertApp/MultiUserClient/Source/MultiUserClient/Private/Replication/Client/ReplicationClient.cpp b/Engine/Plugins/Developer/Concert/ConcertApp/MultiUserClient/Source/MultiUserClient/Private/Replication/Client/ReplicationClient.cpp
--- a/Engine/Plugins/Developer/Concert/ConcertApp/MultiUserClient/Source/MultiUserClient/Private/Replication/Client/ReplicationClient.cpp
+++ b/Engine/Plugins/Developer/Concert/ConcertApp/MultiUserClient/Source/MultiUserClient/Private/Replication/Client/ReplicationClient.cpp
@@ -6,6 +6,7 @@
 #include "Replication/Editor/Model/IEditableObjectToPropertiesModel.h"
 #include "Replication/ReplicationWidgetFactories.h"
 #include "Replication/Stream/StreamChangeTracker.h"
+#include "Replication/Submission/Data/AuthoritySubmission.h"
 
 namespace UE::MultiUserClient
 {
@@ -41,6 +42,8 @@ namespace UE::MultiUserClient
 		{
 			OnModelExternallyChangedDelegate.Broadcast();
 		});
+
+		SubmissionWorkflow->OnAuthorityRequestCompleted().AddRaw(this, &FReplicationClient::OnAuthoritySubmissionCompleted);
 	}
 
 	void FReplicationClient::OnObjectsChanged(
@@ -69,4 +72,18 @@ namespace UE::MultiUserClient
 		// Refresh because authority changes may no longer be valid after modifying the stream
 		LocalAuthorityDiffer.RefreshChanges();
 	}
+
+	void FReplicationClient::OnAuthoritySubmissionCompleted(const FSubmitAuthorityChangesRequest& Request, const FSubmitAuthorityChangesResponse& Response)
+	{
+		if (!Response.Response)
+		{
+			return;
+		}
+
+		// Use case: You and another client submit at the same time. You lose. Revert your local changes so the checkboxes accurately reflect the authority state.
+		for (const TPair<FSoftObjectPath, FConcertStreamArray>& RejectedObjectPair : Response.Response.GetValue().RejectedObjects)
+		{
+			LocalAuthorityDiffer.ClearAuthorityChange({ RejectedObjectPair.Key });
+		}
+	}
 }
diff --git a/Engine/Plugins/Developer/Concert/ConcertApp/MultiUserClient/Source/MultiUserClient/Private/Replication/Client/ReplicationClient.h b/Engine/Plugins/Developer/Concert/ConcertApp/MultiUserClient/Source/MultiUserClient/Private/Replication/Client/ReplicationClient.h
--- a/Engine/Plugins/Developer/Concert/ConcertApp/MultiUserClient/Source/MultiUserClient/Private/Replication/Client/ReplicationClient.h
+++ b/Engine/Plugins/Developer/Concert/ConcertApp/MultiUserClient/Source/MultiUserClient/Private/Replication/Client/ReplicationClient.h
@@ -24,7 +24,10 @@ namespace UE::MultiUserClient
 {
 	class ISubmissionWorkflow;
 	
-	/** Holds on to info about a local or remote client. */
+	/**
+	 * Holds on to shared info about a local or remote client.
+	 * This class' responsibility is to initialize all systems that exist for the life time of a client in a session.
+	 */
 	class FReplicationClient : public FNoncopyable
 	{
 	public:
@@ -109,6 +112,9 @@ namespace UE::MultiUserClient
 		// Respond to model changing
 		void OnObjectsChanged(TConstArrayView<UObject*> AddedObjects, TConstArrayView<FSoftObjectPath> RemovedObjects, ConcertClientSharedSlate::EReplicatedObjectChangeReason ReplicatedObjectChangeReason);
 		void OnPropertiesChanged();
+		
+		/** Removes authority if request fails */
+		void OnAuthoritySubmissionCompleted(const FSubmitAuthorityChangesRequest& Request, const FSubmitAuthorityChangesResponse& Response);
 	};
 }
 
diff --git a/Engine/Plugins/Developer/Concert/ConcertApp/MultiUserClient/Source/MultiUserClient/Private/Replication/Submission/AutoSubmissionPolicy.cpp b/Engine/Plugins/Developer/Concert/ConcertApp/MultiUserClient/Source/MultiUserClient/Private/Replication/Submission/AutoSubmissionPolicy.cpp
--- a/Engine/Plugins/Developer/Concert/ConcertApp/MultiUserClient/Source/MultiUserClient/Private/Replication/Submission/AutoSubmissionPolicy.cpp
+++ b/Engine/Plugins/Developer/Concert/ConcertApp/MultiUserClient/Source/MultiUserClient/Private/Replication/Submission/AutoSubmissionPolicy.cpp
@@ -21,7 +21,7 @@ namespace UE::MultiUserClient
 	{
 		StreamEditorModel.OnObjectsChanged().AddRaw(this, &FAutoSubmissionPolicy::OnObjectsChanged);
 		StreamEditorModel.OnPropertiesChanged().AddRaw(this, &FAutoSubmissionPolicy::OnChangesDetected);
-		AuthorityChangeTracker.OnAuthorityChangeMade().AddRaw(this, &FAutoSubmissionPolicy::OnChangesDetected);
+		AuthorityChangeTracker.OnAddedOwnedObjects().AddRaw(this, &FAutoSubmissionPolicy::OnChangesDetected);
 	}
 
 	FAutoSubmissionPolicy::~FAutoSubmissionPolicy()
@@ -29,7 +29,7 @@ namespace UE::MultiUserClient
 		FCoreDelegates::OnEndFrame.RemoveAll(this);
 		StreamEditorModel.OnObjectsChanged().RemoveAll(this);
 		StreamEditorModel.OnPropertiesChanged().RemoveAll(this);
-		AuthorityChangeTracker.OnAuthorityChangeMade().RemoveAll(this);
+		AuthorityChangeTracker.OnAddedOwnedObjects().RemoveAll(this);
 	}
 
 	void FAutoSubmissionPolicy::OnChangesDetected()
diff --git a/Engine/Plugins/Developer/Concert/ConcertApp/MultiUserClient/Source/MultiUserClient/Private/Replication/Submission/Notification/AccumulatedSubmissionErrors.h b/Engine/Plugins/Developer/Concert/ConcertApp/MultiUserClient/Source/MultiUserClient/Private/Replication/Submission/Notification/AccumulatedSubmissionErrors.h
--- a/Engine/Plugins/Developer/Concert/ConcertApp/MultiUserClient/Source/MultiUserClient/Private/Replication/Submission/Notification/AccumulatedSubmissionErrors.h
+++ b/Engine/Plugins/Developer/Concert/ConcertApp/MultiUserClient/Source/MultiUserClient/Private/Replication/Submission/Notification/AccumulatedSubmissionErrors.h
@@ -12,10 +12,16 @@ namespace UE::MultiUserClient
 	{
 		int32 NumTimeouts = 0;
 
-		/** Changing the these objects would cause authority conflicts. */
-		TSet<FSoftObjectPath> AuthorityConflicts;
-		/** This should be empty. If not, that means our client code made a bad request. */
-		TSet<FSoftObjectPath> SemanticErrors;
+		/**
+		 * Changing the these objects would cause authority conflicts.
+		 * Maps to the number of times it was rejected (across multiple submits).
+		 */
+		TMap<FSoftObjectPath, int32> AuthorityConflicts;
+		/**
+		 * This should be empty. If not, that means our client code made a bad request.
+		 * Maps to the number of times it was rejected (across multiple submits).
+		 */
+		TMap<FSoftObjectPath, int32> SemanticErrors;
 		/** Should be false. If not, that means our client code made a bad request. */
 		bool bFailedStreamCreation = false;
 	};
@@ -24,7 +30,10 @@ namespace UE::MultiUserClient
 	{
 		int32 NumTimeouts = 0;
 
-		/** These objects had conflicts */
-		TSet<FSoftObjectPath> Rejected;
+		/**
+		 * These objects had conflicts.
+		 * Maps to the number of times it was rejected (across multiple submits).
+		 */
+		TMap<FSoftObjectPath, int32> Rejected;
 	};
 }
\ No newline at end of file
diff --git a/Engine/Plugins/Developer/Concert/ConcertApp/MultiUserClient/Source/MultiUserClient/Private/Replication/Submission/Notification/SAuthorityRejectedNotification.cpp b/Engine/Plugins/Developer/Concert/ConcertApp/MultiUserClient/Source/MultiUserClient/Private/Replication/Submission/Notification/SAuthorityRejectedNotification.cpp
--- a/Engine/Plugins/Developer/Concert/ConcertApp/MultiUserClient/Source/MultiUserClient/Private/Replication/Submission/Notification/SAuthorityRejectedNotification.cpp
+++ b/Engine/Plugins/Developer/Concert/ConcertApp/MultiUserClient/Source/MultiUserClient/Private/Replication/Submission/Notification/SAuthorityRejectedNotification.cpp
@@ -4,6 +4,7 @@
 
 #include "AccumulatedSubmissionErrors.h"
 
+#include "Algo/Accumulate.h"
 #include "Styling/AppStyle.h"
 #include "Widgets/Input/SButton.h"
 #include "Widgets/SBoxPanel.h"
@@ -55,11 +56,18 @@ namespace UE::MultiUserClient
 
 	void SAuthorityRejectedNotification::AddRejectionErrorWidget(TSharedRef<SVerticalBox> Result, const FAccumulatedAuthorityErrors& Errors)
 	{
+		const int32 NumConflicts = Algo::TransformAccumulate(Errors.Rejected, [](const TPair<FSoftObjectPath, int32>& Pair){ return Pair.Value; }, 0);
+		const FText Text = FText::Format(
+			LOCTEXT("RejectionsFmt", "{0} {0}|plural(one=Rejection,other=Rejections) for {1} {1}|plural(one=Object,other=Objects)"),
+			NumConflicts,
+			Errors.Rejected.Num()
+			);
+		
 		Result->AddSlot()
 			.AutoHeight()
 			[
 				SNew(STextBlock)
-				.Text(FText::Format(LOCTEXT("RejectionsFmt", "{0} {0}|plural(one=Rejection,other=Rejections)"), Errors.Rejected.Num()))
+				.Text(Text)
 				.Font(FAppStyle::Get().GetFontStyle(TEXT("NotificationList.FontLight")))
 				.TextStyle(&FAppStyle::GetWidgetStyle<FTextBlockStyle>("NotificationList.WidgetText"))
 			];
diff --git a/Engine/Plugins/Developer/Concert/ConcertApp/MultiUserClient/Source/MultiUserClient/Private/Replication/Submission/Notification/SStreamRejectedNotification.cpp b/Engine/Plugins/Developer/Concert/ConcertApp/MultiUserClient/Source/MultiUserClient/Private/Replication/Submission/Notification/SStreamRejectedNotification.cpp
--- a/Engine/Plugins/Developer/Concert/ConcertApp/MultiUserClient/Source/MultiUserClient/Private/Replication/Submission/Notification/SStreamRejectedNotification.cpp
+++ b/Engine/Plugins/Developer/Concert/ConcertApp/MultiUserClient/Source/MultiUserClient/Private/Replication/Submission/Notification/SStreamRejectedNotification.cpp
@@ -4,6 +4,7 @@
 
 #include "AccumulatedSubmissionErrors.h"
 
+#include "Algo/Accumulate.h"
 #include "Styling/AppStyle.h"
 #include "Widgets/Input/SButton.h"
 #include "Widgets/SBoxPanel.h"
@@ -63,11 +64,18 @@ namespace UE::MultiUserClient
 
 	void SStreamRejectedNotification::AddAuthorityConflictErrorWidget(TSharedRef<SVerticalBox> Result, const FAccumulatedStreamErrors& Errors)
 	{
+		const int32 NumConflicts = Algo::TransformAccumulate(Errors.AuthorityConflicts, [](const TPair<FSoftObjectPath, int32>& Pair){ return Pair.Value; }, 0);
+		const FText Text = FText::Format(
+			LOCTEXT("ConflictsFmt", "{0} {0}|plural(one=Conflict,other=Conflicts) for {1} {1}|plural(one=Object,other=Objects)"),
+			NumConflicts,
+			Errors.AuthorityConflicts.Num()
+			);
+		
 		Result->AddSlot()
 			.AutoHeight()
 			[
 				SNew(STextBlock)
-				.Text(FText::Format(LOCTEXT("ConflictsFmt", "{0} {0}|plural(one=Conflict,other=Conflicts)"), Errors.AuthorityConflicts.Num()))
+				.Text(Text)
 				.Font(FAppStyle::Get().GetFontStyle(TEXT("NotificationList.FontLight")))
 				.TextStyle(&FAppStyle::GetWidgetStyle<FTextBlockStyle>("NotificationList.WidgetText"))
 			];
@@ -75,11 +83,18 @@ namespace UE::MultiUserClient
 	
 	void SStreamRejectedNotification::AddSemanticErrorWidget(TSharedRef<SVerticalBox> Result, const FAccumulatedStreamErrors& Errors)
 	{
+		const int32 NumConflicts = Algo::TransformAccumulate(Errors.SemanticErrors, [](const TPair<FSoftObjectPath, int32>& Pair){ return Pair.Value; }, 0);
+		const FText Text = FText::Format(
+			LOCTEXT("SemanticErrorsFmt", "{0} {0}|plural(one=Semantic Error,other=Semantic Errors) for {1} {1}|plural(one=Object,other=Objects)"),
+			NumConflicts,
+			Errors.AuthorityConflicts.Num()
+			);
+		
 		Result->AddSlot()
 			.AutoHeight()
 			[
 				SNew(STextBlock)
-				.Text(FText::Format(LOCTEXT("SemanticFmt", "{0} {0}|plural(one=Semantic Error,other=Semantic Errors)"), Errors.SemanticErrors.Num()))
+				.Text(Text)
 				.Font(FAppStyle::Get().GetFontStyle(TEXT("NotificationList.FontLight")))
 				.TextStyle(&FAppStyle::GetWidgetStyle<FTextBlockStyle>("NotificationList.WidgetText"))
 			];
diff --git a/Engine/Plugins/Developer/Concert/ConcertApp/MultiUserClient/Source/MultiUserClient/Private/Replication/Submission/Notification/SubmissionNotifier.cpp b/Engine/Plugins/Developer/Concert/ConcertApp/MultiUserClient/Source/MultiUserClient/Private/Replication/Submission/Notification/SubmissionNotifier.cpp
--- a/Engine/Plugins/Developer/Concert/ConcertApp/MultiUserClient/Source/MultiUserClient/Private/Replication/Submission/Notification/SubmissionNotifier.cpp
+++ b/Engine/Plugins/Developer/Concert/ConcertApp/MultiUserClient/Source/MultiUserClient/Private/Replication/Submission/Notification/SubmissionNotifier.cpp
@@ -74,14 +74,14 @@ namespace UE::MultiUserClient
 		switch(CompletedOp.ErrorCode)
 		{
 		case EStreamSubmissionErrorCode::Success:
-			Algo::Transform(
-				CompletedOp.SubmissionInfo->Response.AuthorityConflicts, StreamErrors.AuthorityConflicts,
-				[](const TPair<FObjectInStreamID, FReplicatedObjectId>& Pair) { return Pair.Key.Object; }
-				);
-			Algo::Transform(
-				CompletedOp.SubmissionInfo->Response.ObjectsToPutSemanticErrors, StreamErrors.SemanticErrors,
-				[](const TPair<FObjectInStreamID, EConcertPutObjectErrorCode>& Pair) { return Pair.Key.Object; }
-				);
+			for (const TPair<FObjectInStreamID, FReplicatedObjectId>& Pair : CompletedOp.SubmissionInfo->Response.AuthorityConflicts)
+			{
+				StreamErrors.AuthorityConflicts.FindOrAdd(Pair.Key.Object, 0) += 1;
+			}
+			for (const TPair<FObjectInStreamID, EConcertPutObjectErrorCode>& Pair : CompletedOp.SubmissionInfo->Response.ObjectsToPutSemanticErrors)
+			{
+				StreamErrors.SemanticErrors.FindOrAdd(Pair.Key.Object, 0) += 1;
+			}
 			StreamErrors.bFailedStreamCreation |= !CompletedOp.SubmissionInfo->Response.FailedStreamCreation.IsEmpty();
 			break;
 			
@@ -100,10 +100,10 @@ namespace UE::MultiUserClient
 		switch (ResponseOp.ErrorCode)
 		{
 		case EAuthoritySubmissionResponseErrorCode::Success:
-			Algo::Transform(
-				ResponseOp.Response->RejectedObjects, AuthorityErrors.Rejected,
-				[](const TPair<FSoftObjectPath, FConcertStreamArray>& Pair) { return Pair.Key; }
-				);
+			for (const TPair<FSoftObjectPath, FConcertStreamArray>& Pair : ResponseOp.Response->RejectedObjects)
+			{
+				AuthorityErrors.Rejected.FindOrAdd(Pair.Key, 0) += 1;
+			}
 			break;
 			
 		case EAuthoritySubmissionResponseErrorCode::Timeout:
