Replication view has two modes controlled with console variable:
- 3 sections (old) > Outliner, subobject editor, properties
- 2 sections (new) > Objects, properties

diff --git a/Engine/Plugins/Developer/Concert/ConcertApp/MultiUserClient/Source/MultiUserClient/Private/Widgets/ActiveSession/Replication/Client/Columns/SOwnerClientList.cpp b/Engine/Plugins/Developer/Concert/ConcertApp/MultiUserClient/Source/MultiUserClient/Private/Widgets/ActiveSession/Replication/Client/Columns/SOwnerClientList.cpp
new file mode 100644
--- /dev/null
+++ b/Engine/Plugins/Developer/Concert/ConcertApp/MultiUserClient/Source/MultiUserClient/Private/Widgets/ActiveSession/Replication/Client/Columns/SOwnerClientList.cpp
@@ -0,0 +1,59 @@
+// Copyright Epic Games, Inc. All Rights Reserved.
+
+#include "SOwnerClientList.h"
+
+#include "Replication/Util/GlobalAuthorityCache.h"
+#include "Widgets/ClientName/SHorizontalClientList.h"
+
+#include "Widgets/Layout/SBox.h"
+#include "Widgets/Text/STextBlock.h"
+
+#define LOCTEXT_NAMESPACE "SOwnerClientList"
+
+namespace UE::MultiUserClient
+{
+	void SOwnerClientList::Construct(const FArguments& InArgs, TSharedRef<IConcertClient> InClient, FGlobalAuthorityCache& InAuthorityCache)
+	{
+		AuthorityCache = &InAuthorityCache;
+		GetClientListDelegate = InArgs._GetClientList;
+		check(GetClientListDelegate.IsBound());
+		
+		AuthorityCache->OnCacheChanged().AddSP(this, &SOwnerClientList::OnClientChanged);
+		
+		ChildSlot
+		[
+			SAssignNew(ClientList, ConcertClientSharedSlate::SHorizontalClientList, MoveTemp(InClient))
+			.HighlightText(InArgs._HighlightText)
+			.EmptyListSlot()
+			[
+				SNew(SBox)
+				.HAlign(HAlign_Left)
+				.Padding(5.f, 0.f, 0.f, 0.f)
+				[
+					SNew(STextBlock)
+					.Text(LOCTEXT("NotApplicable", "-"))
+				]
+			]
+		];
+
+		RefreshList();
+	}
+
+	SOwnerClientList::~SOwnerClientList()
+	{
+		AuthorityCache->OnCacheChanged().RemoveAll(this);
+	}
+
+	void SOwnerClientList::OnClientChanged(const FGuid& ChangedClientId)
+	{
+		RefreshList();
+	}
+	
+	void SOwnerClientList::RefreshList()
+	{
+		const TArray<FGuid> Clients = GetClientListDelegate.Execute(*AuthorityCache);
+		ClientList->RefreshList(Clients);
+	}
+}
+
+#undef LOCTEXT_NAMESPACE
\ No newline at end of file
diff --git a/Engine/Plugins/Developer/Concert/ConcertApp/MultiUserClient/Source/MultiUserClient/Private/Widgets/ActiveSession/Replication/Client/Columns/SOwnerClientList.h b/Engine/Plugins/Developer/Concert/ConcertApp/MultiUserClient/Source/MultiUserClient/Private/Widgets/ActiveSession/Replication/Client/Columns/SOwnerClientList.h
new file mode 100644
--- /dev/null
+++ b/Engine/Plugins/Developer/Concert/ConcertApp/MultiUserClient/Source/MultiUserClient/Private/Widgets/ActiveSession/Replication/Client/Columns/SOwnerClientList.h
@@ -0,0 +1,56 @@
+// Copyright Epic Games, Inc. All Rights Reserved.
+
+#pragma once
+
+#include "Widgets/DeclarativeSyntaxSupport.h"
+#include "Widgets/SCompoundWidget.h"
+
+namespace UE::ConcertClientSharedSlate
+{
+	class SHorizontalClientList;
+}
+
+class IConcertClient;
+
+namespace UE::MultiUserClient
+{
+	class FGlobalAuthorityCache;
+	
+	/**
+	 * Displays a list of authoritive clients in the SReplicationClientView columns for top-level, subobject and properties;
+	 * the passed in FGetClientList determines which is displayed
+	 */
+	class SOwnerClientList : public SCompoundWidget
+	{
+	public:
+
+		DECLARE_DELEGATE_RetVal_OneParam(TArray<FGuid>, FGetClientList, const FGlobalAuthorityCache&);
+		
+		SLATE_BEGIN_ARGS(SOwnerClientList)
+		{}
+			/** Obtains the client list to display */
+			SLATE_EVENT(FGetClientList, GetClientList)
+			/** Used for highlighting in the text */
+			SLATE_ATTRIBUTE(FText, HighlightText)
+		SLATE_END_ARGS()
+
+		void Construct(const FArguments& InArgs, TSharedRef<IConcertClient> InClient, FGlobalAuthorityCache& InAuthorityCache);
+		virtual ~SOwnerClientList();
+
+	private:
+		
+		/** Tells us when the authority change */
+		FGlobalAuthorityCache* AuthorityCache;
+
+		/** Displays the clients */
+		TSharedPtr<ConcertClientSharedSlate::SHorizontalClientList> ClientList;
+
+		/** Used to get the client list from FGlobalAuthorityCache */
+		FGetClientList GetClientListDelegate;
+
+		void OnClientChanged(const FGuid& ChangedClientId);
+		void RefreshList();
+	};
+}
+
+
diff --git a/Engine/Plugins/Developer/Concert/ConcertApp/MultiUserClient/Source/MultiUserClient/Private/Widgets/ActiveSession/Replication/Client/Columns/SWarningIcon.cpp b/Engine/Plugins/Developer/Concert/ConcertApp/MultiUserClient/Source/MultiUserClient/Private/Widgets/ActiveSession/Replication/Client/Columns/SWarningIcon.cpp
new file mode 100644
--- /dev/null
+++ b/Engine/Plugins/Developer/Concert/ConcertApp/MultiUserClient/Source/MultiUserClient/Private/Widgets/ActiveSession/Replication/Client/Columns/SWarningIcon.cpp
@@ -0,0 +1,25 @@
+// Copyright Epic Games, Inc. All Rights Reserved.
+
+#include "SWarningIcon.h"
+
+#include "Styling/AppStyle.h"
+#include "Widgets/Images/SImage.h"
+#include "Widgets/Layout/SScaleBox.h"
+
+namespace UE::MultiUserClient
+{
+	void SWarningIcon::Construct(const FArguments& InArgs)
+	{
+		ChildSlot
+		[
+			SNew(SScaleBox)
+			.Stretch(EStretch::ScaleToFit)
+			[
+				SNew(SImage)
+				.Image(FAppStyle::Get().GetBrush("Icons.WarningWithColor"))
+				.DesiredSizeOverride(FVector2D(16,16))
+			]
+		];
+	}
+}
+
diff --git a/Engine/Plugins/Developer/Concert/ConcertApp/MultiUserClient/Source/MultiUserClient/Private/Widgets/ActiveSession/Replication/Client/Columns/SWarningIcon.h b/Engine/Plugins/Developer/Concert/ConcertApp/MultiUserClient/Source/MultiUserClient/Private/Widgets/ActiveSession/Replication/Client/Columns/SWarningIcon.h
new file mode 100644
--- /dev/null
+++ b/Engine/Plugins/Developer/Concert/ConcertApp/MultiUserClient/Source/MultiUserClient/Private/Widgets/ActiveSession/Replication/Client/Columns/SWarningIcon.h
@@ -0,0 +1,23 @@
+// Copyright Epic Games, Inc. All Rights Reserved.
+
+#pragma once
+
+#include "Widgets/DeclarativeSyntaxSupport.h"
+#include "Widgets/SCompoundWidget.h"
+
+namespace UE::MultiUserClient
+{
+	/** Displays a warning triangle */
+	class SWarningIcon : public SCompoundWidget
+	{
+	public:
+
+		SLATE_BEGIN_ARGS(SWarningIcon)
+		{}
+		SLATE_END_ARGS()
+
+		void Construct(const FArguments& InArgs);
+	};
+}
+
+
diff --git a/Engine/Plugins/Developer/Concert/ConcertApp/MultiUserClient/Source/MultiUserClient/Private/Widgets/ActiveSession/Replication/Client/Columns/SingleClientColumns.h b/Engine/Plugins/Developer/Concert/ConcertApp/MultiUserClient/Source/MultiUserClient/Private/Widgets/ActiveSession/Replication/Client/Columns/SingleClientColumns.h
new file mode 100644
--- /dev/null
+++ b/Engine/Plugins/Developer/Concert/ConcertApp/MultiUserClient/Source/MultiUserClient/Private/Widgets/ActiveSession/Replication/Client/Columns/SingleClientColumns.h
@@ -0,0 +1,160 @@
+// Copyright Epic Games, Inc. All Rights Reserved.
+
+#pragma once
+
+#include "Replication/Editor/View/PredefinedReplicationColumns.h"
+#include "Widgets/DeclarativeSyntaxSupport.h"
+#include "Widgets/SCompoundWidget.h"
+
+class IConcertClient;
+
+namespace UE::ConcertClientSharedSlate
+{
+	class IReplicationStreamViewer;
+	class IObjectToPropertiesModel;
+}
+
+namespace UE::MultiUserClient
+{
+	class FAuthorityChangeTracker;
+	class FGlobalAuthorityCache;
+	class IClientAuthoritySynchronizer;
+	class ISubmissionWorkflow;
+}
+
+namespace UE::MultiUserClient::SingleClientColumns
+{
+	/** @see UE::ConcertClientSharedSlate::ReplicationColumns::Subobject::ETopLevelColumnOrder */
+	enum class ETopLevelObjectColumnOrder
+	{
+		ToggleAuthority = 0,
+		Owner = 40
+	};
+	/** @see UE::ConcertClientSharedSlate::ReplicationColumns::Subobject::ESubobjectColumnOrder */
+	enum class ESubobjectColumnOrder
+	{
+		ToggleAuthority = 0,
+		ConflictWarning = 5,
+		Owner = 40
+	};
+	/** @see UE::ConcertClientSharedSlate::ReplicationColumns::Property::EReplicationPropertyColumnOrder */
+	enum class EPropertyColumnOrder
+	{
+		ConflictWarning = 5,
+		Owner = 50
+	};
+
+
+	/********** Toggle authority **********/
+	extern const FName ToggleTopLevelAuthorityColumnId;
+	extern const FName ToggleSubobjectAuthorityColumnId;
+
+	/**
+	 * Checkbox placed in the top-level view.
+	 * It gives / removes authority for the top-level object and all of its subobjects.
+	 *
+	 * @param ClientStreamModel Used to discover the client's registered subobjects
+	 * @param ChangeTracker Used to determine the checkbox state and change authority.
+	 * @param SubmissionWorkflow Used to determine whether changing authority is at all enabled
+	 * @return Column that can be placed in the table
+	 */
+	ConcertClientSharedSlate::ReplicationColumns::FReplicationTopLevelObjectColumn ToggleTopLevelAuthority(
+		ConcertClientSharedSlate::IObjectToPropertiesModel& ClientStreamModel,
+		FAuthorityChangeTracker& ChangeTracker,
+		ISubmissionWorkflow& SubmissionWorkflow
+		);
+
+	/**
+	 * Checkbox placed in the subobject view
+	 * It gives / removes authority for the object it is placed next to.
+	 * 
+	 * @param AuthoritySynchronizer Determines whether the client has authority or not.
+	 * @param ChangeTracker Used to determine the checkbox state and change authority.
+	 * @param AuthoritySynchronizer Used to determine whether changing authority is at all enabled
+	 * @return Column that can be placed in the table
+	 */
+	ConcertClientSharedSlate::ReplicationColumns::FReplicationSubobjectObjectColumn ToggleSubobjectAuthority(
+		FAuthorityChangeTracker& ChangeTracker,
+		ISubmissionWorkflow& SubmissionWorkflow
+		);
+
+	
+	/********** Owner **********/
+	extern const FName OwnerOfSubobjectColumnId;
+	extern const FName OwnerOfPropertyColumnId;
+	
+	/**
+	 * Displays the owner of a top-level object.
+	 * @param InClient The local Concert client used to look up other client display info
+	 * @param InAuthorityCache Used to determine which client owns the property
+	 * @param InObjectModel Used to get subobjects for a top-level object
+	 * @return Column that can be placed in the table
+	 */
+	ConcertClientSharedSlate::ReplicationColumns::FReplicationTopLevelObjectColumn OwnerOfTopLevelObject(
+		const TSharedRef<IConcertClient>& InClient,
+		FGlobalAuthorityCache& InAuthorityCache,
+		ConcertClientSharedSlate::IObjectToPropertiesModel& InObjectModel
+		);
+
+	/**
+	 * Displays the owner of a subobject.
+	 * @param InClient The local Concert client used to look up other client display info
+	 * @param InAuthorityCache Used to determine which client owns the property
+	 * @return Column that can be placed in the table
+	 */
+	ConcertClientSharedSlate::ReplicationColumns::FReplicationSubobjectObjectColumn OwnerOfSubobject(
+		const TSharedRef<IConcertClient>& InClient,
+		FGlobalAuthorityCache& InAuthorityCache
+		);
+	
+	/**
+	 * Displays the owner of a property.
+	 * @param InClient The local Concert client used to look up other client display info
+	 * @param InAuthorityCache Used to determine which client owns the property
+	 * @param InViewerAttribute Used to determine which objects the property box is displaying
+	 * @return Column that can be placed in the table
+	 */
+	ConcertClientSharedSlate::ReplicationColumns::FReplicationPropertyColumn OwnerOfProperty(
+		const TSharedRef<IConcertClient>& InClient,
+		FGlobalAuthorityCache& InAuthorityCache,
+		const TAttribute<const ConcertClientSharedSlate::IReplicationStreamViewer*>& InViewerAttribute
+		);
+
+	
+	/********** Conflict warning **********/
+	extern const FName ConflictWarningTopLevelObjectColumnId;
+	extern const FName ConflictWarningSubobjectColumnId;
+	extern const FName ConflictWarningPropertyColumnId;
+	
+	/**
+	 * Displays a warning symbol next to the checkbox if checking the checkbox would cause an authority conflict when submitted.
+	 * 
+	 * @param InClient The local Concert client used to look up other client display info
+	 * @param InAuthorityCache Used to determine which client owns the object
+	 * @param ClientId ID of the client for which the icon is being created
+	 * 
+	 * @return Column that can be placed in the table
+	 */
+	ConcertClientSharedSlate::ReplicationColumns::FReplicationSubobjectObjectColumn ConflictWarningForSubobject(
+		TSharedRef<IConcertClient> InClient,
+		FGlobalAuthorityCache& InAuthorityCache,
+		const FGuid& ClientId
+		);
+
+	/**
+	 * Displays a warning symbol next to the checkbox if checking the checkbox would cause an authority conflict when submitted.
+	 * 
+	 * @param InClient The local Concert client used to look up other client display info
+	 * @param InViewer Used to determine which objects the property box is displaying
+	 * @param InAuthorityCache Used to determine which client owns the property
+	 * @param ClientId ID of the client for which to check whether the property can be added to the stream
+	 * 
+	 * @return Column that can be placed in the table
+	 */
+	ConcertClientSharedSlate::ReplicationColumns::FReplicationPropertyColumn ConflictWarningForProperty(
+		TSharedRef<IConcertClient> InClient,
+		TAttribute<const ConcertClientSharedSlate::IReplicationStreamViewer*> InViewer,
+		FGlobalAuthorityCache& InAuthorityCache,
+		const FGuid& ClientId
+		);
+}
diff --git a/Engine/Plugins/Developer/Concert/ConcertApp/MultiUserClient/Source/MultiUserClient/Private/Widgets/ActiveSession/Replication/Client/Columns/SingleClientColumns_ConflictWarning.cpp b/Engine/Plugins/Developer/Concert/ConcertApp/MultiUserClient/Source/MultiUserClient/Private/Widgets/ActiveSession/Replication/Client/Columns/SingleClientColumns_ConflictWarning.cpp
new file mode 100644
--- /dev/null
+++ b/Engine/Plugins/Developer/Concert/ConcertApp/MultiUserClient/Source/MultiUserClient/Private/Widgets/ActiveSession/Replication/Client/Columns/SingleClientColumns_ConflictWarning.cpp
@@ -0,0 +1,130 @@
+// Copyright Epic Games, Inc. All Rights Reserved.
+
+#include "SingleClientColumns.h"
+
+#include "IConcertClient.h"
+#include "Replication/Editor/Model/ReplicatedObjectData.h"
+#include "Replication/Editor/View/IReplicationStreamViewer.h"
+#include "Replication/Util/GlobalAuthorityCache.h"
+#include "SWarningIcon.h"
+
+#include "Algo/AllOf.h"
+#include "Widgets/ActiveSession/Replication/Client/ClientUtils.h"
+#include "Widgets/ClientName/SClientName.h"
+#include "Widgets/Images/SImage.h"
+#include "Widgets/Views/SHeaderRow.h"
+
+#define LOCTEXT_NAMESPACE "SingleClientColumns.ConflictWarning"
+
+namespace UE::MultiUserClient::SingleClientColumns
+{
+	const FName ConflictWarningTopLevelObjectColumnId = TEXT("ConflictWarningTopLevelObjectColumn");
+	const FName ConflictWarningSubobjectColumnId = TEXT("ConflictWarningSubobjectColumn");
+	const FName ConflictWarningPropertyColumnId = TEXT("ConflictWarningPropertyColumn");
+	
+	ConcertClientSharedSlate::ReplicationColumns::FReplicationSubobjectObjectColumn ConflictWarningForSubobject(
+		TSharedRef<IConcertClient> InClient,
+		FGlobalAuthorityCache& InAuthorityCache,
+		const FGuid& ClientId
+		)
+	{
+		using FColumnType = ConcertClientSharedSlate::ReplicationColumns::FReplicationSubobjectObjectColumn;
+		return FColumnType(
+			FColumnType::FArguments()
+				.GenerateWidgetColumn_Lambda([InClient = MoveTemp(InClient), &InAuthorityCache, ClientId](const FColumnType::FBuildArgs& InArgs) mutable
+				{
+					const auto GetVisibility = [&InAuthorityCache, ClientId, ObjectPath = InArgs.RowData.GetObjectPath()]()
+					{
+						const FGlobalAuthorityCache::ECanTakeAuthority TakeAuthorityResult = InAuthorityCache.CanClientTakeAuthorityAfterSubmission(ObjectPath, ClientId);
+						return TakeAuthorityResult == FGlobalAuthorityCache::ECanTakeAuthority::Conflict
+							? EVisibility::Visible
+							: EVisibility::Collapsed;
+					};
+					const auto GetToolTip = [InClient = MoveTemp(InClient), &InAuthorityCache, ClientId, ObjectPath = InArgs.RowData.GetObjectPath()]()
+					{
+						TSet<FString> ClientNames;
+						const FGlobalAuthorityCache::ECanTakeAuthority TakeAuthorityResult = InAuthorityCache.CanClientTakeAuthorityAfterSubmission(ObjectPath, ClientId,
+						[InClient, &ClientNames](const FGuid& ClientId, const FConcertPropertyChain& ConflictingProperty)
+						{
+							ClientNames.Add(ClientUtils::GetClientDisplayName(*InClient, ClientId));
+							return EBreakBehavior::Continue;
+						});
+						
+						return TakeAuthorityResult != FGlobalAuthorityCache::ECanTakeAuthority::Conflict
+							? FText::GetEmpty()
+							: FText::Format(
+								LOCTEXT("Subobject.WarningFmt", "{0} {1}|plural(one=is,other=are) replicating some of the assigned properties. This change will not be sent to the server."),
+								FText::FromString(FString::Join(ClientNames, TEXT(", "))),
+								ClientNames.Num()
+							);
+					};
+					
+					return SNew(SWarningIcon)
+						.Visibility_Lambda(GetVisibility)
+						.ToolTipText_Lambda(GetToolTip);
+				})
+				.PopulateSearchItems_Lambda([](const ConcertClientSharedSlate::FReplicatedObjectData&, TArray<FString>&){})
+				.ColumnSortOrder(static_cast<int32>(ESubobjectColumnOrder::ConflictWarning)),
+			SHeaderRow::Column(ConflictWarningSubobjectColumnId)
+				.DefaultLabel(FText::GetEmpty())
+				.FixedWidth(20.f)
+			);
+	}
+
+	ConcertClientSharedSlate::ReplicationColumns::FReplicationPropertyColumn ConflictWarningForProperty(
+		TSharedRef<IConcertClient> InClient,
+		TAttribute<const ConcertClientSharedSlate::IReplicationStreamViewer*> Viewer,
+		FGlobalAuthorityCache& InAuthorityCache,
+		const FGuid& ClientId
+		)
+	{
+		using FColumnType = ConcertClientSharedSlate::ReplicationColumns::FReplicationPropertyColumn;
+		return FColumnType(
+			FColumnType::FArguments()
+				.GenerateWidgetColumn_Lambda([InClient = MoveTemp(InClient), Viewer, &InAuthorityCache, ClientId](const FColumnType::FBuildArgs& InArgs) mutable
+				{
+					const auto GetVisibility = [&Viewer, &InAuthorityCache, ClientId, Property = InArgs.RowData.GetProperty()]()
+					{
+						const bool bCanAddProperty = Algo::AllOf(Viewer.Get()->GetObjectsBeingPropertyEdited(),
+							[&InAuthorityCache, ClientId, Property](const FSoftObjectPath& ObjectPath)
+							{
+								return InAuthorityCache.CanClientAddProperty(ObjectPath, ClientId, Property);
+							});
+						return bCanAddProperty
+							? EVisibility::Collapsed
+							: EVisibility::Visible;
+					};
+					const auto GetToolTip = [InClient = MoveTemp(InClient), Viewer, &InAuthorityCache, ClientId, Property = InArgs.RowData.GetProperty()]()
+					{
+						TSet<FString> Names;
+						for (const FSoftObjectPath& ObjectPath: Viewer.Get()->GetObjectsBeingPropertyEdited())
+						{
+							const TOptional<FGuid> Owner = InAuthorityCache.GetClientWithAuthorityOverProperty(ObjectPath, Property);
+							if (Owner || *Owner != ClientId)
+							{
+								Names.Add(ClientUtils::GetClientDisplayName(*InClient, *Owner));
+							}
+						}
+						
+						return Names.IsEmpty()
+							? FText::GetEmpty()
+							: FText::Format(
+								LOCTEXT("Property.WarningFmt", "Replicated by {0}. This property will not be submitted to the server."),
+								FText::FromString(FString::Join(Names, TEXT(", ")))
+								);
+					};
+					
+					return SNew(SWarningIcon)
+						.Visibility_Lambda(GetVisibility)
+						.ToolTipText_Lambda(GetToolTip);
+				})
+				.PopulateSearchItems_Lambda([](const ConcertClientSharedSlate::FReplicatedPropertyData&, TArray<FString>&){})
+				.ColumnSortOrder(static_cast<int32>(EPropertyColumnOrder::ConflictWarning)),
+			SHeaderRow::Column(ConflictWarningPropertyColumnId)
+				.DefaultLabel(FText::GetEmpty())
+				.FixedWidth(20.f)
+			);
+	}
+}
+
+#undef LOCTEXT_NAMESPACE
\ No newline at end of file
diff --git a/Engine/Plugins/Developer/Concert/ConcertApp/MultiUserClient/Source/MultiUserClient/Private/Widgets/ActiveSession/Replication/Client/Columns/SingleClientColumns_Owner.cpp b/Engine/Plugins/Developer/Concert/ConcertApp/MultiUserClient/Source/MultiUserClient/Private/Widgets/ActiveSession/Replication/Client/Columns/SingleClientColumns_Owner.cpp
new file mode 100644
--- /dev/null
+++ b/Engine/Plugins/Developer/Concert/ConcertApp/MultiUserClient/Source/MultiUserClient/Private/Widgets/ActiveSession/Replication/Client/Columns/SingleClientColumns_Owner.cpp
@@ -0,0 +1,175 @@
+// Copyright Epic Games, Inc. All Rights Reserved.
+
+#include "SingleClientColumns.h"
+
+#include "IConcertClient.h"
+#include "SOwnerClientList.h"
+#include "Replication/Editor/View/IReplicationStreamViewer.h"
+#include "Replication/Editor/View/ReplicationColumnsUtils.h"
+#include "Replication/Util/GlobalAuthorityCache.h"
+#include "Widgets/ActiveSession/Replication/Client/ClientUtils.h"
+#include "Widgets/ClientName/SClientName.h"
+
+#define LOCTEXT_NAMESPACE "SingleClientColumns.Owner"
+
+namespace UE::MultiUserClient::SingleClientColumns
+{
+	const FName OwnerOfSubobjectColumnId = TEXT("OwnerOfSubobjectColumn");
+	const FName OwnerOfPropertyColumnId = TEXT("OwnerOfPropertyColumn");
+
+	namespace Private::Shared
+	{
+		constexpr float DefaultWidth = 200.f;
+		
+		static void PopulateSearchTerms(const TSharedRef<IConcertClient>& InClient, const TArray<FGuid>& Clients, TArray<FString>& InOutSearchStrings)
+		{
+			for (const FGuid& ClientId : Clients)
+			{
+				InOutSearchStrings.Add(ClientUtils::GetClientDisplayName(*InClient, ClientId));
+			}
+		}
+	}
+
+	namespace Private::TopLevel
+	{
+		static TArray<FGuid> GetOwnersOfTopLevelObject(
+			const FGlobalAuthorityCache& InAuthorityCache,
+			ConcertClientSharedSlate::IObjectToPropertiesModel& InObjectModel,
+			const FSoftObjectPath& InTopLevelObject
+			)
+		{
+			TArray<FGuid> Owners = InAuthorityCache.GetClientsWithAuthorityOverObject(InTopLevelObject);
+			InObjectModel.ForEachSubobject(InTopLevelObject, [&InAuthorityCache, &Owners](const FSoftObjectPath& Subobject)
+			{
+				for (const FGuid& Client : InAuthorityCache.GetClientsWithAuthorityOverObject(Subobject))
+				{
+					Owners.AddUnique(Client);
+				}
+				return EBreakBehavior::Continue;
+			});
+			return Owners;
+		}
+	}
+	
+	ConcertClientSharedSlate::ReplicationColumns::FReplicationTopLevelObjectColumn OwnerOfTopLevelObject(
+		const TSharedRef<IConcertClient>& InClient,
+		FGlobalAuthorityCache& InAuthorityCache,
+		ConcertClientSharedSlate::IObjectToPropertiesModel& InObjectModel
+		)
+	{
+		using FColumnType = ConcertClientSharedSlate::ReplicationColumns::FReplicationSubobjectObjectColumn;
+		return FColumnType(
+			FColumnType::FArguments()
+				.GenerateWidgetColumn_Lambda([InClient, &InObjectModel, &InAuthorityCache](const FColumnType::FBuildArgs& InArgs)
+				{
+					return SNew(SOwnerClientList, InClient, InAuthorityCache)
+						.GetClientList_Lambda([&InAuthorityCache, &InObjectModel, ObjectPath = InArgs.RowData.GetObjectPath()](const FGlobalAuthorityCache&)
+						{
+							return Private::TopLevel::GetOwnersOfTopLevelObject(InAuthorityCache, InObjectModel, ObjectPath);
+						})
+						.HighlightText_Lambda([HighlightText = InArgs.HighlightText](){ return *HighlightText.Get(); });
+				})
+				.PopulateSearchItems_Lambda([InClient, &InAuthorityCache, &InObjectModel](const ConcertClientSharedSlate::FReplicatedObjectData& InArgs, TArray<FString>& InOutSearchStrings)
+				{
+					Private::Shared::PopulateSearchTerms(
+						InClient,
+						Private::TopLevel::GetOwnersOfTopLevelObject(InAuthorityCache, InObjectModel, InArgs.GetObjectPath()),
+						InOutSearchStrings
+						);
+				})
+				.ColumnSortOrder(static_cast<int32>(ETopLevelObjectColumnOrder::Owner)),
+			SHeaderRow::Column(OwnerOfSubobjectColumnId)
+				.DefaultLabel(LOCTEXT("Subobject.Owner", "Owner"))
+				.FillSized(Private::Shared::DefaultWidth)
+			);
+	}
+	
+	ConcertClientSharedSlate::ReplicationColumns::FReplicationSubobjectObjectColumn OwnerOfSubobject(
+		const TSharedRef<IConcertClient>& InClient,
+		FGlobalAuthorityCache& InAuthorityCache
+		)
+	{
+		using FColumnType = ConcertClientSharedSlate::ReplicationColumns::FReplicationSubobjectObjectColumn;
+		return FColumnType(
+			FColumnType::FArguments()
+				.GenerateWidgetColumn_Lambda([InClient, &InAuthorityCache](const FColumnType::FBuildArgs& InArgs)
+				{
+					return SNew(SOwnerClientList, InClient, InAuthorityCache)
+						.GetClientList_Lambda([&InAuthorityCache, ObjectPath = InArgs.RowData.GetObjectPath()](const FGlobalAuthorityCache&)
+						{
+							return InAuthorityCache.GetClientsWithAuthorityOverObject(ObjectPath);
+						})
+						.HighlightText_Lambda([HighlightText = InArgs.HighlightText](){ return *HighlightText.Get(); });
+				})
+				.PopulateSearchItems_Lambda([InClient, &InAuthorityCache](const ConcertClientSharedSlate::FReplicatedObjectData& InArgs, TArray<FString>& InOutSearchStrings)
+				{
+					Private::Shared::PopulateSearchTerms(
+						InClient,
+						InAuthorityCache.GetClientsWithAuthorityOverObject(InArgs.GetObjectPath()),
+						InOutSearchStrings
+						);
+				})
+				.ColumnSortOrder(static_cast<int32>(ESubobjectColumnOrder::Owner)),
+			SHeaderRow::Column(OwnerOfSubobjectColumnId)
+				.DefaultLabel(LOCTEXT("Subobject.Owner", "Owner"))
+				.FillSized(Private::Shared::DefaultWidth)
+			);
+	}
+	
+	namespace Private::Property
+	{
+		static TArray<FGuid> GetPropertyOwners(
+			const FGlobalAuthorityCache& InAuthorityCache,
+			const ConcertClientSharedSlate::IReplicationStreamViewer& InViewer,
+			const FConcertPropertyChain& Property
+			)
+		{
+			TArray<FGuid> Result;
+			const TArray<FSoftObjectPath> SelectedObjects = InViewer.GetObjectsBeingPropertyEdited();
+			for (const FSoftObjectPath& SelectedObject : SelectedObjects)
+			{
+				const TOptional<FGuid> Owner = InAuthorityCache.GetClientWithAuthorityOverProperty(SelectedObject, Property);
+				if (Owner)
+				{
+					Result.AddUnique(*Owner);
+				}
+			}
+			return Result;
+		}
+	}
+
+	ConcertClientSharedSlate::ReplicationColumns::FReplicationPropertyColumn OwnerOfProperty(
+		const TSharedRef<IConcertClient>& InClient,
+		FGlobalAuthorityCache& InAuthorityCache,
+		const TAttribute<const ConcertClientSharedSlate::IReplicationStreamViewer*>& InViewer
+		)
+	{
+		using FColumnType = ConcertClientSharedSlate::ReplicationColumns::FReplicationPropertyColumn;
+		return FColumnType(
+			FColumnType::FArguments()
+				.GenerateWidgetColumn_Lambda([InClient, &InAuthorityCache, InViewer](const FColumnType::FBuildArgs& InArgs)
+				{
+					return SNew(SOwnerClientList, InClient, InAuthorityCache)
+						.GetClientList_Lambda([&InAuthorityCache, InViewer, Property = InArgs.RowData.GetProperty()](const FGlobalAuthorityCache&)
+						{
+							return Private::Property::GetPropertyOwners(InAuthorityCache, *InViewer.Get(), Property);
+						})
+						.HighlightText_Lambda([HighlightText = InArgs.HighlightText](){ return *HighlightText.Get(); });
+				})
+				.PopulateSearchItems_Lambda([InClient, &InAuthorityCache, InViewer](const ConcertClientSharedSlate::FReplicatedPropertyData& InArgs, TArray<FString>& InOutSearchStrings)
+				{
+					Private::Shared::PopulateSearchTerms(
+						InClient,
+						Private::Property::GetPropertyOwners(InAuthorityCache, *InViewer.Get(), InArgs.GetProperty()),
+						InOutSearchStrings
+						);
+				})
+				.ColumnSortOrder(static_cast<int32>(EPropertyColumnOrder::Owner)),
+			SHeaderRow::Column(OwnerOfPropertyColumnId)
+				.DefaultLabel(LOCTEXT("Property.Owner", "Owner"))
+				.FillSized(Private::Shared::DefaultWidth)
+			);
+	}
+}
+
+#undef LOCTEXT_NAMESPACE
\ No newline at end of file
diff --git a/Engine/Plugins/Developer/Concert/ConcertApp/MultiUserClient/Source/MultiUserClient/Private/Widgets/ActiveSession/Replication/Client/Columns/SingleClientColumns_ToggleAuthority.cpp b/Engine/Plugins/Developer/Concert/ConcertApp/MultiUserClient/Source/MultiUserClient/Private/Widgets/ActiveSession/Replication/Client/Columns/SingleClientColumns_ToggleAuthority.cpp
new file mode 100644
--- /dev/null
+++ b/Engine/Plugins/Developer/Concert/ConcertApp/MultiUserClient/Source/MultiUserClient/Private/Widgets/ActiveSession/Replication/Client/Columns/SingleClientColumns_ToggleAuthority.cpp
@@ -0,0 +1,218 @@
+// Copyright Epic Games, Inc. All Rights Reserved.
+
+#include "SingleClientColumns.h"
+
+#include "IConcertClient.h"
+#include "SOwnerClientList.h"
+#include "Replication/Authority/AuthorityChangeTracker.h"
+#include "Replication/Authority/EAuthorityMutability.h"
+#include "Replication/Authority/IClientAuthoritySynchronizer.h"
+#include "Replication/Editor/Model/IObjectToPropertiesModel.h"
+#include "Replication/Editor/View/IReplicationStreamViewer.h"
+#include "Replication/Editor/View/ReplicationColumnsUtils.h"
+#include "Replication/Submission/ISubmissionWorkflow.h"
+#include "Replication/Util/GlobalAuthorityCache.h"
+#include "Widgets/ClientName/SClientName.h"
+
+#define LOCTEXT_NAMESPACE "SingleClientColumns.ToggleAuthority"
+
+namespace UE::MultiUserClient::SingleClientColumns
+{
+	const FName ToggleTopLevelAuthorityColumnId = TEXT("ToggleTopLevelAuthorityColumn");
+	const FName ToggleSubobjectAuthorityColumnId = TEXT("ToggleSubobjectAuthorityColumn");
+
+	namespace Private::ToggleTopLevelAuthority
+	{
+		static ECheckBoxState GetCheckboxState(
+			const ConcertClientSharedSlate::FReplicatedObjectData& ObjectData,
+			ConcertClientSharedSlate::IObjectToPropertiesModel* ClientStreamModel,
+			FAuthorityChangeTracker* ChangeTracker
+			)
+		{
+			const FSoftObjectPath& ObjectPath = ObjectData.GetObjectPath();
+
+			// Init the checkbox state to the top level object, if the top level object has properties associated with it...
+			const bool bCanSetTopLevelState = ChangeTracker->CanSetAuthorityFor(ObjectPath)
+				|| ChangeTracker->GetChangeAuthorityMutability(ObjectPath) == EAuthorityMutability::NotSupported;
+			TOptional<ECheckBoxState> CheckBoxState = bCanSetTopLevelState
+				? ChangeTracker->GetAuthorityStateAfterApplied(ObjectPath) ? ECheckBoxState::Checked : ECheckBoxState::Unchecked
+				: TOptional<ECheckBoxState>{};
+
+			// ... and then determine whether all the other subobjects share that state (Check / Unchecked) or not (Undetermined)
+			ClientStreamModel->ForEachSubobject(ObjectPath, [&CheckBoxState, &ChangeTracker](const FSoftObjectPath& Child)
+			{
+				if (!ChangeTracker->CanSetAuthorityFor(Child)
+					&& ChangeTracker->GetChangeAuthorityMutability(Child) != EAuthorityMutability::NotSupported)
+				{
+					return EBreakBehavior::Continue;
+				}
+
+				// Init checkbox state
+				const bool bCurrentAuthorityState = ChangeTracker->GetAuthorityStateAfterApplied(Child);
+				const bool bIsFirstObject = !CheckBoxState.IsSet(); 
+				if (bIsFirstObject)
+				{
+					CheckBoxState = bCurrentAuthorityState ? ECheckBoxState::Checked : ECheckBoxState::Unchecked;
+					return EBreakBehavior::Continue;
+				}
+
+				// Return the checkbox state that all subobjects share, or Undetermined if mixed
+				const bool bMustBeUndetermined = (*CheckBoxState == ECheckBoxState::Unchecked && bCurrentAuthorityState)
+					|| (*CheckBoxState == ECheckBoxState::Checked && !bCurrentAuthorityState);
+				if (bMustBeUndetermined)
+				{
+					CheckBoxState = ECheckBoxState::Undetermined;
+					return EBreakBehavior::Break;
+				}
+				return EBreakBehavior::Continue;
+			});
+						
+			return CheckBoxState
+				? *CheckBoxState
+				// If this case happens, there are no objects to check. In that case, it visually "looks better" to be unchecked
+				: ECheckBoxState::Unchecked;
+		}
+
+		static void OnCheckboxStateChanged(
+			bool bIsChecked,
+			const ConcertClientSharedSlate::FReplicatedObjectData& ObjectData,
+			ConcertClientSharedSlate::IObjectToPropertiesModel* ClientStreamModel,
+			FAuthorityChangeTracker* ChangeTracker
+			)
+		{
+			const FSoftObjectPath& ObjectPath = ObjectData.GetObjectPath();
+			TArray<FSoftObjectPath> Paths { ObjectPath };
+			ClientStreamModel->ForEachSubobject(ObjectPath, [&Paths](const FSoftObjectPath& Child)
+			{
+				Paths.Add(Child);
+				return EBreakBehavior::Continue;
+			});
+			
+			ChangeTracker->SetAuthorityIfAllowed(Paths, bIsChecked);
+		}
+
+		static bool IsEnabled(
+			const ConcertClientSharedSlate::FReplicatedObjectData& ObjectData,
+			ConcertClientSharedSlate::IObjectToPropertiesModel* ClientStreamModel,
+			FAuthorityChangeTracker* ChangeTracker
+			)
+		{
+			const FSoftObjectPath& TopLevelObjectPath = ObjectData.GetObjectPath();
+			const bool bCanChangeTopLevel = ChangeTracker->CanSetAuthorityFor(TopLevelObjectPath);
+			if (bCanChangeTopLevel)
+			{
+				return true;
+			}
+						
+			bool bCanChangeAnySubobject = false;
+			ClientStreamModel->ForEachSubobject(TopLevelObjectPath, [&ChangeTracker, &bCanChangeAnySubobject](const FSoftObjectPath& Child)
+			{
+				bCanChangeAnySubobject |= ChangeTracker->CanSetAuthorityFor(Child);
+				return bCanChangeAnySubobject ? EBreakBehavior::Break : EBreakBehavior::Continue;
+			});
+			return bCanChangeAnySubobject;
+		}
+	}
+	
+	ConcertClientSharedSlate::ReplicationColumns::FReplicationTopLevelObjectColumn ToggleTopLevelAuthority(
+		ConcertClientSharedSlate::IObjectToPropertiesModel& ClientStreamModel,
+		FAuthorityChangeTracker& ChangeTracker,
+		ISubmissionWorkflow& SubmissionWorkflow
+		)
+	{
+		using namespace ConcertClientSharedSlate;
+		using namespace Private::ToggleTopLevelAuthority;
+		using FTopLevelColumnDelegates = TReplicationColumnDelegates<FReplicatedObjectData>;
+		
+		return MakeCheckboxColumn<FReplicatedObjectData>(
+			ToggleTopLevelAuthorityColumnId,
+				FTopLevelColumnDelegates(
+					FTopLevelColumnDelegates::FGetColumnCheckboxState::CreateStatic(&GetCheckboxState, &ClientStreamModel, &ChangeTracker),
+					FTopLevelColumnDelegates::FOnColumnCheckboxChanged::CreateStatic(&OnCheckboxStateChanged, &ClientStreamModel, &ChangeTracker),
+					FTopLevelColumnDelegates::FGetToolTipText::CreateLambda([&ClientStreamModel, &ChangeTracker, &SubmissionWorkflow](const FReplicatedObjectData& ObjectData)
+					{
+						
+						if (SubmissionWorkflow.GetUploadability() == EChangeUploadability::NotImplemented)
+						{
+							const ECheckBoxState CheckBoxState = GetCheckboxState(ObjectData, &ClientStreamModel, &ChangeTracker);
+							const FText CheckBoxText = [CheckBoxState]()
+							{
+								switch (CheckBoxState)
+								{
+								case ECheckBoxState::Unchecked: return LOCTEXT("TopLevel.ChangeAuthority.ToolTip.NotSupported.Unchecked", "None of the subobjects are being replicated.");
+								case ECheckBoxState::Checked: return LOCTEXT("TopLevel.ChangeAuthority.ToolTip.NotSupported.Checked", "All of the subobjects are being replicated.");
+								case ECheckBoxState::Undetermined: return LOCTEXT("TopLevel.ChangeAuthority.ToolTip.NotSupported.Undetermined", "Some of the subobjects are being replicated.");
+								default: checkNoEntry(); return FText::GetEmpty();
+								}
+							}();
+							
+							return FText::Format(
+								LOCTEXT("TopLevel.ChangeAuthority.ToolTip.NotSupportedFmt", "{0}\nEditing remote clients is not implemented. You can only edit the local client."),
+								CheckBoxText
+								);
+						}
+						
+						return IsEnabled(ObjectData, &ClientStreamModel, &ChangeTracker)
+							? LOCTEXT("TopLevel.ChangeAuthority.Allowed", "Toggles whether this object and its subobjects will replicate the assigned properties.")
+							: LOCTEXT("TopLevel.ChangeAuthority.Disallowed", "Toggles whether this object and its subobjects will replicate the assigned properties.\nDisabled because none of the subobjects can be replicated with their current configuration.");
+					}),
+					FTopLevelColumnDelegates::FIsEnabled::CreateStatic(&IsEnabled, &ClientStreamModel, &ChangeTracker)
+				),
+			FText::GetEmpty(),
+			static_cast<int32>(ETopLevelObjectColumnOrder::ToggleAuthority)
+		);
+	}
+
+	ConcertClientSharedSlate::ReplicationColumns::FReplicationSubobjectObjectColumn ToggleSubobjectAuthority(
+		FAuthorityChangeTracker& ChangeTracker,
+		ISubmissionWorkflow& SubmissionWorkflow
+		)
+	{
+		using namespace ConcertClientSharedSlate;
+		using FSubobjectColumnDelegates = TReplicationColumnDelegates<FReplicatedObjectData>;
+		return MakeCheckboxColumn<FReplicatedObjectData>(
+			ToggleTopLevelAuthorityColumnId,
+				FSubobjectColumnDelegates(
+					FSubobjectColumnDelegates::FGetColumnCheckboxState::CreateLambda(
+					[&ChangeTracker](const FReplicatedObjectData& ObjectData)
+					{
+						const bool bHasAuthority = ChangeTracker.GetAuthorityStateAfterApplied(ObjectData.GetObjectPath());
+						return bHasAuthority ? ECheckBoxState::Checked : ECheckBoxState::Unchecked;
+					}),
+					FSubobjectColumnDelegates::FOnColumnCheckboxChanged::CreateLambda(
+					[&ChangeTracker](bool bIsChecked, const FReplicatedObjectData& ObjectData)
+					{
+						ChangeTracker.SetAuthorityIfAllowed({ ObjectData.GetObjectPath() }, bIsChecked);
+					}),
+					FSubobjectColumnDelegates::FGetToolTipText::CreateLambda([&ChangeTracker, &SubmissionWorkflow](const FReplicatedObjectData& ObjectData)
+					{
+						const bool bHasAuthority = ChangeTracker.GetAuthorityStateAfterApplied(ObjectData.GetObjectPath());
+						const FText NotSupportedText = FText::Format(
+							LOCTEXT("Subobject.ChangeAuthority.ToolTip.NotSupportedFmt", "{0}\nEditing remote clients is not implemented. You can only edit the local client."),
+							bHasAuthority ? LOCTEXT("Subobject.ChangeAuthority.Replicating.True", "This object is being replicated.") : LOCTEXT("Subobject.ChangeAuthority.Replicating.False", "This object is not currently being replicated.")
+							);
+						if (SubmissionWorkflow.GetUploadability() == EChangeUploadability::NotImplemented)
+						{
+							return NotSupportedText;
+						}
+						
+						switch (ChangeTracker.GetChangeAuthorityMutability(ObjectData.GetObjectPath()))
+						{
+						case EAuthorityMutability::Allowed: return LOCTEXT("Subobject.ChangeAuthority.ToolTip.Allowed", "Toggles whether this object should be replicated.");
+						case EAuthorityMutability::NoProperties: return LOCTEXT("Subobject.ChangeAuthority.ToolTip.NoProperties", "Toggles whether this object should be replicated.\nAssign properties to this object first.\n");
+						case EAuthorityMutability::NotSupported: return NotSupportedText;
+						default: checkNoEntry(); return FText::GetEmpty();
+						}
+					}),
+					FSubobjectColumnDelegates::FIsEnabled::CreateLambda([&ChangeTracker](const FReplicatedObjectData& ObjectData)
+					{
+						return ChangeTracker.CanSetAuthorityFor(ObjectData.GetObjectPath());
+					})
+				),
+			FText::GetEmpty(),
+			static_cast<int32>(ESubobjectColumnOrder::ToggleAuthority)
+		);
+	}
+}
+
+#undef LOCTEXT_NAMESPACE
\ No newline at end of file
diff --git a/Engine/Plugins/Developer/Concert/ConcertApp/MultiUserClient/Source/MultiUserClient/Private/Widgets/ActiveSession/Replication/Client/Single/SOwnerClientList.cpp b/Engine/Plugins/Developer/Concert/ConcertApp/MultiUserClient/Source/MultiUserClient/Private/Widgets/ActiveSession/Replication/Client/Single/SOwnerClientList.cpp
deleted file mode 100644
diff --git a/Engine/Plugins/Developer/Concert/ConcertApp/MultiUserClient/Source/MultiUserClient/Private/Widgets/ActiveSession/Replication/Client/Single/SOwnerClientList.h b/Engine/Plugins/Developer/Concert/ConcertApp/MultiUserClient/Source/MultiUserClient/Private/Widgets/ActiveSession/Replication/Client/Single/SOwnerClientList.h
deleted file mode 100644
diff --git a/Engine/Plugins/Developer/Concert/ConcertApp/MultiUserClient/Source/MultiUserClient/Private/Widgets/ActiveSession/Replication/Client/Single/SReplicationClientView.cpp b/Engine/Plugins/Developer/Concert/ConcertApp/MultiUserClient/Source/MultiUserClient/Private/Widgets/ActiveSession/Replication/Client/Single/SReplicationClientView.cpp
--- a/Engine/Plugins/Developer/Concert/ConcertApp/MultiUserClient/Source/MultiUserClient/Private/Widgets/ActiveSession/Replication/Client/Single/SReplicationClientView.cpp
+++ b/Engine/Plugins/Developer/Concert/ConcertApp/MultiUserClient/Source/MultiUserClient/Private/Widgets/ActiveSession/Replication/Client/Single/SReplicationClientView.cpp
@@ -3,101 +3,206 @@
 #include "SReplicationClientView.h"
 
 #include "MultiUserReplicationSettings.h"
-#include "SingleClientColumns.h"
 #include "Replication/Client/ReplicationClient.h"
 #include "Replication/Editor/Model/IEditableObjectToPropertiesModel.h"
-#include "Replication/Editor/Model/Object/EditorObjectSelectionSourceModel.h"
+#include "Replication/Editor/Model/Object/ActorSelectionSourceModel.h"
 #include "Replication/Editor/Model/Property/SelectPropertyFromUClassModel.h"
 #include "Replication/Editor/View/IReplicationStreamEditor.h"
 #include "Replication/ReplicationWidgetFactories.h"
 #include "Replication/Client/ReplicationClientManager.h"
 #include "Replication/Submission/ISubmissionWorkflow.h"
+#include "Widgets/ActiveSession/Replication/Client/Columns/SingleClientColumns.h"
 #include "Widgets/ActiveSession/Replication/Client/Single/SSingleClientToolbar.h"
 
+#include "HAL/IConsoleManager.h"
 #include "Widgets/SBoxPanel.h"
+#include "Widgets/Layout/SBox.h"
 
 #define LOCTEXT_NAMESPACE "SReplicationClientView"
 
 namespace UE::MultiUserClient
 {
+	TAutoConsoleVariable<int32> CVarReplicationClientViewMode(
+		TEXT("MultiUser.ReplicationEditorMode"),
+		1,
+		TEXT("Determines the look of the editor mode.\n0 - Three Sections\n1 - Two sections")
+		);
+	
 	void SReplicationClientView::Construct(const FArguments& InArgs, const TSharedRef<IConcertClient>& InClient, FReplicationClientManager& InClientManager)
 	{
+		ConcertClient = InClient;
+		ClientManager = &InClientManager;
+		
 		GetReplicationClientAttribute = InArgs._GetReplicationClient;
 		FReplicationClient* ReplicationClient = GetReplicationClientAttribute.Get();
 		check(ReplicationClient);
-
+		
 		ConcertClientSharedSlate::IObjectToPropertiesModel& PropertyModel = *ReplicationClient->GetClientEditModel();
-		FAuthorityChangeTracker& AuthorityTracker = ReplicationClient->GetAuthorityDiffer();
-		ISubmissionWorkflow& SubmissionWorkflow = ReplicationClient->GetSubmissionWorkflow();
-		FGlobalAuthorityCache& AuthorityCache = InClientManager.GetAuthorityCache();
+		FGlobalAuthorityCache& AuthorityCache = ClientManager->GetAuthorityCache();
+		ChildSlot
+		[
+			SNew(SVerticalBox)
+
+			// Toolbar
+			+SVerticalBox::Slot()
+			.AutoHeight()
+			.Padding(2.f)
+			[
+				SNew(SSingleClientToolbar, PropertyModel, AuthorityCache)
+				.ViewSelectionArea() [ InArgs._ViewSelectionArea.Widget ]
+				.DisplayedClients(TSet{ ReplicationClient->GetEndpointId() })
+			]
+
+			// Editor
+			+SVerticalBox::Slot()
+			.FillHeight(1.f)
+			[
+				CreateEditorContent()
+			]
+		];
 		
-		TAttribute<const ConcertClientSharedSlate::IReplicationStreamViewer*> GetReplicationViewerAttribute =
-			TAttribute<const ConcertClientSharedSlate::IReplicationStreamViewer*>::CreateLambda([this](){ return EditorView.Get(); });
+		// Refresh UI if streams change externally, e.g. a remote client changed what they sent
+		ReplicationClient->OnModelExternallyChanged().AddSP(this, &SReplicationClientView::OnModelChanged);
+
+		CVarReplicationClientViewMode->OnChangedDelegate().AddSP(this, &SReplicationClientView::OnConsoleVariableChanged);
+	}
+
+	TSharedRef<SWidget> SReplicationClientView::CreateEditorContent()
+	{
+		Content = SNew(SBox);
+		RebuildContent();
+		return Content.ToSharedRef();
+	}
+
+	void SReplicationClientView::RebuildContent()
+	{
+		FReplicationClient* ReplicationClient = GetReplicationClientAttribute.Get();
+		check(ReplicationClient);
+
+		// Only one widget should be visible at the time so they do not interfere with each other.
+		switch (CVarReplicationClientViewMode.GetValueOnGameThread())
+		{
+		case 0:
+			Content->SetContent(CreateThreeSectionedContent(*ReplicationClient));
+			break;
+		case 1:
+			Content->SetContent(CreateTwoSectionedContent(*ReplicationClient));
+			break;
+
+		default:
+			Content->SetContent(SNullWidget::NullWidget);
+		}
+	}
+
+	TSharedRef<SWidget> SReplicationClientView::CreateThreeSectionedContent(
+		FReplicationClient& InReplicationClient
+		)
+	{
+		using namespace ConcertClientSharedSlate;
+		IObjectToPropertiesModel& PropertyModel = *InReplicationClient.GetClientEditModel();
+		FAuthorityChangeTracker& AuthorityTracker = InReplicationClient.GetAuthorityDiffer();
+		ISubmissionWorkflow& SubmissionWorkflow = InReplicationClient.GetSubmissionWorkflow();
+		FGlobalAuthorityCache& AuthorityCache = ClientManager->GetAuthorityCache();
+		
+		const TAttribute<const IReplicationStreamViewer*> GetReplicationViewerAttribute =
+			TAttribute<const IReplicationStreamViewer*>::CreateLambda([this](){ return EditorView_TwoSectioned.Get(); });
 		TAttribute<const FConcertReplicationEditorSettings*> ReplicationSettingsAttribute =
 			TAttribute<const FConcertReplicationEditorSettings*>::CreateLambda([](){ return &UMultiUserReplicationSettings::Get()->ReplicationEditorSettings; });
 
 		// Add checkboxes in front of top level and subobject rows for changing authority
-		using namespace ConcertClientSharedSlate;
 		const FCreateSubobjectViewParams SubobjectViewParams
 		{
 			.AdditionalColumns =
 			{
 				SingleClientColumns::ToggleSubobjectAuthority(AuthorityTracker, SubmissionWorkflow),
-				SingleClientColumns::ConflictWarningForSubobject(InClient, AuthorityCache, ReplicationClient->GetEndpointId()),
-				SingleClientColumns::OwnerOfSubobject(InClient, AuthorityCache)
+				SingleClientColumns::ConflictWarningForSubobject(ConcertClient.ToSharedRef(), AuthorityCache, InReplicationClient.GetEndpointId()),
+				SingleClientColumns::OwnerOfSubobject(ConcertClient.ToSharedRef(), AuthorityCache)
 			}
 		};
 		const FCreateEditorParams ReplicationEditorCreationParams
 		{
-			.DataModel = ReplicationClient->GetClientEditModel(),
-			.ObjectSource = MakeShared<FEditorObjectSelectionSourceModel>(),
+			.DataModel = InReplicationClient.GetClientEditModel(),
+			.ObjectSource = MakeShared<FActorSelectionSourceModel>(),
 			.PropertySource = MakeShared<FSelectPropertyFromUClassModel>(),
 			.SubobjectView = CreateUnrealEditorSubobjectView(SubobjectViewParams),
 			.AdditionalObjectColumns =
 			{
 				SingleClientColumns::ToggleTopLevelAuthority(PropertyModel, AuthorityTracker, SubmissionWorkflow),
-				SingleClientColumns::OwnerOfTopLevelObject(InClient, AuthorityCache, PropertyModel)
+				SingleClientColumns::OwnerOfTopLevelObject(ConcertClient.ToSharedRef(), AuthorityCache, PropertyModel)
 			},
 			.AdditionalPropertyColumns =
 			{
-				SingleClientColumns::OwnerOfProperty(InClient, AuthorityCache, MoveTemp(GetReplicationViewerAttribute)),
-				SingleClientColumns::ConflictWarningForProperty(InClient, GetReplicationViewerAttribute, AuthorityCache, ReplicationClient->GetEndpointId())
+				SingleClientColumns::OwnerOfProperty(ConcertClient.ToSharedRef(), AuthorityCache, GetReplicationViewerAttribute),
+				SingleClientColumns::ConflictWarningForProperty(ConcertClient.ToSharedRef(), GetReplicationViewerAttribute, AuthorityCache, InReplicationClient.GetEndpointId())
 			},
 			.IsEditingEnabled = TAttribute<bool>::CreateLambda([&SubmissionWorkflow](){ return SubmissionWorkflow.GetUploadability() != EChangeUploadability::NotImplemented; }),
 			.EditingDisabledToolTipText = LOCTEXT("Editing.NotImplemented", "Editing remote clients is not implemented. You can only edit the local client."),
 			.ReplicationSettingsAttribute = MoveTemp(ReplicationSettingsAttribute)
 		};
 
-		EditorView = CreateDefaultStreamEditor(ReplicationEditorCreationParams);
-		ChildSlot
-		[
-			SNew(SVerticalBox)
-
-			// Toolbar
-			+SVerticalBox::Slot()
-			.AutoHeight()
-			.Padding(2.f)
-			[
-				SNew(SSingleClientToolbar, PropertyModel, AuthorityCache)
-				.ViewSelectionArea() [ InArgs._ViewSelectionArea.Widget ]
-				.DisplayedClients(TSet{ ReplicationClient->GetEndpointId() })
-			]
+		EditorView_ThreeSectioned = CreateDefaultStreamEditor(ReplicationEditorCreationParams);
+		return EditorView_ThreeSectioned.ToSharedRef();
+	}
 
-			// Editor
-			+SVerticalBox::Slot()
-			.FillHeight(1.f)
-			[
-				EditorView.ToSharedRef()
-			]
-		];
+	TSharedRef<SWidget> SReplicationClientView::CreateTwoSectionedContent(
+		FReplicationClient& InReplicationClient
+		)
+	{
+		using namespace ConcertClientSharedSlate;
+		IObjectToPropertiesModel& PropertyModel = *InReplicationClient.GetClientEditModel();
+		FAuthorityChangeTracker& AuthorityTracker = InReplicationClient.GetAuthorityDiffer();
+		ISubmissionWorkflow& SubmissionWorkflow = InReplicationClient.GetSubmissionWorkflow();
+		FGlobalAuthorityCache& AuthorityCache = ClientManager->GetAuthorityCache();
 		
-		// Refresh UI if streams change externally, e.g. a remote client changed what they sent
-		ReplicationClient->OnModelExternallyChanged().AddSP(this, &SReplicationClientView::OnModelChanged);
+		const TAttribute<const IReplicationStreamViewer*> GetReplicationViewerAttribute =
+			TAttribute<const IReplicationStreamViewer*>::CreateLambda([this](){ return EditorView_TwoSectioned.Get(); });
+		TAttribute<const FConcertReplicationEditorSettings*> ReplicationSettingsAttribute =
+			TAttribute<const FConcertReplicationEditorSettings*>::CreateLambda([](){ return &UMultiUserReplicationSettings::Get()->ReplicationEditorSettings; });
+
+		// Add checkboxes in front of top level and subobject rows for changing authority
+		const FCreateEditorParams ReplicationEditorCreationParams
+		{
+			.DataModel = InReplicationClient.GetClientEditModel(),
+			.ObjectSource = MakeShared<FActorSelectionSourceModel>(),
+			.PropertySource = MakeShared<FSelectPropertyFromUClassModel>(),
+			.SubobjectModel = CreateDefaultComponentHierarchySubobjectModel(), // This makes actors have children in the top view
+			.AdditionalObjectColumns =
+			{
+				// TODO DP: When we decide which UI version to choose, remove subobject from these function names.
+				SingleClientColumns::ToggleSubobjectAuthority(AuthorityTracker, SubmissionWorkflow),
+				SingleClientColumns::ConflictWarningForSubobject(ConcertClient.ToSharedRef(), AuthorityCache, InReplicationClient.GetEndpointId()),
+				SingleClientColumns::OwnerOfSubobject(ConcertClient.ToSharedRef(), AuthorityCache)
+			},
+			.AdditionalPropertyColumns =
+			{
+				SingleClientColumns::OwnerOfProperty(ConcertClient.ToSharedRef(), AuthorityCache, GetReplicationViewerAttribute),
+				SingleClientColumns::ConflictWarningForProperty(ConcertClient.ToSharedRef(), GetReplicationViewerAttribute, AuthorityCache, InReplicationClient.GetEndpointId())
+			},
+			.IsEditingEnabled = TAttribute<bool>::CreateLambda([&SubmissionWorkflow](){ return SubmissionWorkflow.GetUploadability() != EChangeUploadability::NotImplemented; }),
+			.EditingDisabledToolTipText = LOCTEXT("Editing.NotImplemented", "Editing remote clients is not implemented. You can only edit the local client."),
+			.ReplicationSettingsAttribute = MoveTemp(ReplicationSettingsAttribute)
+		};
+
+		EditorView_TwoSectioned = CreateDefaultStreamEditor(ReplicationEditorCreationParams);
+		return EditorView_TwoSectioned.ToSharedRef();
 	}
-	
+
 	void SReplicationClientView::OnModelChanged() const
 	{
-		EditorView->Refresh();
+		// Only one is valid at the same time
+		if (EditorView_TwoSectioned)
+		{
+			EditorView_TwoSectioned->Refresh();
+		}
+		if (EditorView_ThreeSectioned)
+		{
+			EditorView_ThreeSectioned->Refresh();
+		}
+	}
+
+	void SReplicationClientView::OnConsoleVariableChanged(IConsoleVariable* ConsoleVariable)
+	{
+		RebuildContent();
 	}
 }
 
diff --git a/Engine/Plugins/Developer/Concert/ConcertApp/MultiUserClient/Source/MultiUserClient/Private/Widgets/ActiveSession/Replication/Client/Single/SReplicationClientView.h b/Engine/Plugins/Developer/Concert/ConcertApp/MultiUserClient/Source/MultiUserClient/Private/Widgets/ActiveSession/Replication/Client/Single/SReplicationClientView.h
--- a/Engine/Plugins/Developer/Concert/ConcertApp/MultiUserClient/Source/MultiUserClient/Private/Widgets/ActiveSession/Replication/Client/Single/SReplicationClientView.h
+++ b/Engine/Plugins/Developer/Concert/ConcertApp/MultiUserClient/Source/MultiUserClient/Private/Widgets/ActiveSession/Replication/Client/Single/SReplicationClientView.h
@@ -6,6 +6,7 @@
 #include "Widgets/SCompoundWidget.h"
 
 class IConcertClient;
+class SBox;
 
 namespace UE::ConcertClientSharedSlate
 {
@@ -14,6 +15,7 @@ namespace UE::ConcertClientSharedSlate
 
 namespace UE::MultiUserClient
 {
+	class FGlobalAuthorityCache;
 	class FReplicationClient;
 	class FReplicationClientManager;
 	
@@ -33,15 +35,32 @@ namespace UE::MultiUserClient
 		void Construct(const FArguments& InArgs, const TSharedRef<IConcertClient>& InClient, FReplicationClientManager& InClientManager);
 
 	private:
+
+		/** Holds the dynamic content, which changes based on CVarReplicationClientViewMode  */
+		TSharedPtr<SBox> Content;
+		
+		/** Used to rebuild Content */
+		TSharedPtr<IConcertClient> ConcertClient;
+		FReplicationClientManager* ClientManager;
 		
 		/** The editor view of the replication content. */
-		TSharedPtr<ConcertClientSharedSlate::IReplicationStreamEditor> EditorView;
+		TSharedPtr<ConcertClientSharedSlate::IReplicationStreamEditor> EditorView_TwoSectioned;
+		/** The editor view of the replication content. */
+		TSharedPtr<ConcertClientSharedSlate::IReplicationStreamEditor> EditorView_ThreeSectioned;
 		
 		/** The client to depict. Should always return true. If the client is destroyed, so should this widget be. */
 		TAttribute<FReplicationClient*> GetReplicationClientAttribute;
+
+
+		TSharedRef<SWidget> CreateEditorContent();
+		void RebuildContent();
+		TSharedRef<SWidget> CreateThreeSectionedContent(FReplicationClient& InReplicationClient);
+		TSharedRef<SWidget> CreateTwoSectionedContent(FReplicationClient& InReplicationClient);
 		
 		/** Called when any of the streams change. */
 		void OnModelChanged() const;
+		
+		void OnConsoleVariableChanged(IConsoleVariable* ConsoleVariable);
 	};
 }
 
diff --git a/Engine/Plugins/Developer/Concert/ConcertApp/MultiUserClient/Source/MultiUserClient/Private/Widgets/ActiveSession/Replication/Client/Single/SWarningIcon.cpp b/Engine/Plugins/Developer/Concert/ConcertApp/MultiUserClient/Source/MultiUserClient/Private/Widgets/ActiveSession/Replication/Client/Single/SWarningIcon.cpp
deleted file mode 100644
diff --git a/Engine/Plugins/Developer/Concert/ConcertApp/MultiUserClient/Source/MultiUserClient/Private/Widgets/ActiveSession/Replication/Client/Single/SWarningIcon.h b/Engine/Plugins/Developer/Concert/ConcertApp/MultiUserClient/Source/MultiUserClient/Private/Widgets/ActiveSession/Replication/Client/Single/SWarningIcon.h
deleted file mode 100644
diff --git a/Engine/Plugins/Developer/Concert/ConcertApp/MultiUserClient/Source/MultiUserClient/Private/Widgets/ActiveSession/Replication/Client/Single/SingleClientColumns.h b/Engine/Plugins/Developer/Concert/ConcertApp/MultiUserClient/Source/MultiUserClient/Private/Widgets/ActiveSession/Replication/Client/Single/SingleClientColumns.h
deleted file mode 100644
diff --git a/Engine/Plugins/Developer/Concert/ConcertApp/MultiUserClient/Source/MultiUserClient/Private/Widgets/ActiveSession/Replication/Client/Single/SingleClientColumns_ConflictWarning.cpp b/Engine/Plugins/Developer/Concert/ConcertApp/MultiUserClient/Source/MultiUserClient/Private/Widgets/ActiveSession/Replication/Client/Single/SingleClientColumns_ConflictWarning.cpp
deleted file mode 100644
diff --git a/Engine/Plugins/Developer/Concert/ConcertApp/MultiUserClient/Source/MultiUserClient/Private/Widgets/ActiveSession/Replication/Client/Single/SingleClientColumns_Owner.cpp b/Engine/Plugins/Developer/Concert/ConcertApp/MultiUserClient/Source/MultiUserClient/Private/Widgets/ActiveSession/Replication/Client/Single/SingleClientColumns_Owner.cpp
deleted file mode 100644
diff --git a/Engine/Plugins/Developer/Concert/ConcertApp/MultiUserClient/Source/MultiUserClient/Private/Widgets/ActiveSession/Replication/Client/Single/SingleClientColumns_ToggleAuthority.cpp b/Engine/Plugins/Developer/Concert/ConcertApp/MultiUserClient/Source/MultiUserClient/Private/Widgets/ActiveSession/Replication/Client/Single/SingleClientColumns_ToggleAuthority.cpp
deleted file mode 100644
diff --git a/Engine/Plugins/Developer/Concert/ConcertApp/MultiUserClient/Source/MultiUserReplicationEditor/Private/StreamEditor/ReplicationSessionPresetEditorToolkit.cpp b/Engine/Plugins/Developer/Concert/ConcertApp/MultiUserClient/Source/MultiUserReplicationEditor/Private/StreamEditor/ReplicationSessionPresetEditorToolkit.cpp
--- a/Engine/Plugins/Developer/Concert/ConcertApp/MultiUserClient/Source/MultiUserReplicationEditor/Private/StreamEditor/ReplicationSessionPresetEditorToolkit.cpp
+++ b/Engine/Plugins/Developer/Concert/ConcertApp/MultiUserClient/Source/MultiUserReplicationEditor/Private/StreamEditor/ReplicationSessionPresetEditorToolkit.cpp
@@ -5,7 +5,7 @@
 #include "Assets/MultiUserReplicationClientPreset.h"
 #include "Assets/MultiUserReplicationSessionPreset.h"
 #include "Replication/Editor/Model/Property/SelectPropertyFromUClassModel.h"
-#include "Replication/Editor/Model/Object/EditorObjectSelectionSourceModel.h"
+#include "Replication/Editor/Model/Object/ActorSelectionSourceModel.h"
 #include "Replication/Editor/View/IReplicationStreamEditor.h"
 #include "Replication/ReplicationWidgetFactories.h"
 
@@ -78,7 +78,7 @@ namespace UE::MultiUserReplicationEditor
 			*GetEditedStreamAsset(),
 			GetEditedStreamAsset()->GetUnassignedClient()->Stream->MakeReplicationMapGetterAttribute()
 			);
-		const TSharedRef<FEditorObjectSelectionSourceModel> ObjectSourceModel = MakeShared<FEditorObjectSelectionSourceModel>();
+		const TSharedRef<FActorSelectionSourceModel> ObjectSourceModel = MakeShared<FActorSelectionSourceModel>();
 		const TSharedRef<FSelectPropertyFromUClassModel> PropertySourceModel = MakeShared<FSelectPropertyFromUClassModel>();
 		const TSharedRef<IReplicationStreamEditor> EditorView = CreateDefaultStreamEditor(FCreateEditorParams{ AssetReadWriteModel, ObjectSourceModel, PropertySourceModel });
 		return SNew(SDockTab)
diff --git a/Engine/Plugins/Developer/Concert/ConcertUI/ConcertClientSharedSlate/Source/ConcertClientSharedSlate/Private/Replication/Editor/Model/ISubobjectModel.cpp b/Engine/Plugins/Developer/Concert/ConcertUI/ConcertClientSharedSlate/Source/ConcertClientSharedSlate/Private/Replication/Editor/Model/ISubobjectModel.cpp
new file mode 100644
--- /dev/null
+++ b/Engine/Plugins/Developer/Concert/ConcertUI/ConcertClientSharedSlate/Source/ConcertClientSharedSlate/Private/Replication/Editor/Model/ISubobjectModel.cpp
@@ -0,0 +1,38 @@
+// Copyright Epic Games, Inc. All Rights Reserved.
+
+#include "Replication/Editor/Model/ISubobjectModel.h"
+
+#include "UObject/SoftObjectPath.h"
+
+namespace UE::ConcertClientSharedSlate
+{
+	void ISubobjectModel::ForEachSubobject(TFunctionRef<EBreakBehavior(const FSoftObjectPath& Parent, const FSoftObjectPath& ChildObject)> Callback)
+	{
+		struct FHelper
+		{
+			static EBreakBehavior Visit(ISubobjectModel& Model, const FSoftObjectPath& Parent, const FSoftObjectPath& ChildObject, TFunctionRef<EBreakBehavior(const FSoftObjectPath&, const FSoftObjectPath&)> Callback)
+			{
+				if (Callback(Parent, ChildObject) == EBreakBehavior::Break)
+				{
+					return EBreakBehavior::Break;
+				}
+				
+				EBreakBehavior Result = EBreakBehavior::Continue;
+				Model.ForEachDirectChildSubobject(ChildObject, [&Model, NewParent = ChildObject, &Callback, &Result](const FSoftObjectPath& ChildObject)
+				{
+					Result = Visit(Model, NewParent, ChildObject, Callback);
+					return Result;
+				});
+				return Result;
+			}
+		};
+			
+		for (const FName& Category : GetCategories())
+		{
+			ForEachRootSubobject(Category, [this, &Callback](const FSoftObjectPath& ChildObject)
+			{
+				return FHelper::Visit(*this, GetTopLevelObject(), ChildObject, Callback);
+			});
+		}
+	}
+}
\ No newline at end of file
diff --git a/Engine/Plugins/Developer/Concert/ConcertUI/ConcertClientSharedSlate/Source/ConcertClientSharedSlate/Private/Replication/Editor/Model/Object/ActorSelectionSourceModel.cpp b/Engine/Plugins/Developer/Concert/ConcertUI/ConcertClientSharedSlate/Source/ConcertClientSharedSlate/Private/Replication/Editor/Model/Object/ActorSelectionSourceModel.cpp
new file mode 100644
--- /dev/null
+++ b/Engine/Plugins/Developer/Concert/ConcertUI/ConcertClientSharedSlate/Source/ConcertClientSharedSlate/Private/Replication/Editor/Model/Object/ActorSelectionSourceModel.cpp
@@ -0,0 +1,44 @@
+// Copyright Epic Games, Inc. All Rights Reserved.
+
+#include "Replication/Editor/Model/Object/ActorSelectionSourceModel.h"
+
+#include "Replication/Editor/Model/Object/SelectedActorsSource.h"
+#include "Replication/Editor/Model/Object/WorldActorSource.h"
+
+#include "GameFramework/Actor.h"
+
+#define LOCTEXT_NAMESPACE "FEditorObjectSelectionSourceModel"
+
+namespace UE::ConcertClientSharedSlate
+{
+	FActorSelectionSourceModel::FActorSelectionSourceModel()
+		: BaseActorCategory([this]()
+		{
+			return FObjectSourceCategory
+			{
+				ConcertSharedSlate::FBaseDisplayInfo
+				{
+					LOCTEXT("ObjectsCategory.Label", "Add Objects"),
+					LOCTEXT("ObjectsCategory.Tooltip", "Options for adding objects from the open editor world")
+				},
+				{
+					MakeShared<FSelectedActorsSource>(),
+					MakeShared<FWorldActorSource>()
+				}
+			};
+		}())
+	{}
+
+	TArray<FObjectSourceCategory> FActorSelectionSourceModel::GetRootSources() const
+	{
+		FObjectSourceCategory ActorCategory = BaseActorCategory;
+		return { ActorCategory };
+	}
+
+	TArray<TSharedRef<IObjectSourceModel>> FActorSelectionSourceModel::GetContextMenuOptions(const FSoftObjectPath& Item)
+	{
+		return {};
+	}
+}
+
+#undef LOCTEXT_NAMESPACE
\ No newline at end of file
diff --git a/Engine/Plugins/Developer/Concert/ConcertUI/ConcertClientSharedSlate/Source/ConcertClientSharedSlate/Private/Replication/Editor/Model/Object/EditorObjectSelectionSourceModel.cpp b/Engine/Plugins/Developer/Concert/ConcertUI/ConcertClientSharedSlate/Source/ConcertClientSharedSlate/Private/Replication/Editor/Model/Object/EditorObjectSelectionSourceModel.cpp
deleted file mode 100644
diff --git a/Engine/Plugins/Developer/Concert/ConcertUI/ConcertClientSharedSlate/Source/ConcertClientSharedSlate/Private/Replication/Editor/Model/Subobject/ComponentHierarchySubobjectModel.cpp b/Engine/Plugins/Developer/Concert/ConcertUI/ConcertClientSharedSlate/Source/ConcertClientSharedSlate/Private/Replication/Editor/Model/Subobject/ComponentHierarchySubobjectModel.cpp
--- a/Engine/Plugins/Developer/Concert/ConcertUI/ConcertClientSharedSlate/Source/ConcertClientSharedSlate/Private/Replication/Editor/Model/Subobject/ComponentHierarchySubobjectModel.cpp
+++ b/Engine/Plugins/Developer/Concert/ConcertUI/ConcertClientSharedSlate/Source/ConcertClientSharedSlate/Private/Replication/Editor/Model/Subobject/ComponentHierarchySubobjectModel.cpp
@@ -7,6 +7,7 @@
 #include "Misc/EBreakBehavior.h"
 
 #include "GameFramework/Actor.h"
+#include "Replication/Editor/View/ObjectUtils.h"
 
 namespace UE::ConcertClientSharedSlate
 {
@@ -26,8 +27,8 @@ namespace UE::ConcertClientSharedSlate
 		if (InTopLevelObject != TopLevelObject)
 		{
 			TopLevelObject = InTopLevelObject;
-
-			if (AActor* Actor = GetTopLevelObjectAsActor())
+			if (AActor* Actor = GetTopLevelObjectAsActor()
+				; !TopLevelObject.IsNull() && ensureAlwaysMsgf(Actor, TEXT("Not a top-level object")))
 			{
 				USubobjectDataSubsystem* SubobjectDataSubsystem = USubobjectDataSubsystem::Get();
 				TArray<FSubobjectDataHandle> Handles;
@@ -41,18 +42,22 @@ namespace UE::ConcertClientSharedSlate
 					ObjectMetaData.Add(Object, { FText::FromString(SubobjectData->GetDisplayString()) });
 				}
 			}
-			
 			OnHierarchyChangedDelegate.Broadcast();
 		}
 	}
 
+	bool FComponentHierarchySubobjectModel::IsTopLevelObject(const FSoftObjectPath& Object) const
+	{
+		return ObjectUtils::IsActor(Object);
+	}
+
 	FText FComponentHierarchySubobjectModel::GetSubobjectDisplayName(const FSoftObjectPath& ObjectPath) const
 	{
 		const FObjectMetaData* MetaData = ObjectMetaData.Find(ObjectPath);
 		return MetaData ? MetaData->DisplayLabel : FText::GetEmpty();
 	}
 
-	void FComponentHierarchySubobjectModel::ForEachRootSubobject(FName Category, TFunctionRef<EBreakBehavior(const FSoftObjectPath& Object)> Callback)
+	void FComponentHierarchySubobjectModel::ForEachRootSubobject(FName Category, TFunctionRef<EBreakBehavior(const FSoftObjectPath& Object)> Callback) const
 	{
 		const AActor* Actor = GetTopLevelObjectAsActor();
 		if (!Actor)
@@ -81,11 +86,11 @@ namespace UE::ConcertClientSharedSlate
 		}
 	}
 
-	void FComponentHierarchySubobjectModel::ForEachDirectChildSubobject(const FSoftObjectPath& Parent, TFunctionRef<EBreakBehavior(const FSoftObjectPath&)> Callback)
+	void FComponentHierarchySubobjectModel::ForEachDirectChildSubobject(const FSoftObjectPath& Parent, TFunctionRef<EBreakBehavior(const FSoftObjectPath&)> Callback) const
 	{
 		const UObject* Object = Parent.ResolveObject();
 		const UObject* TopLevel = GetTopLevelObjectAsActor();
-		if (!Object || !ensure(TopLevel && Object->IsIn(TopLevel)))
+		if (!Object || !ensureAlways(TopLevel && Object->IsIn(TopLevel)))
 		{
 			return;
 		}
diff --git a/Engine/Plugins/Developer/Concert/ConcertUI/ConcertClientSharedSlate/Source/ConcertClientSharedSlate/Private/Replication/Editor/Model/Subobject/ComponentHierarchySubobjectModel.h b/Engine/Plugins/Developer/Concert/ConcertUI/ConcertClientSharedSlate/Source/ConcertClientSharedSlate/Private/Replication/Editor/Model/Subobject/ComponentHierarchySubobjectModel.h
--- a/Engine/Plugins/Developer/Concert/ConcertUI/ConcertClientSharedSlate/Source/ConcertClientSharedSlate/Private/Replication/Editor/Model/Subobject/ComponentHierarchySubobjectModel.h
+++ b/Engine/Plugins/Developer/Concert/ConcertUI/ConcertClientSharedSlate/Source/ConcertClientSharedSlate/Private/Replication/Editor/Model/Subobject/ComponentHierarchySubobjectModel.h
@@ -2,7 +2,7 @@
 
 #pragma once
 
-#include "ISubobjectModel.h"
+#include "Replication/Editor/Model/ISubobjectModel.h"
 #include "UObject/SoftObjectPtr.h"
 
 class AActor;
@@ -25,10 +25,11 @@ namespace UE::ConcertClientSharedSlate
 		//~ Begin ISubobjectModel Interface
 		virtual void SetTopLevelObject(const FSoftObjectPath& InTopLevelObject) override;
 		virtual FSoftObjectPath GetTopLevelObject() const override { return TopLevelObject; }
+		virtual bool IsTopLevelObject(const FSoftObjectPath& Object) const override;
 		virtual TArray<FName> GetCategories() const override { return { SceneComponentsCategory, ActorComponentsCategory }; }
 		virtual FText GetSubobjectDisplayName(const FSoftObjectPath& ObjectPath) const override;
-		virtual void ForEachRootSubobject(FName Category, TFunctionRef<EBreakBehavior(const FSoftObjectPath& Object)> Callback) override;
-		virtual void ForEachDirectChildSubobject(const FSoftObjectPath& Parent, TFunctionRef<EBreakBehavior(const FSoftObjectPath&)> Callback) override;
+		virtual void ForEachRootSubobject(FName Category, TFunctionRef<EBreakBehavior(const FSoftObjectPath& Object)> Callback) const override;
+		virtual void ForEachDirectChildSubobject(const FSoftObjectPath& Parent, TFunctionRef<EBreakBehavior(const FSoftObjectPath&)> Callback) const override;
 		virtual FOnHierarchyChanged& OnHierarchyChanged() override { return OnHierarchyChangedDelegate; }
 		//~ End ISubobjectModel Interface
 
diff --git a/Engine/Plugins/Developer/Concert/ConcertUI/ConcertClientSharedSlate/Source/ConcertClientSharedSlate/Private/Replication/Editor/Model/Subobject/ISubobjectModel.h b/Engine/Plugins/Developer/Concert/ConcertUI/ConcertClientSharedSlate/Source/ConcertClientSharedSlate/Private/Replication/Editor/Model/Subobject/ISubobjectModel.h
deleted file mode 100644
diff --git a/Engine/Plugins/Developer/Concert/ConcertUI/ConcertClientSharedSlate/Source/ConcertClientSharedSlate/Private/Replication/Editor/View/ObjectEditor/FakeObjectToPropertiesEditorModel.cpp b/Engine/Plugins/Developer/Concert/ConcertUI/ConcertClientSharedSlate/Source/ConcertClientSharedSlate/Private/Replication/Editor/View/ObjectEditor/FakeObjectToPropertiesEditorModel.cpp
--- a/Engine/Plugins/Developer/Concert/ConcertUI/ConcertClientSharedSlate/Source/ConcertClientSharedSlate/Private/Replication/Editor/View/ObjectEditor/FakeObjectToPropertiesEditorModel.cpp
+++ b/Engine/Plugins/Developer/Concert/ConcertUI/ConcertClientSharedSlate/Source/ConcertClientSharedSlate/Private/Replication/Editor/View/ObjectEditor/FakeObjectToPropertiesEditorModel.cpp
@@ -4,44 +4,13 @@
 
 #include "Replication/Editor/Model/Property/IPropertySelectionSourceModel.h"
 #include "Replication/Editor/Model/Property/IPropertySourceModel.h"
+#include "Replication/Editor/View/ObjectUtils.h"
 #include "Replication/PropertyChainUtils.h"
 
 #include "GameFramework/Actor.h"
 
 namespace UE::ConcertClientSharedSlate
 {
-	namespace Private
-	{
-		static TOptional<FSoftObjectPath> GetTopLevelObjectOf(const FSoftObjectPath& SoftObjectPath)
-		{
-			// Example of an actor called floor
-			// SoftObjectPath = { AssetPath = {PackageName = "/Game/Maps/SyncBoxLevel", AssetName = "SyncBoxLevel"}, SubPathString = "PersistentLevel.Floor" } }
-			const FString& SubPathString = SoftObjectPath.GetSubPathString();
-
-			constexpr int32 PersistentLevelStringLength = 16; // "PersistentLevel." has 16 characters
-			const bool bIsWorldObject = SubPathString.Contains(TEXT("PersistentLevel."), ESearchCase::CaseSensitive);
-			if (!bIsWorldObject)
-			{
-				// Not a path to a world object
-				return {};
-			}
-
-			// Start search after the . behind PersistentLevel
-			const int32 StartSearch = PersistentLevelStringLength + 1;
-			const int32 IndexOfDotAfterActorName = SubPathString.Find(TEXT("."), ESearchCase::CaseSensitive, ESearchDir::FromStart, StartSearch);
-			if (IndexOfDotAfterActorName == INDEX_NONE)
-			{
-				// SoftObjectPath points to an actor
-				return {};
-			}
-
-			const int32 NumToChopOffRight = SubPathString.Len() - IndexOfDotAfterActorName;
-			const FString NewSubstring = SubPathString.LeftChop(NumToChopOffRight);
-			const FSoftObjectPath PathToOwningActor(SoftObjectPath.GetAssetPath(), NewSubstring);
-			return PathToOwningActor;
-		}
-	}
-	
 	bool FFakeObjectToPropertiesEditorModel::IsTopLevelObject(const FSoftObjectPath& ObjectPath) const
 	{
 		const FSoftClassPath ClassPath = GetObjectClass(ObjectPath);
@@ -64,8 +33,10 @@ namespace UE::ConcertClientSharedSlate
 
 	bool FFakeObjectToPropertiesEditorModel::ForEachReplicatedObject(TFunctionRef<EBreakBehavior(const FSoftObjectPath& Object)> Delegate) const
 	{
-		// This makes SObjectToPropertyView only show actors in the outliner
-		return IterateTopLevelObjects(Delegate);
+		return EnumHasAnyFlags(Flags, EFakeObjectModelFlags::OnlyTopLevelObjects)
+			// This makes SReplicationStreamViewer only show actors in the outliner
+			? IterateTopLevelObjects(Delegate) 
+			: RealModel->ForEachReplicatedObject(Delegate);
 	}
 
 	bool FFakeObjectToPropertiesEditorModel::ForEachProperty(const FSoftObjectPath& ObjectPath, TFunctionRef<EBreakBehavior(const FConcertPropertyChain& Property)> Delegate) const
@@ -105,7 +76,7 @@ namespace UE::ConcertClientSharedSlate
 		// Now determine the top level objects that are not in RealModel but that need to be shown for subobjects that are in RealModel
 		for (const FSoftObjectPath& NonTopLevelObject : NonTopLevelObjects)
 		{
-			const TOptional<FSoftObjectPath> TopLevelObject = Private::GetTopLevelObjectOf(NonTopLevelObject);
+			const TOptional<FSoftObjectPath> TopLevelObject = ObjectUtils::GetActorOf(NonTopLevelObject);
 			if (!TopLevelObject)
 			{
 				continue;
diff --git a/Engine/Plugins/Developer/Concert/ConcertUI/ConcertClientSharedSlate/Source/ConcertClientSharedSlate/Private/Replication/Editor/View/ObjectEditor/FakeObjectToPropertiesEditorModel.h b/Engine/Plugins/Developer/Concert/ConcertUI/ConcertClientSharedSlate/Source/ConcertClientSharedSlate/Private/Replication/Editor/View/ObjectEditor/FakeObjectToPropertiesEditorModel.h
--- a/Engine/Plugins/Developer/Concert/ConcertUI/ConcertClientSharedSlate/Source/ConcertClientSharedSlate/Private/Replication/Editor/View/ObjectEditor/FakeObjectToPropertiesEditorModel.h
+++ b/Engine/Plugins/Developer/Concert/ConcertUI/ConcertClientSharedSlate/Source/ConcertClientSharedSlate/Private/Replication/Editor/View/ObjectEditor/FakeObjectToPropertiesEditorModel.h
@@ -7,6 +7,13 @@
 namespace UE::ConcertClientSharedSlate
 {
 	class IPropertySelectionSourceModel;
+
+	enum class EFakeObjectModelFlags
+	{
+		None,
+		OnlyTopLevelObjects = 1 << 0
+	};
+	ENUM_CLASS_FLAGS(EFakeObjectModelFlags);
 	
 	/**
 	 * This model is passed tp SObjectToPropertyView by SObjectToPropertyEditor.
@@ -24,24 +31,27 @@ namespace UE::ConcertClientSharedSlate
 	{
 	public:
 
-		FFakeObjectToPropertiesEditorModel(TSharedRef<IObjectToPropertiesModel> RealModel, TSharedRef<IPropertySelectionSourceModel> PropertySelectionSource)
+		FFakeObjectToPropertiesEditorModel(
+			TSharedRef<IObjectToPropertiesModel> RealModel,
+			TSharedRef<IPropertySelectionSourceModel> PropertySelectionSource,
+			EFakeObjectModelFlags Flags
+			)
 			: RealModel(MoveTemp(RealModel))
 			, PropertySelectionSource(MoveTemp(PropertySelectionSource))
+			, Flags(Flags)
 		{}
 		
 		/** Whether this object should be displayed in the outliner. */
 		bool IsTopLevelObject(const FSoftObjectPath& ObjectPath) const;
 
 		//~ Begin IObjectToPropertiesModel Interface
-		
 		// Technically these functions should be also be wrapped but the SObjectToPropertyView does not use them so let's not for now.
 		virtual uint32 GetNumReplicatedObjects() const override { return RealModel->GetNumReplicatedObjects(); }
 		virtual uint32 GetNumProperties(const FSoftObjectPath& Object) const override { return RealModel->GetNumProperties(Object); }
 		virtual bool ContainsObjects(const TSet<FSoftObjectPath>& Objects) const override { return RealModel->ContainsObjects(Objects); }
+		virtual bool ContainsProperties(const FSoftObjectPath& Object, const TSet<FConcertPropertyChain>& Properties) const override { return RealModel->ContainsProperties(Object, Properties); }
 		
 		virtual FSoftClassPath GetObjectClass(const FSoftObjectPath& Object) const override;
-		// This function is by SObjectToPropertyEditor and not used by SObjectToPropertyView so forward it normally
-		virtual bool ContainsProperties(const FSoftObjectPath& Object, const TSet<FConcertPropertyChain>& Properties) const override { return RealModel->ContainsProperties(Object, Properties); }
 		virtual bool ForEachReplicatedObject(TFunctionRef<EBreakBehavior(const FSoftObjectPath& Object)> Delegate) const override;
 		virtual bool ForEachProperty(const FSoftObjectPath& ObjectPath, TFunctionRef<EBreakBehavior(const FConcertPropertyChain& Property)> Delegate) const override;
 		//~ End IObjectToPropertiesModel Interface
@@ -49,14 +59,16 @@ namespace UE::ConcertClientSharedSlate
 	private:
 
 		/** The real model which is used to implement all the other functions. */
-		TSharedRef<IObjectToPropertiesModel> RealModel;
+		const TSharedRef<IObjectToPropertiesModel> RealModel;
 
 		/** Determines the properties that can be selected. */
-		TSharedRef<IPropertySelectionSourceModel> PropertySelectionSource;
+		const TSharedRef<IPropertySelectionSourceModel> PropertySelectionSource;
+
+		/** Affects behaviour of this model. */
+		const EFakeObjectModelFlags Flags;
 
 		/** Returns all objects that must be listed in the top-outliner section of the replication view. */
 		bool IterateTopLevelObjects(TFunctionRef<EBreakBehavior(const FSoftObjectPath& Object)> Delegate) const;
-
 		/** Iterates all properties on the given object's path. */
 		bool IterateDisplayedProperties(const FSoftObjectPath& ObjectPath, TFunctionRef<EBreakBehavior(const FConcertPropertyChain& Property)> Delegate) const;
 	};
diff --git a/Engine/Plugins/Developer/Concert/ConcertUI/ConcertClientSharedSlate/Source/ConcertClientSharedSlate/Private/Replication/Editor/View/ObjectEditor/SBaseReplicationStreamEditor.cpp b/Engine/Plugins/Developer/Concert/ConcertUI/ConcertClientSharedSlate/Source/ConcertClientSharedSlate/Private/Replication/Editor/View/ObjectEditor/SBaseReplicationStreamEditor.cpp
--- a/Engine/Plugins/Developer/Concert/ConcertUI/ConcertClientSharedSlate/Source/ConcertClientSharedSlate/Private/Replication/Editor/View/ObjectEditor/SBaseReplicationStreamEditor.cpp
+++ b/Engine/Plugins/Developer/Concert/ConcertUI/ConcertClientSharedSlate/Source/ConcertClientSharedSlate/Private/Replication/Editor/View/ObjectEditor/SBaseReplicationStreamEditor.cpp
@@ -13,6 +13,7 @@
 #include "Replication/Settings/ConcertReplicationEditorSettings.h"
 
 #include "Algo/AnyOf.h"
+#include "Replication/Editor/Model/ISubobjectModel.h"
 #include "UObject/Class.h"
 #include "Widgets/SBoxPanel.h"
 
@@ -28,11 +29,15 @@ namespace UE::ConcertClientSharedSlate
 	{
 		ObjectSelectionSource = MoveTemp(InObjectSelectionSource);
 		PropertySelectionSource = MoveTemp(InPropertySelectionSource);
+		SubobjectModel = InArgs._SubobjectModel;
 		
 		EditablePropertiesModel = MoveTemp(InPropertiesModel);
 		EditablePropertiesModel->OnObjectsChanged().AddSP(this, &SBaseReplicationStreamEditor::OnObjectsChanged);
 		EditablePropertiesModel->OnPropertiesChanged().AddSP(this, &SBaseReplicationStreamEditor::OnPropertiesChanged);
-		PropertiesModelAdapter = MakeShared<FFakeObjectToPropertiesEditorModel>(EditablePropertiesModel.ToSharedRef(), PropertySelectionSource.ToSharedRef());
+		
+		const bool bOutlinerCanShowSubobjects = InArgs._SubobjectModel.IsValid();
+		const EFakeObjectModelFlags ModelFlags = bOutlinerCanShowSubobjects ? EFakeObjectModelFlags::None : EFakeObjectModelFlags::OnlyTopLevelObjects;
+		PropertiesModelAdapter = MakeShared<FFakeObjectToPropertiesEditorModel>(EditablePropertiesModel.ToSharedRef(), PropertySelectionSource.ToSharedRef(), ModelFlags);
 
 		IsEditingEnabledAttribute = InArgs._IsEditingEnabled;
 		EditingDisabledToolTipTextAttribute = InArgs._EditingDisabledToolTipText;
@@ -45,6 +50,7 @@ namespace UE::ConcertClientSharedSlate
 				.AdditionalObjectColumns(InArgs._AdditionalObjectColumns)
 				.AdditionalPropertyColumns(InArgs._AdditionalPropertyColumns)
 				.SubobjectView(InArgs._SubobjectView)
+				.SubobjectModel(InArgs._SubobjectModel)
 				.OnDeleteObjects(this, &SBaseReplicationStreamEditor::OnDeleteObjects)
 				.OnObjectsContextMenuOpening(this, &SBaseReplicationStreamEditor::OnObjectsContextMenuOpening)
 				.SortPropertyRowPredicate(InArgs._SortPropertyRowPredicate)
@@ -109,7 +115,7 @@ namespace UE::ConcertClientSharedSlate
 	void SBaseReplicationStreamEditor::OnObjectsChanged(TConstArrayView<UObject*> AddedObjects, TConstArrayView<FSoftObjectPath> RemovedObjects, EReplicatedObjectChangeReason ChangeReason)
 	{
 		ReplicationViewer->RefreshObjectData();
-
+			
 		// Newly added objects should be automatically selected
 		if (!AddedObjects.IsEmpty())
 		{
@@ -121,6 +127,8 @@ namespace UE::ConcertClientSharedSlate
 				return PropertiesModelAdapter->IsTopLevelObject(Object);
 			}, [](const UObject* Object){ return Object; });
 			ReplicationViewer->SelectTopLevelObjects(TopLevelObjects);
+			constexpr bool bRecursive = true;
+			ReplicationViewer->ExpandObjects(TopLevelObjects, bRecursive);
 		}
 	}
 
diff --git a/Engine/Plugins/Developer/Concert/ConcertUI/ConcertClientSharedSlate/Source/ConcertClientSharedSlate/Private/Replication/Editor/View/ObjectEditor/SBaseReplicationStreamEditor.h b/Engine/Plugins/Developer/Concert/ConcertUI/ConcertClientSharedSlate/Source/ConcertClientSharedSlate/Private/Replication/Editor/View/ObjectEditor/SBaseReplicationStreamEditor.h
--- a/Engine/Plugins/Developer/Concert/ConcertUI/ConcertClientSharedSlate/Source/ConcertClientSharedSlate/Private/Replication/Editor/View/ObjectEditor/SBaseReplicationStreamEditor.h
+++ b/Engine/Plugins/Developer/Concert/ConcertUI/ConcertClientSharedSlate/Source/ConcertClientSharedSlate/Private/Replication/Editor/View/ObjectEditor/SBaseReplicationStreamEditor.h
@@ -24,6 +24,7 @@ namespace UE::ConcertClientSharedSlate
 	class IObjectToPropertiesModel;
 	class IPropertySelectionSourceModel;
 	class IReplicationSubobjectView;
+	class ISubobjectModel;
 	class SReplicationStreamViewer;
 	
 	enum class EReplicatedObjectChangeReason : uint8;
@@ -57,6 +58,8 @@ namespace UE::ConcertClientSharedSlate
 		
 			/** Optional. Placed between root object outliner and property editor. */
 			SLATE_ARGUMENT(TSharedPtr<IReplicationSubobjectView>, SubobjectView)
+			/** Optional. If set, this determines the children nested under the root objects. */
+			SLATE_ARGUMENT(TSharedPtr<ISubobjectModel>, SubobjectModel)
 
 			/** Called to generate the context menu for objects. This extends the options already generated by this widget. */
 			SLATE_EVENT(FExtendMenu, OnExtendObjectsContextMenu)
@@ -103,6 +106,8 @@ namespace UE::ConcertClientSharedSlate
 		 * We inject checkboxes to SObjectToPropertyView which do the actual adding and removing.
 		 */
 		TSharedPtr<FFakeObjectToPropertiesEditorModel> PropertiesModelAdapter;
+		/** Can be null. If set, adds all subobjects to the top level view. */
+		TSharedPtr<ISubobjectModel> SubobjectModel;
 		
 		/** For deciding which objects can be added to EditablePropertiesModel. */
 		TSharedPtr<IObjectSelectionSourceModel> ObjectSelectionSource;
diff --git a/Engine/Plugins/Developer/Concert/ConcertUI/ConcertClientSharedSlate/Source/ConcertClientSharedSlate/Private/Replication/Editor/View/ObjectEditor/SDefaultReplicationStreamEditor.cpp b/Engine/Plugins/Developer/Concert/ConcertUI/ConcertClientSharedSlate/Source/ConcertClientSharedSlate/Private/Replication/Editor/View/ObjectEditor/SDefaultReplicationStreamEditor.cpp
--- a/Engine/Plugins/Developer/Concert/ConcertUI/ConcertClientSharedSlate/Source/ConcertClientSharedSlate/Private/Replication/Editor/View/ObjectEditor/SDefaultReplicationStreamEditor.cpp
+++ b/Engine/Plugins/Developer/Concert/ConcertUI/ConcertClientSharedSlate/Source/ConcertClientSharedSlate/Private/Replication/Editor/View/ObjectEditor/SDefaultReplicationStreamEditor.cpp
@@ -37,6 +37,7 @@ namespace UE::ConcertClientSharedSlate
 				.AdditionalObjectColumns(InArgs._AdditionalObjectColumns)
 				.AdditionalPropertyColumns(PropertyColumns)
 				.SubobjectView(InArgs._SubobjectView)
+				.SubobjectModel(InArgs._SubobjectModel)
 				.OnExtendObjectsContextMenu(InArgs._OnExtendObjectsContextMenu)
 				.SortPropertyRowPredicate(this, &SDefaultReplicationStreamEditor::SortPropertiesPredicate)
 				.LeftOfObjectSearchBar()
diff --git a/Engine/Plugins/Developer/Concert/ConcertUI/ConcertClientSharedSlate/Source/ConcertClientSharedSlate/Private/Replication/Editor/View/ObjectEditor/SDefaultReplicationStreamEditor.h b/Engine/Plugins/Developer/Concert/ConcertUI/ConcertClientSharedSlate/Source/ConcertClientSharedSlate/Private/Replication/Editor/View/ObjectEditor/SDefaultReplicationStreamEditor.h
--- a/Engine/Plugins/Developer/Concert/ConcertUI/ConcertClientSharedSlate/Source/ConcertClientSharedSlate/Private/Replication/Editor/View/ObjectEditor/SDefaultReplicationStreamEditor.h
+++ b/Engine/Plugins/Developer/Concert/ConcertUI/ConcertClientSharedSlate/Source/ConcertClientSharedSlate/Private/Replication/Editor/View/ObjectEditor/SDefaultReplicationStreamEditor.h
@@ -18,6 +18,7 @@ namespace UE::ConcertClientSharedSlate
 	class IObjectSelectionSourceModel;
 	class IPropertySelectionSourceModel;
 	class IReplicationSubobjectView;
+	class ISubobjectModel;
 	class SBaseReplicationStreamEditor;
 
 	/**
@@ -39,6 +40,8 @@ namespace UE::ConcertClientSharedSlate
 			
 			/** Optional. Placed between root object outliner and property editor. */
 			SLATE_ARGUMENT(TSharedPtr<IReplicationSubobjectView>, SubobjectView)
+			/** Optional. If set, this determines the children nested under the root objects. */
+			SLATE_ARGUMENT(TSharedPtr<ISubobjectModel>, SubobjectModel)
 
 			/** Called to generate the context menu for objects. This extends the options already generated by this widget. */
 			SLATE_EVENT(FExtendMenu, OnExtendObjectsContextMenu)
diff --git a/Engine/Plugins/Developer/Concert/ConcertUI/ConcertClientSharedSlate/Source/ConcertClientSharedSlate/Private/Replication/Editor/View/ObjectUtils.cpp b/Engine/Plugins/Developer/Concert/ConcertUI/ConcertClientSharedSlate/Source/ConcertClientSharedSlate/Private/Replication/Editor/View/ObjectUtils.cpp
new file mode 100644
--- /dev/null
+++ b/Engine/Plugins/Developer/Concert/ConcertUI/ConcertClientSharedSlate/Source/ConcertClientSharedSlate/Private/Replication/Editor/View/ObjectUtils.cpp
@@ -0,0 +1,59 @@
+// Copyright Epic Games, Inc. All Rights Reserved.
+
+#include "ObjectUtils.h"
+
+#include "Containers/UnrealString.h"
+#include "UObject/SoftObjectPath.h"
+
+namespace UE::ConcertClientSharedSlate::ObjectUtils
+{
+	bool IsActor(const FSoftObjectPath& SoftObjectPath)
+	{
+		// Example of an actor called floor
+		// SoftObjectPath = { AssetPath = {PackageName = "/Game/Maps/SyncBoxLevel", AssetName = "SyncBoxLevel"}, SubPathString = "PersistentLevel.Floor" } }
+		const FString& SubPathString = SoftObjectPath.GetSubPathString();
+
+		constexpr int32 PersistentLevelStringLength = 16; // "PersistentLevel." has 16 characters
+		const bool bIsWorldObject = SubPathString.Contains(TEXT("PersistentLevel."), ESearchCase::CaseSensitive);
+		if (!bIsWorldObject)
+		{
+			// Not a path to a world object
+			return {};
+		}
+
+		// Start search after the . behind PersistentLevel
+		const int32 StartSearch = PersistentLevelStringLength + 1;
+		const int32 IndexOfDotAfterActorName = SubPathString.Find(TEXT("."), ESearchCase::CaseSensitive, ESearchDir::FromStart, StartSearch);
+		return IndexOfDotAfterActorName == INDEX_NONE;
+	}
+	
+	TOptional<FSoftObjectPath> GetActorOf(const FSoftObjectPath& SoftObjectPath)
+	{
+		// Example of an actor called floor
+		// SoftObjectPath = { AssetPath = {PackageName = "/Game/Maps/SyncBoxLevel", AssetName = "SyncBoxLevel"}, SubPathString = "PersistentLevel.Floor" } }
+		const FString& SubPathString = SoftObjectPath.GetSubPathString();
+
+		constexpr int32 PersistentLevelStringLength = 16; // "PersistentLevel." has 16 characters
+		const bool bIsWorldObject = SubPathString.Contains(TEXT("PersistentLevel."), ESearchCase::CaseSensitive);
+		if (!bIsWorldObject)
+		{
+			// Not a path to a world object
+			return {};
+		}
+
+		// Start search after the . behind PersistentLevel
+		const int32 StartSearch = PersistentLevelStringLength + 1;
+		const int32 IndexOfDotAfterActorName = SubPathString.Find(TEXT("."), ESearchCase::CaseSensitive, ESearchDir::FromStart, StartSearch);
+		if (IndexOfDotAfterActorName == INDEX_NONE)
+		{
+			// SoftObjectPath points to an actor
+			return {};
+		}
+
+		const int32 NumToChopOffRight = SubPathString.Len() - IndexOfDotAfterActorName;
+		const FString NewSubstring = SubPathString.LeftChop(NumToChopOffRight);
+		const FSoftObjectPath PathToOwningActor(SoftObjectPath.GetAssetPath(), NewSubstring);
+		return PathToOwningActor;
+	}
+};
+
diff --git a/Engine/Plugins/Developer/Concert/ConcertUI/ConcertClientSharedSlate/Source/ConcertClientSharedSlate/Private/Replication/Editor/View/ObjectUtils.h b/Engine/Plugins/Developer/Concert/ConcertUI/ConcertClientSharedSlate/Source/ConcertClientSharedSlate/Private/Replication/Editor/View/ObjectUtils.h
new file mode 100644
--- /dev/null
+++ b/Engine/Plugins/Developer/Concert/ConcertUI/ConcertClientSharedSlate/Source/ConcertClientSharedSlate/Private/Replication/Editor/View/ObjectUtils.h
@@ -0,0 +1,16 @@
+// Copyright Epic Games, Inc. All Rights Reserved.
+
+#pragma once
+
+#include "Misc/Optional.h"
+
+struct FSoftObjectPath;
+
+namespace UE::ConcertClientSharedSlate::ObjectUtils
+{
+	/** @return Checks whether the given object is an actor. */
+	bool IsActor(const FSoftObjectPath& Object);
+	
+	/** @return Get the owning actor of Subobject. If Subobject is an actor, then this returns unset. */
+	TOptional<FSoftObjectPath> GetActorOf(const FSoftObjectPath& Subobject);
+};
diff --git a/Engine/Plugins/Developer/Concert/ConcertUI/ConcertClientSharedSlate/Source/ConcertClientSharedSlate/Private/Replication/Editor/View/ObjectViewer/SReplicationStreamViewer.cpp b/Engine/Plugins/Developer/Concert/ConcertUI/ConcertClientSharedSlate/Source/ConcertClientSharedSlate/Private/Replication/Editor/View/ObjectViewer/SReplicationStreamViewer.cpp
--- a/Engine/Plugins/Developer/Concert/ConcertUI/ConcertClientSharedSlate/Source/ConcertClientSharedSlate/Private/Replication/Editor/View/ObjectViewer/SReplicationStreamViewer.cpp
+++ b/Engine/Plugins/Developer/Concert/ConcertUI/ConcertClientSharedSlate/Source/ConcertClientSharedSlate/Private/Replication/Editor/View/ObjectViewer/SReplicationStreamViewer.cpp
@@ -4,14 +4,15 @@
 
 #include "ConcertFrontendUtils.h"
 #include "Replication/Editor/Model/IObjectToPropertiesModel.h"
+#include "Replication/Editor/Model/ISubobjectModel.h"
 #include "Replication/Editor/Model/ReplicatedObjectData.h"
 #include "Replication/Editor/View/ObjectViewer/Property/SReplicatedPropertiesView.h"
 #include "Replication/Editor/View/ObjectViewer/Tree/SelectionViewerColumns.h"
 #include "SSubobjectAndPropertySection.h"
+#include "Replication/Editor/View/ObjectUtils.h"
 
 #include "Widgets/Layout/SBorder.h"
 #include "Widgets/Layout/SSplitter.h"
-#include "Widgets/Layout/SWidgetSwitcher.h"
 #include "Widgets/SBoxPanel.h"
 #include "Widgets/Text/STextBlock.h"
 
@@ -22,6 +23,7 @@ namespace UE::ConcertClientSharedSlate
 	void SReplicationStreamViewer::Construct(const FArguments& InArgs, TSharedRef<IObjectToPropertiesModel> InPropertiesModel)
 	{
 		PropertiesModel = MoveTemp(InPropertiesModel);
+		SubobjectModel = InArgs._SubobjectModel;
 		
 		ChildSlot
 		[
@@ -73,9 +75,12 @@ namespace UE::ConcertClientSharedSlate
 		PropertiesModel->ForEachReplicatedObject([this, &NewPathToObjectDataCache](const FSoftObjectPath& Path) mutable
 		{
 			const TSharedPtr<FReplicatedObjectData>* ExistingItem = PathToObjectDataCache.Find(Path);
+			ExistingItem = ExistingItem ? ExistingItem : NewPathToObjectDataCache.Find(Path);
 			const TSharedRef<FReplicatedObjectData> Item = ExistingItem ? ExistingItem->ToSharedRef() : AllocateObjectData(Path);
-			AllObjectRowData.Emplace(Item);
+			AllObjectRowData.AddUnique(Item);
 			NewPathToObjectDataCache.Emplace(Path, Item);
+			
+			BuildObjectHierarchyIfNeeded(Item, NewPathToObjectDataCache);
 			return EBreakBehavior::Continue;
 		});
 
@@ -119,6 +124,41 @@ namespace UE::ConcertClientSharedSlate
 		}
 	}
 
+	void SReplicationStreamViewer::ExpandObjects(TConstArrayView<FSoftObjectPath> Objects, bool bRecursive)
+	{
+		if (Objects.IsEmpty())
+		{
+			return;
+		}
+		
+		TArray<TSharedPtr<FReplicatedObjectData>> ItemsToExpand;
+		ItemsToExpand.Reserve(Objects.Num());
+		for (const FSoftObjectPath& Path : Objects)
+		{
+			if (const TSharedPtr<FReplicatedObjectData>* Item = PathToObjectDataCache.Find(Path))
+			{
+				ItemsToExpand.Add(*Item);
+			}
+			
+			if (bRecursive && SubobjectModel)
+			{
+				SubobjectModel->ForEachSubobject([this, &ItemsToExpand](const FSoftObjectPath&, const FSoftObjectPath& ChildObject)
+				{
+					if (const TSharedPtr<FReplicatedObjectData>* Item = PathToObjectDataCache.Find(ChildObject))
+					{
+						ItemsToExpand.Add(*Item);
+					}
+					return EBreakBehavior::Continue;
+				});
+			}
+		}
+
+		if (ItemsToExpand.IsEmpty())
+		{
+			ReplicatedObjects->SetSelectedItems(ItemsToExpand, true);
+		}
+	}
+
 	void SReplicationStreamViewer::ClearSubobjectSelection()
 	{
 		SubobjectAndPropertySection->ClearSubobjectSelection();
@@ -152,7 +192,7 @@ namespace UE::ConcertClientSharedSlate
 	{
 		TArray Columns
 		{
-			ReplicationColumns::TopLevel::LabelColumn(PropertiesModel.ToSharedRef()),
+			ReplicationColumns::TopLevel::LabelColumn(PropertiesModel.ToSharedRef(), SubobjectModel.Get()),
 			ReplicationColumns::TopLevel::TypeColumn(PropertiesModel.ToSharedRef())
 		};
 		Columns.Append(InArgs._AdditionalObjectColumns);
@@ -241,11 +281,83 @@ namespace UE::ConcertClientSharedSlate
 		});
 	}
 
+	void SReplicationStreamViewer::BuildObjectHierarchyIfNeeded(TSharedPtr<FReplicatedObjectData> ReplicatedObjectData, TMap<FSoftObjectPath, TSharedPtr<FReplicatedObjectData>>& NewPathToObjectDataCache)
+	{
+		// We're are supposed to display any hierarchy in the outliner if SubobjectModel is not set.
+		const FSoftObjectPath& ObjectPath = ReplicatedObjectData->GetObjectPath();
+		if (!SubobjectModel)
+		{
+			return;
+		}
+
+		// Find top level object of ReplicatedObjectData
+		const TOptional<FSoftObjectPath> OwningActor = ObjectUtils::GetActorOf(ObjectPath);
+		if (!SubobjectModel->IsTopLevelObject(ObjectPath) && !OwningActor)
+		{
+			return;
+		}
+		const FSoftObjectPath TopLevelObject = OwningActor.Get(ObjectPath);
+		SubobjectModel->SetTopLevelObject(TopLevelObject);
+		
+		const auto AddItem = [this, &NewPathToObjectDataCache](const FSoftObjectPath& ObjectPath)
+		{
+			const TSharedPtr<FReplicatedObjectData>* ExistingItem = PathToObjectDataCache.Find(ObjectPath);
+			ExistingItem = ExistingItem ? ExistingItem : NewPathToObjectDataCache.Find(ObjectPath);
+			const TSharedRef<FReplicatedObjectData> Item = ExistingItem ? ExistingItem->ToSharedRef() : AllocateObjectData(ObjectPath);
+			AllObjectRowData.AddUnique(Item);
+			NewPathToObjectDataCache.Emplace(ObjectPath, Item);
+		};
+		// Add all objects that appear in the hierarchy of ReplicatedObjectData
+		AddItem(TopLevelObject);
+		SubobjectModel->ForEachSubobject([this, &AddItem](const FSoftObjectPath&, const FSoftObjectPath& ChildObject)
+		{
+			 AddItem(ChildObject);
+			return EBreakBehavior::Continue;
+		});
+	}
+
 	void SReplicationStreamViewer::GetObjectRowChildren(TSharedPtr<FReplicatedObjectData> ReplicatedObjectData, TFunctionRef<void(TSharedPtr<FReplicatedObjectData>)> ProcessChild)
 	{
-		// For now there are no children in the outliner.
-		// In the future we could add an external delegate that allows child actors to be parented.
 		// Important: this view should be possible to be built in programs, so it should not reference things like AActor, UActorComponent, ResolveObject, etc. directly.
+		
+		const FSoftObjectPath& SearchedObject = ReplicatedObjectData->GetObjectPath();
+		if (!SubobjectModel)
+		{
+			return;
+		}
+
+		const TOptional<FSoftObjectPath> OwningActor = ObjectUtils::GetActorOf(SearchedObject);
+		if (!SubobjectModel->IsTopLevelObject(SearchedObject) && !OwningActor)
+		{
+			return;
+		}
+		SubobjectModel->SetTopLevelObject(OwningActor.Get(SearchedObject));
+		
+		if (SubobjectModel->IsTopLevelObject(SearchedObject))
+		{
+			for (FName Category : SubobjectModel->GetCategories())
+			{
+				SubobjectModel->ForEachRootSubobject(Category, [this, &ProcessChild](const FSoftObjectPath& ChildObject)
+				{
+					if (const TSharedPtr<FReplicatedObjectData>* ObjectData = PathToObjectDataCache.Find(ChildObject))
+					{
+						ProcessChild(*ObjectData);
+					}
+					return EBreakBehavior::Continue;
+				});
+			}
+		}
+		else
+		{
+			SubobjectModel->ForEachDirectChildSubobject(SearchedObject, [this, &ProcessChild](const FSoftObjectPath& ChildObject)
+			{
+				if (const TSharedPtr<FReplicatedObjectData>* ObjectData = PathToObjectDataCache.Find(ChildObject))
+				{
+					ProcessChild(*ObjectData);
+				}
+				return EBreakBehavior::Continue;
+			});
+		}
 	}
 }
 
diff --git a/Engine/Plugins/Developer/Concert/ConcertUI/ConcertClientSharedSlate/Source/ConcertClientSharedSlate/Private/Replication/Editor/View/ObjectViewer/SReplicationStreamViewer.h b/Engine/Plugins/Developer/Concert/ConcertUI/ConcertClientSharedSlate/Source/ConcertClientSharedSlate/Private/Replication/Editor/View/ObjectViewer/SReplicationStreamViewer.h
--- a/Engine/Plugins/Developer/Concert/ConcertUI/ConcertClientSharedSlate/Source/ConcertClientSharedSlate/Private/Replication/Editor/View/ObjectViewer/SReplicationStreamViewer.h
+++ b/Engine/Plugins/Developer/Concert/ConcertUI/ConcertClientSharedSlate/Source/ConcertClientSharedSlate/Private/Replication/Editor/View/ObjectViewer/SReplicationStreamViewer.h
@@ -27,6 +27,7 @@ namespace UE::ConcertClientSharedSlate
 	class IEditableObjectToPropertiesModel;
 	class IReplicationSubobjectView;
 	class IObjectToPropertiesModel;
+	class ISubobjectModel;
 	class SReplicatedPropertiesView;
 	
 	/**
@@ -52,6 +53,8 @@ namespace UE::ConcertClientSharedSlate
 
 			/** Optional. Placed between root object outliner and property editor. */
 			SLATE_ARGUMENT(TSharedPtr<IReplicationSubobjectView>, SubobjectView)
+			/** Optional. If set, this determines the children nested under the root objects. */
+			SLATE_ARGUMENT(TSharedPtr<ISubobjectModel>, SubobjectModel)
 
 			/** Optional. Called when the delete key is pressed in the object view. */
 			SLATE_EVENT(SReplicationTreeView<FReplicatedObjectData>::FDeleteItems, OnDeleteObjects)
@@ -85,6 +88,9 @@ namespace UE::ConcertClientSharedSlate
 
 		/** Selects the given objects from the top level view, if applicable. */
 		void SelectTopLevelObjects(TConstArrayView<FSoftObjectPath> Objects);
+
+		/** Expands the given objects, recursively if desired. */
+		void ExpandObjects(TConstArrayView<FSoftObjectPath> Objects, bool bRecursive);
 		
 		/** Clears all objects selected in the subobject view, if there is one. */
 		void ClearSubobjectSelection();
@@ -96,6 +102,8 @@ namespace UE::ConcertClientSharedSlate
 
 		/** The model this view is visualizing. */
 		TSharedPtr<IObjectToPropertiesModel> PropertiesModel;
+		/** Can be null. If set, this determines the children nested under the root objects. */
+		TSharedPtr<ISubobjectModel> SubobjectModel;
 
 		/** Lists the properties of the selected actor */
 		TSharedPtr<SExpandableArea> PropertyArea;
@@ -114,7 +122,7 @@ namespace UE::ConcertClientSharedSlate
 
 		bool bIsPropertyAreaExpanded = false;
 
-		TSharedRef<FReplicatedObjectData> AllocateObjectData(FSoftObjectPath ObjectPath);
+		static TSharedRef<FReplicatedObjectData> AllocateObjectData(FSoftObjectPath ObjectPath);
 
 		// Widget creation helpers
 		TSharedRef<SWidget> CreateContentWidget(const FArguments& InArgs);
@@ -123,7 +131,9 @@ namespace UE::ConcertClientSharedSlate
 
 		/** Sets RootObjectRowData to all non-root nodes from ObjectRowData. */
 		void BuildRootObjectRowData();
-		
+
+		/** Creates an item for every object in the hierarchy of ReplicatedObjectData */
+		void BuildObjectHierarchyIfNeeded(TSharedPtr<FReplicatedObjectData> ReplicatedObjectData, TMap<FSoftObjectPath, TSharedPtr<FReplicatedObjectData>>& NewPathToObjectDataCache);
 		void GetObjectRowChildren(TSharedPtr<FReplicatedObjectData> ReplicatedObjectData, TFunctionRef<void(TSharedPtr<FReplicatedObjectData>)> ProcessChild);
 		
 		/** Handles how much space the 'Clients' area uses with respect to its expansion state. */
diff --git a/Engine/Plugins/Developer/Concert/ConcertUI/ConcertClientSharedSlate/Source/ConcertClientSharedSlate/Private/Replication/Editor/View/ObjectViewer/Tree/SelectionViewerColumns.cpp b/Engine/Plugins/Developer/Concert/ConcertUI/ConcertClientSharedSlate/Source/ConcertClientSharedSlate/Private/Replication/Editor/View/ObjectViewer/Tree/SelectionViewerColumns.cpp
--- a/Engine/Plugins/Developer/Concert/ConcertUI/ConcertClientSharedSlate/Source/ConcertClientSharedSlate/Private/Replication/Editor/View/ObjectViewer/Tree/SelectionViewerColumns.cpp
+++ b/Engine/Plugins/Developer/Concert/ConcertUI/ConcertClientSharedSlate/Source/ConcertClientSharedSlate/Private/Replication/Editor/View/ObjectViewer/Tree/SelectionViewerColumns.cpp
@@ -4,9 +4,9 @@
 
 #include "ClassIconFinder.h"
 #include "Replication/Editor/Model/IEditableObjectToPropertiesModel.h"
+#include "Replication/Editor/Model/ISubobjectModel.h"
 #include "Replication/Editor/Model/ReplicatedPropertyData.h"
 #include "Replication/Editor/Model/ReplicatedObjectData.h"
-#include "Replication/Editor/Model/Subobject/ISubobjectModel.h"
 #include "Replication/Editor/View/DisplayUtils.h"
 #include "Replication/Editor/View/ObjectEditor/SBaseReplicationStreamEditor.h"
 #include "Replication/Editor/View/ReplicationColumnsUtils.h"
@@ -14,6 +14,7 @@
 
 #include "Internationalization/Internationalization.h"
 #include "GameFramework/Actor.h"
+#include "Replication/Editor/View/ObjectUtils.h"
 #include "Textures/SlateIcon.h"
 #include "Widgets/Images/SImage.h"
 #include "Widgets/SBoxPanel.h"
@@ -28,12 +29,23 @@ namespace UE::ConcertClientSharedSlate::ReplicationColumns::TopLevel
 	const FName LabelColumnId = TEXT("LabelColumn");
 	const FName TypeColumnId = TEXT("TypeColumn");
 	
-	FReplicationTopLevelObjectColumn LabelColumn(TSharedRef<IObjectToPropertiesModel> Model)
+	FReplicationTopLevelObjectColumn LabelColumn(TSharedRef<IObjectToPropertiesModel> Model, ISubobjectModel* SubobjectModel)
 	{
 		return FReplicationTopLevelObjectColumn(
 			FReplicationTopLevelObjectColumn::FArguments()
-				.GenerateWidgetColumn_Lambda([Model = MoveTemp(Model)](const FReplicationTopLevelObjectColumn::FBuildArgs& Args)
+				.GenerateWidgetColumn_Lambda([Model = MoveTemp(Model), SubobjectModel](const FReplicationTopLevelObjectColumn::FBuildArgs& Args)
 				{
+					const FSoftObjectPath& ObjectPath = Args.RowData.GetObjectPath();
+					FText Text = DisplayUtils::GetObjectDisplayText(ObjectPath);
+					if (SubobjectModel && !SubobjectModel->IsTopLevelObject(ObjectPath))
+					{
+						if (const TOptional<FSoftObjectPath> OwningActorPath = ObjectUtils::GetActorOf(ObjectPath))
+						{
+							SubobjectModel->SetTopLevelObject(*OwningActorPath);
+							Text = SubobjectModel->GetSubobjectDisplayName(ObjectPath);
+						}
+					}
+					
 					return SNew(SHorizontalBox)
 						+SHorizontalBox::Slot()
 						.AutoWidth()
@@ -51,7 +63,7 @@ namespace UE::ConcertClientSharedSlate::ReplicationColumns::TopLevel
 						[
 							SNew(STextBlock)
 							.HighlightText(TAttribute<FText>::CreateLambda([HighlightText = Args.HighlightText](){ return *HighlightText; }))
-							.Text(DisplayUtils::GetObjectDisplayText(Args.RowData.GetObjectPath()))
+							.Text(Text)
 						];
 				})
 				.PopulateSearchItems_Lambda([](const FReplicatedObjectData& ObjectData, TArray<FString>& InOutSearchStrings)
diff --git a/Engine/Plugins/Developer/Concert/ConcertUI/ConcertClientSharedSlate/Source/ConcertClientSharedSlate/Private/Replication/Editor/View/ObjectViewer/Tree/SelectionViewerColumns.h b/Engine/Plugins/Developer/Concert/ConcertUI/ConcertClientSharedSlate/Source/ConcertClientSharedSlate/Private/Replication/Editor/View/ObjectViewer/Tree/SelectionViewerColumns.h
--- a/Engine/Plugins/Developer/Concert/ConcertUI/ConcertClientSharedSlate/Source/ConcertClientSharedSlate/Private/Replication/Editor/View/ObjectViewer/Tree/SelectionViewerColumns.h
+++ b/Engine/Plugins/Developer/Concert/ConcertUI/ConcertClientSharedSlate/Source/ConcertClientSharedSlate/Private/Replication/Editor/View/ObjectViewer/Tree/SelectionViewerColumns.h
@@ -28,7 +28,7 @@ namespace UE::ConcertClientSharedSlate::ReplicationColumns::TopLevel
 		Type = 30,
 	};
 
-	FReplicationTopLevelObjectColumn LabelColumn(TSharedRef<IObjectToPropertiesModel> Model);
+	FReplicationTopLevelObjectColumn LabelColumn(TSharedRef<IObjectToPropertiesModel> Model, ISubobjectModel* SubobjectModel = nullptr);
 	FReplicationTopLevelObjectColumn TypeColumn(TSharedRef<IObjectToPropertiesModel> Model);
 }
 
diff --git a/Engine/Plugins/Developer/Concert/ConcertUI/ConcertClientSharedSlate/Source/ConcertClientSharedSlate/Private/Replication/Editor/View/SubobjectView/SSubobjectView.cpp b/Engine/Plugins/Developer/Concert/ConcertUI/ConcertClientSharedSlate/Source/ConcertClientSharedSlate/Private/Replication/Editor/View/SubobjectView/SSubobjectView.cpp
--- a/Engine/Plugins/Developer/Concert/ConcertUI/ConcertClientSharedSlate/Source/ConcertClientSharedSlate/Private/Replication/Editor/View/SubobjectView/SSubobjectView.cpp
+++ b/Engine/Plugins/Developer/Concert/ConcertUI/ConcertClientSharedSlate/Source/ConcertClientSharedSlate/Private/Replication/Editor/View/SubobjectView/SSubobjectView.cpp
@@ -4,7 +4,7 @@
 
 #include "Misc/EBreakBehavior.h"
 #include "ReplicatedSubobjectData.h"
-#include "Replication/Editor/Model/Subobject/ISubobjectModel.h"
+#include "Replication/Editor/Model/ISubobjectModel.h"
 #include "Replication/Editor/View/ObjectViewer/Tree/SelectionViewerColumns.h"
 
 #include "Containers/Queue.h"
diff --git a/Engine/Plugins/Developer/Concert/ConcertUI/ConcertClientSharedSlate/Source/ConcertClientSharedSlate/Private/Replication/ReplicationWidgetFactories.cpp b/Engine/Plugins/Developer/Concert/ConcertUI/ConcertClientSharedSlate/Source/ConcertClientSharedSlate/Private/Replication/ReplicationWidgetFactories.cpp
--- a/Engine/Plugins/Developer/Concert/ConcertUI/ConcertClientSharedSlate/Source/ConcertClientSharedSlate/Private/Replication/ReplicationWidgetFactories.cpp
+++ b/Engine/Plugins/Developer/Concert/ConcertUI/ConcertClientSharedSlate/Source/ConcertClientSharedSlate/Private/Replication/ReplicationWidgetFactories.cpp
@@ -12,16 +12,22 @@ namespace UE::ConcertClientSharedSlate
 {
 	TSharedRef<IReplicationSubobjectView> CreateUnrealEditorSubobjectView(FCreateSubobjectViewParams Params)
 	{
-		return SNew(SSubobjectView, MakeShared<FComponentHierarchySubobjectModel>())
+		return SNew(SSubobjectView, MoveTemp(Params.SubobjectModel))
 			.AdditionalColumns(Params.AdditionalColumns);
 	}
 	
+	TSharedRef<ISubobjectModel> CreateDefaultComponentHierarchySubobjectModel()
+	{
+		return MakeShared<FComponentHierarchySubobjectModel>();
+	}
+	
 	TSharedRef<IReplicationStreamEditor> CreateEditor(FCreateEditorParams Params)
 	{
 		return SNew(SBaseReplicationStreamEditor, Params.DataModel, Params.ObjectSource, Params.PropertySource)
 			.AdditionalObjectColumns(Params.AdditionalObjectColumns)
 			.AdditionalPropertyColumns(Params.AdditionalPropertyColumns)
 			.SubobjectView(Params.SubobjectView)
+			.SubobjectModel(Params.SubobjectModel)
 			.OnExtendObjectsContextMenu(Params.OnExtendObjectsContextMenu)
 			.SortPropertyRowPredicate(Params.SortPropertyRowPredicate)
 			.LeftOfObjectSearchBar() [ Params.LeftOfObjectSearchBar.Widget ]
@@ -33,13 +39,11 @@ namespace UE::ConcertClientSharedSlate
 
 	TSharedRef<IReplicationStreamEditor> CreateDefaultStreamEditor(FCreateEditorParams Params)
 	{
-		// Default editor relies on a subobject view
-		Params.SubobjectView = Params.SubobjectView ? Params.SubobjectView : CreateUnrealEditorSubobjectView(FCreateSubobjectViewParams{});
-		
 		return SNew(SDefaultReplicationStreamEditor, Params.DataModel, Params.ObjectSource, Params.PropertySource)
 			.AdditionalObjectColumns(Params.AdditionalObjectColumns)
 			.AdditionalPropertyColumns(Params.AdditionalPropertyColumns)
 			.SubobjectView(Params.SubobjectView)
+			.SubobjectModel(Params.SubobjectModel)
 			.OnExtendObjectsContextMenu(Params.OnExtendObjectsContextMenu)
 			.SortPropertyRowPredicate(Params.SortPropertyRowPredicate)
 			.LeftOfObjectSearchBar() [ Params.LeftOfObjectSearchBar.Widget ]
diff --git a/Engine/Plugins/Developer/Concert/ConcertUI/ConcertClientSharedSlate/Source/ConcertClientSharedSlate/Public/Replication/Editor/Model/ISubobjectModel.h b/Engine/Plugins/Developer/Concert/ConcertUI/ConcertClientSharedSlate/Source/ConcertClientSharedSlate/Public/Replication/Editor/Model/ISubobjectModel.h
new file mode 100644
--- /dev/null
+++ b/Engine/Plugins/Developer/Concert/ConcertUI/ConcertClientSharedSlate/Source/ConcertClientSharedSlate/Public/Replication/Editor/Model/ISubobjectModel.h
@@ -0,0 +1,57 @@
+// Copyright Epic Games, Inc. All Rights Reserved.
+
+#pragma once
+
+#include "Delegates/Delegate.h"
+#include "Misc/EBreakBehavior.h"
+#include "Templates/Function.h"
+
+enum class EBreakBehavior : uint8;
+struct FSoftObjectPath;
+
+namespace UE::ConcertClientSharedSlate
+{
+	/**
+	 * Decides on the contents for the subobject tree view.
+	 */
+	class CONCERTCLIENTSHAREDSLATE_API ISubobjectModel
+	{
+	public:
+
+		/** Sets the object for which to build the subobject hierarchy. */
+		virtual void SetTopLevelObject(const FSoftObjectPath& TopLevelObject) = 0;
+		virtual FSoftObjectPath GetTopLevelObject() const = 0;
+
+		/** @return Whether Object is a top-level object, i.e. valid to pass to SetTopLevelObject. */
+		virtual bool IsTopLevelObject(const FSoftObjectPath& Object) const = 0;
+
+		/**
+		 * Gets the categories that can be passed to ForEachRootSubobject.
+		 * Objects are grouped by category. The categories are separated by separator lines.
+		 * Example: One category could be the the USceneComponent hierarchy and another could be all other UActorComponents (like in the SSubobjectEditor).
+		 */
+		virtual TArray<FName> GetCategories() const = 0;
+
+		/** Gets the display name for the subobject */
+		virtual FText GetSubobjectDisplayName(const FSoftObjectPath& ObjectPath) const = 0;
+		
+		/**
+		 * Gets the direct subobjects of the top level objects.
+		 * @param Category The category for the objects.
+		 * @param Callback Callback to invoke for each found root object.
+		 * @see GetCategories for valid categories.
+		 */
+		virtual void ForEachRootSubobject(FName Category, TFunctionRef<EBreakBehavior(const FSoftObjectPath& Object)> Callback) const = 0;
+		/** Gets the direct subobject children of another subobject. Child subobjects are implicitly in the same category as Parent. */
+		virtual void ForEachDirectChildSubobject(const FSoftObjectPath& Parent, TFunctionRef<EBreakBehavior(const FSoftObjectPath& Object)> Callback) const = 0;
+
+		DECLARE_MULTICAST_DELEGATE(FOnHierarchyChanged)
+		/** Called when the hierarchy has changed, e.g. due to calling SetTopLevelObject. */
+		virtual FOnHierarchyChanged& OnHierarchyChanged() = 0;
+
+		/** Util for iterating all subobjects. */
+		void ForEachSubobject(TFunctionRef<EBreakBehavior(const FSoftObjectPath& Parent, const FSoftObjectPath& ChildObject)> Callback);
+
+		virtual ~ISubobjectModel() = default;
+	};
+}
diff --git a/Engine/Plugins/Developer/Concert/ConcertUI/ConcertClientSharedSlate/Source/ConcertClientSharedSlate/Public/Replication/Editor/Model/Object/ActorSelectionSourceModel.h b/Engine/Plugins/Developer/Concert/ConcertUI/ConcertClientSharedSlate/Source/ConcertClientSharedSlate/Public/Replication/Editor/Model/Object/ActorSelectionSourceModel.h
new file mode 100644
--- /dev/null
+++ b/Engine/Plugins/Developer/Concert/ConcertUI/ConcertClientSharedSlate/Source/ConcertClientSharedSlate/Public/Replication/Editor/Model/Object/ActorSelectionSourceModel.h
@@ -0,0 +1,27 @@
+// Copyright Epic Games, Inc. All Rights Reserved.
+
+#pragma once
+
+#include "Replication/Editor/Model/Object/IObjectSelectionSourceModel.h"
+
+namespace UE::ConcertClientSharedSlate
+{
+	/** Logic for selecting actors from the current editor world. */
+	class CONCERTCLIENTSHAREDSLATE_API FActorSelectionSourceModel : public IObjectSelectionSourceModel
+	{
+	public:
+
+		FActorSelectionSourceModel();
+
+		//~ Begin IObjectSelectionSourceModel Interface
+		virtual TArray<FObjectSourceCategory> GetRootSources() const override;
+		virtual TArray<TSharedRef<IObjectSourceModel>> GetContextMenuOptions(const FSoftObjectPath& Item) override;
+		//~ End IObjectSelectionSourceModel Interface
+
+	private:
+
+		/** Contains actions around adding actor related objects. Displayed next to the search bar. Refined with dynamic options in GetRootSources. */
+		FObjectSourceCategory BaseActorCategory;
+	};
+}
+
diff --git a/Engine/Plugins/Developer/Concert/ConcertUI/ConcertClientSharedSlate/Source/ConcertClientSharedSlate/Public/Replication/Editor/Model/Object/EditorObjectSelectionSourceModel.h b/Engine/Plugins/Developer/Concert/ConcertUI/ConcertClientSharedSlate/Source/ConcertClientSharedSlate/Public/Replication/Editor/Model/Object/EditorObjectSelectionSourceModel.h
deleted file mode 100644
diff --git a/Engine/Plugins/Developer/Concert/ConcertUI/ConcertClientSharedSlate/Source/ConcertClientSharedSlate/Public/Replication/Editor/Model/Object/IObjectSelectionSourceModel.h b/Engine/Plugins/Developer/Concert/ConcertUI/ConcertClientSharedSlate/Source/ConcertClientSharedSlate/Public/Replication/Editor/Model/Object/IObjectSelectionSourceModel.h
--- a/Engine/Plugins/Developer/Concert/ConcertUI/ConcertClientSharedSlate/Source/ConcertClientSharedSlate/Public/Replication/Editor/Model/Object/IObjectSelectionSourceModel.h
+++ b/Engine/Plugins/Developer/Concert/ConcertUI/ConcertClientSharedSlate/Source/ConcertClientSharedSlate/Public/Replication/Editor/Model/Object/IObjectSelectionSourceModel.h
@@ -31,14 +31,6 @@ namespace UE::ConcertClientSharedSlate
 
 		/** @return The sources that should be displayed in the right-click context menu when clicked in the menu. */
 		virtual TArray<TSharedRef<IObjectSourceModel>> GetContextMenuOptions(const FSoftObjectPath& Items) = 0;
-
-		/**
-		 * Checks whether the selected item is valid to be added to the IEditableObjectToPropertiesModel.
-		 *
-		 * This can be used to validate the asset when it is initially displayed but could also be useful for validating
-		 * actors / assets that were drag-dropped into the editor.
-		 */
-		virtual EObjectItemValidity GetItemValidity(const FSoftObjectPath& Item) const = 0;
 		
 		virtual ~IObjectSelectionSourceModel() = default;
 	};
diff --git a/Engine/Plugins/Developer/Concert/ConcertUI/ConcertClientSharedSlate/Source/ConcertClientSharedSlate/Public/Replication/ReplicationWidgetFactories.h b/Engine/Plugins/Developer/Concert/ConcertUI/ConcertClientSharedSlate/Source/ConcertClientSharedSlate/Public/Replication/ReplicationWidgetFactories.h
--- a/Engine/Plugins/Developer/Concert/ConcertUI/ConcertClientSharedSlate/Source/ConcertClientSharedSlate/Public/Replication/ReplicationWidgetFactories.h
+++ b/Engine/Plugins/Developer/Concert/ConcertUI/ConcertClientSharedSlate/Source/ConcertClientSharedSlate/Public/Replication/ReplicationWidgetFactories.h
@@ -22,6 +22,7 @@ namespace UE::ConcertClientSharedSlate
 	class IReplicationStreamEditor;
 	class IReplicationStreamViewer;
 	class IReplicationSubobjectView;
+	class ISubobjectModel;
 	class IPropertySelectionSourceModel;
 	
 	/**
@@ -39,11 +40,17 @@ namespace UE::ConcertClientSharedSlate
 		UObject& OwnerObject,
 		TAttribute<FObjectReplicationMap*> ReplicationMapAttribute
 		);
+
+	/** Builds a similar tree hierarchy as SSubobjectEditor. Reports only components as subobjects. */
+	CONCERTCLIENTSHAREDSLATE_API TSharedRef<ISubobjectModel> CreateDefaultComponentHierarchySubobjectModel();
 	
 	struct FCreateSubobjectViewParams
 	{
 		/** Additional columns that should be displayed for the subobjects. */
 		TArray<ReplicationColumns::FReplicationSubobjectObjectColumn> AdditionalColumns;
+
+		/** Determines the subobjects displayed in the view. */
+		TSharedRef<ISubobjectModel> SubobjectModel = CreateDefaultComponentHierarchySubobjectModel();
 	};
 	/** Creates a stream editor with a subobject view that looks like the SSubobjectEditor. */
 	CONCERTCLIENTSHAREDSLATE_API TSharedRef<IReplicationSubobjectView> CreateUnrealEditorSubobjectView(
@@ -74,10 +81,17 @@ namespace UE::ConcertClientSharedSlate
 		 * It e.g. displays the components of the actor selected in the root object view.
 		 * 
 		 * Exists so it can be customized differently depending on whether used in the editor or on the server.
-		 * @note The view will keep a strong reference to this.
+		 * @note The created view will keep a strong reference to this.
 		 */
 		TSharedPtr<IReplicationSubobjectView> SubobjectView;
 
+		/**
+		 * Optional. Determines the objects displayed as children to the top-level objects in the top section.
+		 * If left unspecified, the top will only display actors.
+		 * @note The created view will keep a strong reference to this.
+		 */
+		TSharedPtr<ISubobjectModel> SubobjectModel;
+
 		/** Called to generate the context menu for objects. This extends the options already generated by this widget and is called at the end. */
 		FExtendMenu OnExtendObjectsContextMenu;
 		/** Optional. Used for determining the order in which properties are displayed. */
@@ -109,7 +123,7 @@ namespace UE::ConcertClientSharedSlate
 	 * You are expected to inject widgets or columns, which implement editing properties.
 	 * @see CreateDefaultStreamEditor, which adds a checkbox to the start of every property row.
 	 *
-	 * The editor consists of three areas, which share similar workflows the level editor:
+	 * The editor consists of two or three areas (depending whether on whether SubobjectView was specified), which share similar workflows the level editor:
 	 * 1. Root objects, similar to world outliner:
 	 *		- Similar to World Outliner: Displays top-level objects, i.e. actors, are added here.
 	 *		- FCreateEditorParams::ObjectSource is used to build a combo button through which new objects can be added.
