Add all clients view for replication

diff --git a/Engine/Plugins/Developer/Concert/ConcertApp/MultiUserClient/Source/MultiUserClient/Private/Replication/Client/RemoteReplicationClient.cpp b/Engine/Plugins/Developer/Concert/ConcertApp/MultiUserClient/Source/MultiUserClient/Private/Replication/Client/RemoteReplicationClient.cpp
--- a/Engine/Plugins/Developer/Concert/ConcertApp/MultiUserClient/Source/MultiUserClient/Private/Replication/Client/RemoteReplicationClient.cpp
+++ b/Engine/Plugins/Developer/Concert/ConcertApp/MultiUserClient/Source/MultiUserClient/Private/Replication/Client/RemoteReplicationClient.cpp
@@ -36,7 +36,6 @@ namespace UE::MultiUserClient
 		GetStreamSynchronizer().OnServerStateChanged().AddLambda([this]()
 		{
 			GetClientContent()->Stream->ReplicationMap = GetStreamSynchronizer().GetServerState();
-			OnModelExternallyChanged().Broadcast();
 		});
 	}
 }
diff --git a/Engine/Plugins/Developer/Concert/ConcertApp/MultiUserClient/Source/MultiUserClient/Private/Replication/Client/ReplicationClient.cpp b/Engine/Plugins/Developer/Concert/ConcertApp/MultiUserClient/Source/MultiUserClient/Private/Replication/Client/ReplicationClient.cpp
--- a/Engine/Plugins/Developer/Concert/ConcertApp/MultiUserClient/Source/MultiUserClient/Private/Replication/Client/ReplicationClient.cpp
+++ b/Engine/Plugins/Developer/Concert/ConcertApp/MultiUserClient/Source/MultiUserClient/Private/Replication/Client/ReplicationClient.cpp
@@ -8,6 +8,8 @@
 #include "Replication/Stream/StreamChangeTracker.h"
 #include "Replication/Submission/Data/AuthoritySubmission.h"
 
+#include "Misc/CoreDelegates.h"
+
 namespace UE::MultiUserClient
 {
 	FReplicationClient::FReplicationClient(
@@ -37,13 +39,19 @@ namespace UE::MultiUserClient
 	{
 		LocalClientEditModel->OnObjectsChanged().AddRaw(this, &FReplicationClient::OnObjectsChanged);
 		LocalClientEditModel->OnPropertiesChanged().AddRaw(this, &FReplicationClient::OnPropertiesChanged);
-		
-		LocalClientStreamDiffer.OnChangesReverted_GameThread().AddLambda([this]()
-		{
-			OnModelExternallyChangedDelegate.Broadcast();
-		});
 
 		SubmissionWorkflow->OnAuthorityRequestCompleted().AddRaw(this, &FReplicationClient::OnAuthoritySubmissionCompleted);
+		StreamSynchronizer->OnServerStateChanged().AddRaw(this, &FReplicationClient::DeferOnModelChanged);
+	}
+
+	FReplicationClient::~FReplicationClient()
+	{
+		FCoreDelegates::OnEndFrame.RemoveAll(this);
+	}
+
+	bool FReplicationClient::AllowsEditing() const
+	{
+		return CanEverSubmit(SubmissionWorkflow->GetUploadability()); 
 	}
 
 	void FReplicationClient::OnObjectsChanged(
@@ -52,25 +60,58 @@ namespace UE::MultiUserClient
 		ConcertClientSharedSlate::EReplicatedObjectChangeReason ReplicatedObjectChangeReason
 		)
 	{
+		DeferOnModelChanged(AddedObjects);
+	}
+
+	void FReplicationClient::OnPropertiesChanged()
+	{
+		DeferOnModelChanged();
+	}
+
+	void FReplicationClient::DeferOnModelChanged(TConstArrayView<UObject*> AddedObjects)
+	{
+		if (!DeferredOnModelChangedData)
+		{
+			DeferredOnModelChangedData.Emplace();
+			FCoreDelegates::OnEndFrame.AddRaw(this, &FReplicationClient::ProcessOnModelChanged);
+		}
+
+		Algo::Transform(AddedObjects, DeferredOnModelChangedData->AccumulatedAddedObjects, [](UObject* Object)
+		{
+			return Object;
+		});
+	}
+
+	void FReplicationClient::ProcessOnModelChanged()
+	{
+		check(DeferredOnModelChangedData);
+		const FDeferredOnModelChangedData ChangeData = MoveTemp(*DeferredOnModelChangedData);
+		DeferredOnModelChangedData.Reset();
+		FCoreDelegates::OnEndFrame.RemoveAll(this);
+		
 		// Could improve performance by just considering what actually changed instead of doing a full rebuild
 		// This must be done before SetAuthorityIfAllowed because it uses the cache for checking whether the object has properties assigned
 		LocalClientStreamDiffer.RefreshChangesCache();
 		
 		// Better UX for user: automatically take authority for newly added objects (but only if it is allowed and causes no conflicts)
 		TArray<FSoftObjectPath> ObjectPaths;
-		Algo::Transform(AddedObjects, ObjectPaths, [](const UObject* Object){ return FSoftObjectPath(Object); });
+		Algo::TransformIf(ChangeData.AccumulatedAddedObjects, ObjectPaths, 
+			[](const TWeakObjectPtr<UObject>& Object)
+			{
+				// The object might have been made invalid last frame.
+				return Object.IsValid();
+			},
+			[](const TWeakObjectPtr<UObject>& Object)
+			{
+				return FSoftObjectPath(Object.Get()) ;
+			});
 		LocalAuthorityDiffer.SetAuthorityIfAllowed(ObjectPaths, true);
 
-		// Refresh because authority changes may no longer be valid after modifying the stream
+		// Refresh because local authority changes may no longer be valid after modifying the stream
 		LocalAuthorityDiffer.RefreshChanges();
-	}
 
-	void FReplicationClient::OnPropertiesChanged()
-	{
-		// Could improve performance by just considering what actually changed instead of doing a full rebuild
-		LocalClientStreamDiffer.RefreshChangesCache();
-		// Refresh because authority changes may no longer be valid after modifying the stream
-		LocalAuthorityDiffer.RefreshChanges();
+		// Finally, let everybody else know.
+		OnModelChangedDelegate.Broadcast();
 	}
 
 	void FReplicationClient::OnAuthoritySubmissionCompleted(const FSubmitAuthorityChangesRequest& Request, const FSubmitAuthorityChangesResponse& Response)
diff --git a/Engine/Plugins/Developer/Concert/ConcertApp/MultiUserClient/Source/MultiUserClient/Private/Replication/Client/ReplicationClient.h b/Engine/Plugins/Developer/Concert/ConcertApp/MultiUserClient/Source/MultiUserClient/Private/Replication/Client/ReplicationClient.h
--- a/Engine/Plugins/Developer/Concert/ConcertApp/MultiUserClient/Source/MultiUserClient/Private/Replication/Client/ReplicationClient.h
+++ b/Engine/Plugins/Developer/Concert/ConcertApp/MultiUserClient/Source/MultiUserClient/Private/Replication/Client/ReplicationClient.h
@@ -43,6 +43,7 @@ namespace UE::MultiUserClient
 			TUniquePtr<IClientAuthoritySynchronizer> InAuthoritySynchronizer,
 			TFunctionRef<FMakeSubmissionWorkflow> MakeSubmissionWorkflowFunc
 			);
+		~FReplicationClient();
 
 		UMultiUserReplicationClientPreset* GetClientContent() const { return ClientContentStorage; }
 		TSharedRef<ConcertClientSharedSlate::IEditableReplicationStreamModel> GetClientEditModel() const { return LocalClientEditModel; }
@@ -58,19 +59,21 @@ namespace UE::MultiUserClient
 		const ISubmissionWorkflow& GetSubmissionWorkflow() const { return *SubmissionWorkflow; }
 		ISubmissionWorkflow& GetSubmissionWorkflow() { return *SubmissionWorkflow; }
 
+		/** @return The endpoint ID of this client in the Concert session. */
 		const FGuid& GetEndpointId() const { return EndpointId; }
+		/** @return Whether it is allowed to edit the stream and authority for this client. */
+		bool AllowsEditing() const;
+		
+		bool operator==(const FReplicationClient& Client) const { return GetEndpointId() == Client.GetEndpointId(); }
 
+		DECLARE_MULTICAST_DELEGATE(FOnModelExternallyChanged);
 		/**
-		 * Called when the data underlying the model has changed externally. Since the change was not caused by the model,
-		 * its events, like IEditableReplicationStreamModel::OnObjectsChanged, were not called.
-		 * 
-		 * Subscribers are intended to call IReplicationStreamEditor::Refresh() in response.
-		 * 
-		 * Examples: Remote client changed their stream, remote client joined with streams,
-		 * or the local client reverted a change that was rejected by the server.
+		 * Broadcast when the data underlying the model has changed for any reason:
+		 * - Edited directly by client
+		 * - Changed by transaction
+		 * - Server state changed
 		 */
-		DECLARE_MULTICAST_DELEGATE(FOnModelExternallyChanged);
-		FOnModelExternallyChanged& OnModelExternallyChanged() { return OnModelExternallyChangedDelegate; }
+		FOnModelExternallyChanged& OnModelChanged() { return OnModelChangedDelegate; }
 		
 	private:
 
@@ -106,12 +109,29 @@ namespace UE::MultiUserClient
 		/** Automatically submits changes as they are made by the user. */
 		FAutoSubmissionPolicy AutoSubmissionPolicy;
 
-		/** Called when the data underlying the model has changed (and the UI needs to be refreshed). */
-		FOnModelExternallyChanged OnModelExternallyChangedDelegate;
+		/**
+		 * Broadcast when the data underlying the model has changed for any reason:
+		 * - Edited directly by client
+		 * - Changed by transaction
+		 * - Server state changed
+		 */
+		FOnModelExternallyChanged OnModelChangedDelegate;
+
+		struct FDeferredOnModelChangedData
+		{
+			TSet<TWeakObjectPtr<UObject>> AccumulatedAddedObjects;
+		};
+		TOptional<FDeferredOnModelChangedData> DeferredOnModelChangedData;
 
 		// Respond to model changing
 		void OnObjectsChanged(TConstArrayView<UObject*> AddedObjects, TConstArrayView<FSoftObjectPath> RemovedObjects, ConcertClientSharedSlate::EReplicatedObjectChangeReason ReplicatedObjectChangeReason);
 		void OnPropertiesChanged();
+
+		/** Defers rebuilding operations, such as refreshing state and calling OnModelChanged, in case there are multiple changes in the same frame. */
+		void DeferOnModelChanged() { DeferOnModelChanged({}); }
+		void DeferOnModelChanged(TConstArrayView<UObject*> AddedObjects);
+		/** Processes all changes that have happened to the stream this frame. */
+		void ProcessOnModelChanged();
 		
 		/** Removes authority if request fails */
 		void OnAuthoritySubmissionCompleted(const FSubmitAuthorityChangesRequest& Request, const FSubmitAuthorityChangesResponse& Response);
diff --git a/Engine/Plugins/Developer/Concert/ConcertApp/MultiUserClient/Source/MultiUserClient/Private/Replication/Client/ReplicationClientManager.cpp b/Engine/Plugins/Developer/Concert/ConcertApp/MultiUserClient/Source/MultiUserClient/Private/Replication/Client/ReplicationClientManager.cpp
--- a/Engine/Plugins/Developer/Concert/ConcertApp/MultiUserClient/Source/MultiUserClient/Private/Replication/Client/ReplicationClientManager.cpp
+++ b/Engine/Plugins/Developer/Concert/ConcertApp/MultiUserClient/Source/MultiUserClient/Private/Replication/Client/ReplicationClientManager.cpp
@@ -64,6 +64,31 @@ namespace UE::MultiUserClient
 		return Client ? Client->Get() : nullptr;
 	}
 
+	void FReplicationClientManager::ForEachClient(TFunctionRef<EBreakBehavior(const FReplicationClient&)> ProcessClient) const
+	{
+		if (ProcessClient(GetLocalClient()) == EBreakBehavior::Break)
+		{
+			return;
+		}
+		for (const TNonNullPtr<FRemoteReplicationClient>& RemoteClient : GetRemoteClients())
+		{
+			if (ProcessClient(*RemoteClient) == EBreakBehavior::Break)
+			{
+				return;
+			}
+		}
+	}
+
+	void FReplicationClientManager::ForEachClient(TFunctionRef<EBreakBehavior(FReplicationClient&)> ProcessClient)
+	{
+		const FReplicationClientManager* ConstThis = this;
+		ConstThis->ForEachClient([&ProcessClient](const FReplicationClient& Client)
+		{
+			// const_cast safe here because remote clients are never const and GetLocalClient() is not const here since this overload is non-const
+			return ProcessClient(const_cast<FReplicationClient&>(Client));
+		});
+	}
+
 	void FReplicationClientManager::AddReferencedObjects(FReferenceCollector& Collector)
 	{
 		Collector.AddReferencedObject(SessionContent);
diff --git a/Engine/Plugins/Developer/Concert/ConcertApp/MultiUserClient/Source/MultiUserClient/Private/Replication/Client/ReplicationClientManager.h b/Engine/Plugins/Developer/Concert/ConcertApp/MultiUserClient/Source/MultiUserClient/Private/Replication/Client/ReplicationClientManager.h
--- a/Engine/Plugins/Developer/Concert/ConcertApp/MultiUserClient/Source/MultiUserClient/Private/Replication/Client/ReplicationClientManager.h
+++ b/Engine/Plugins/Developer/Concert/ConcertApp/MultiUserClient/Source/MultiUserClient/Private/Replication/Client/ReplicationClientManager.h
@@ -77,6 +77,10 @@ namespace UE::MultiUserClient
 			return const_cast<FReplicationClient*>(ConstThis->FindClient(EndpointId));
 		}
 
+		/** Iterates through every client */
+		void ForEachClient(TFunctionRef<EBreakBehavior(const FReplicationClient&)> ProcessClient) const;
+		/** Iterates through every client */
+		void ForEachClient(TFunctionRef<EBreakBehavior(FReplicationClient&)> ProcessClient);
 		
 		DECLARE_MULTICAST_DELEGATE(FRemoteClientsChanged);
 		/** Called when RemoteClients changes. Called after OnPostRemoteClientAdded. */
diff --git a/Engine/Plugins/Developer/Concert/ConcertApp/MultiUserClient/Source/MultiUserClient/Private/Replication/Stream/StreamChangeTracker.cpp b/Engine/Plugins/Developer/Concert/ConcertApp/MultiUserClient/Source/MultiUserClient/Private/Replication/Stream/StreamChangeTracker.cpp
--- a/Engine/Plugins/Developer/Concert/ConcertApp/MultiUserClient/Source/MultiUserClient/Private/Replication/Stream/StreamChangeTracker.cpp
+++ b/Engine/Plugins/Developer/Concert/ConcertApp/MultiUserClient/Source/MultiUserClient/Private/Replication/Stream/StreamChangeTracker.cpp
@@ -34,22 +34,6 @@ namespace UE::MultiUserClient
 		CachedDeltaChange = DiffChanges();
 	}
 
-	void FStreamChangeTracker::RevertCachedChanges()
-	{
-		if (HasChanges())
-		{
-			FScopedTransaction Transaction(LOCTEXT("RevertCachedChanges", "Revert replication changes"));
-			OnModifyReplicationMapDelegate.ExecuteIfBound();
-			
-			*StreamWithInProgressChangesAttribute.Get() = StreamSynchronizer.GetServerState();
-			CachedDeltaChange = {};
-
-			check(IsInGameThread());
-			OnChangesRevertedDelegate.Broadcast();
-			RefreshChangesCache();
-		}
-	}
-
 	bool FStreamChangeTracker::HasChanges() const
 	{
 		return !CachedDeltaChange.ObjectsToPut.IsEmpty() || !CachedDeltaChange.ObjectsToRemove.IsEmpty();
diff --git a/Engine/Plugins/Developer/Concert/ConcertApp/MultiUserClient/Source/MultiUserClient/Private/Replication/Stream/StreamChangeTracker.h b/Engine/Plugins/Developer/Concert/ConcertApp/MultiUserClient/Source/MultiUserClient/Private/Replication/Stream/StreamChangeTracker.h
--- a/Engine/Plugins/Developer/Concert/ConcertApp/MultiUserClient/Source/MultiUserClient/Private/Replication/Stream/StreamChangeTracker.h
+++ b/Engine/Plugins/Developer/Concert/ConcertApp/MultiUserClient/Source/MultiUserClient/Private/Replication/Stream/StreamChangeTracker.h
@@ -63,9 +63,6 @@ namespace UE::MultiUserClient
 		 */
 		void RefreshChangesCache();
 		const FStreamChangelist& GetCachedDeltaChange() const { return CachedDeltaChange; }
-		
-		/** Reverts the external stream StreamWithInProgressChangesAttribute to be the version registered on the server. */
-		void RevertCachedChanges();
 
 		/** @return Whether there are any changes that that can be submitted to the server (excludes those with warnings). */
 		bool HasChanges() const;
@@ -84,10 +81,6 @@ namespace UE::MultiUserClient
 		};
 		/** @return Whether this object's local configuration differs from the server version. */
 		EObjectChangeType GetObjectChanges(const FSoftObjectPath& Object) const;
-
-		/** Called when RevertCachedChanges is called. The UI should be refreshed. */
-		DECLARE_MULTICAST_DELEGATE(FOnChangesReverted);
-		FOnChangesReverted& OnChangesReverted_GameThread() { return OnChangesRevertedDelegate; }
 		
 	private:
 
@@ -108,8 +101,6 @@ namespace UE::MultiUserClient
 		
 		/** Called when StreamWithInProgressChangesAttribute is about to be modified. */
 		FOnModifyReplicationMap OnModifyReplicationMapDelegate;
-		/** Called when StreamWithInProgressChangesAttribute was modified by this FLocalClientStreamDiffer. */
-		FOnChangesReverted OnChangesRevertedDelegate;
 
 		FStreamChangelist DiffChanges() const
 		{
diff --git a/Engine/Plugins/Developer/Concert/ConcertApp/MultiUserClient/Source/MultiUserClient/Private/Replication/Submission/Data/EChangeUploadability.h b/Engine/Plugins/Developer/Concert/ConcertApp/MultiUserClient/Source/MultiUserClient/Private/Replication/Submission/Data/EChangeUploadability.h
--- a/Engine/Plugins/Developer/Concert/ConcertApp/MultiUserClient/Source/MultiUserClient/Private/Replication/Submission/Data/EChangeUploadability.h
+++ b/Engine/Plugins/Developer/Concert/ConcertApp/MultiUserClient/Source/MultiUserClient/Private/Replication/Submission/Data/EChangeUploadability.h
@@ -18,4 +18,9 @@ namespace UE::MultiUserClient
 		// TODO DP UE-198088: Remove once remote changing is implemented
 		NotImplemented
 	};
+
+	bool CanEverSubmit(EChangeUploadability Uploadability)
+	{
+		return Uploadability != EChangeUploadability::NotImplemented;
+	}
 }
\ No newline at end of file
diff --git a/Engine/Plugins/Developer/Concert/ConcertApp/MultiUserClient/Source/MultiUserClient/Private/Replication/Submission/ISubmissionWorkflow.h b/Engine/Plugins/Developer/Concert/ConcertApp/MultiUserClient/Source/MultiUserClient/Private/Replication/Submission/ISubmissionWorkflow.h
--- a/Engine/Plugins/Developer/Concert/ConcertApp/MultiUserClient/Source/MultiUserClient/Private/Replication/Submission/ISubmissionWorkflow.h
+++ b/Engine/Plugins/Developer/Concert/ConcertApp/MultiUserClient/Source/MultiUserClient/Private/Replication/Submission/ISubmissionWorkflow.h
@@ -43,10 +43,6 @@ namespace UE::MultiUserClient
 		/** @return Detailed information about whether Submit can be called */
 		virtual EChangeUploadability GetUploadability() const = 0;
 		bool CanSubmit() const { return GetUploadability() == EChangeUploadability::Ready; }
-
-		/** @return Detailed information about whether Revert can be called */
-		virtual EChangeRevertability GetRevertability() const = 0;
-		bool CanRevert() const { return GetRevertability() == EChangeRevertability::Revertable; }
 		
 		DECLARE_MULTICAST_DELEGATE_OneParam(FOnStreamRequestCompleted, const FSubmitStreamChangesResponse&);
 		/**
diff --git a/Engine/Plugins/Developer/Concert/ConcertApp/MultiUserClient/Source/MultiUserClient/Private/Replication/Submission/SubmissionWorkflow_LocalClient.cpp b/Engine/Plugins/Developer/Concert/ConcertApp/MultiUserClient/Source/MultiUserClient/Private/Replication/Submission/SubmissionWorkflow_LocalClient.cpp
--- a/Engine/Plugins/Developer/Concert/ConcertApp/MultiUserClient/Source/MultiUserClient/Private/Replication/Submission/SubmissionWorkflow_LocalClient.cpp
+++ b/Engine/Plugins/Developer/Concert/ConcertApp/MultiUserClient/Source/MultiUserClient/Private/Replication/Submission/SubmissionWorkflow_LocalClient.cpp
@@ -36,19 +36,6 @@ namespace UE::MultiUserClient
 		return bHasChanges ? EChangeUploadability::Ready : EChangeUploadability::NoChanges;
 	}
 
-	EChangeRevertability FSubmissionWorkflow_LocalClient::GetRevertability() const
-	{
-		if (GetUploadability() == EChangeUploadability::InProgress)
-		{
-			return EChangeRevertability::UploadInProgress;
-		}
-
-		const bool bStreamHasChanges = StreamChangeTracker.HasChanges();
-		const bool bAuthorityHasChanges = AuthorityChangeTracker.HasChanges();
-		const bool bHasChanges = bStreamHasChanges || bAuthorityHasChanges;
-		return bHasChanges ? EChangeRevertability::Revertable : EChangeRevertability::NoChanges;
-	}
-
 	TSharedPtr<ISubmissionOperation> FSubmissionWorkflow_LocalClient::SubmitChanges()
 	{
 		using namespace ConcertSyncClient::Replication;
diff --git a/Engine/Plugins/Developer/Concert/ConcertApp/MultiUserClient/Source/MultiUserClient/Private/Replication/Submission/SubmissionWorkflow_LocalClient.h b/Engine/Plugins/Developer/Concert/ConcertApp/MultiUserClient/Source/MultiUserClient/Private/Replication/Submission/SubmissionWorkflow_LocalClient.h
--- a/Engine/Plugins/Developer/Concert/ConcertApp/MultiUserClient/Source/MultiUserClient/Private/Replication/Submission/SubmissionWorkflow_LocalClient.h
+++ b/Engine/Plugins/Developer/Concert/ConcertApp/MultiUserClient/Source/MultiUserClient/Private/Replication/Submission/SubmissionWorkflow_LocalClient.h
@@ -36,7 +36,6 @@ namespace UE::MultiUserClient
 		//~ Begin ISubmissionWorkflow Interface
 		virtual TSharedPtr<ISubmissionOperation> SubmitChanges() override;
 		virtual EChangeUploadability GetUploadability() const override;
-		virtual EChangeRevertability GetRevertability() const override;
 		//~ End ISubmissionWorkflow Interface
 	
 	private:
diff --git a/Engine/Plugins/Developer/Concert/ConcertApp/MultiUserClient/Source/MultiUserClient/Private/Replication/Submission/SubmissionWorkflow_RemoteClient.cpp b/Engine/Plugins/Developer/Concert/ConcertApp/MultiUserClient/Source/MultiUserClient/Private/Replication/Submission/SubmissionWorkflow_RemoteClient.cpp
--- a/Engine/Plugins/Developer/Concert/ConcertApp/MultiUserClient/Source/MultiUserClient/Private/Replication/Submission/SubmissionWorkflow_RemoteClient.cpp
+++ b/Engine/Plugins/Developer/Concert/ConcertApp/MultiUserClient/Source/MultiUserClient/Private/Replication/Submission/SubmissionWorkflow_RemoteClient.cpp
@@ -17,11 +17,4 @@ namespace UE::MultiUserClient
 		// TODO DP UE-198088 changing remote client authority
 		return EChangeUploadability::NotImplemented;
 	}
-
-	EChangeRevertability FSubmissionWorkflow_RemoteClient::GetRevertability() const
-	{
-		// TODO DP UE-180657: Changing remote client's stream is not implemented for now
-		// TODO DP UE-198088 changing remote client authority
-		return EChangeRevertability::NoChanges;
-	}
 }
diff --git a/Engine/Plugins/Developer/Concert/ConcertApp/MultiUserClient/Source/MultiUserClient/Private/Replication/Submission/SubmissionWorkflow_RemoteClient.h b/Engine/Plugins/Developer/Concert/ConcertApp/MultiUserClient/Source/MultiUserClient/Private/Replication/Submission/SubmissionWorkflow_RemoteClient.h
--- a/Engine/Plugins/Developer/Concert/ConcertApp/MultiUserClient/Source/MultiUserClient/Private/Replication/Submission/SubmissionWorkflow_RemoteClient.h
+++ b/Engine/Plugins/Developer/Concert/ConcertApp/MultiUserClient/Source/MultiUserClient/Private/Replication/Submission/SubmissionWorkflow_RemoteClient.h
@@ -19,7 +19,6 @@ namespace UE::MultiUserClient
 		//~ Begin ISubmissionWorkflow Interface
 		virtual TSharedPtr<ISubmissionOperation> SubmitChanges() override;
 		virtual EChangeUploadability GetUploadability() const override;
-		virtual EChangeRevertability GetRevertability() const override;
 		//~ End ISubmissionWorkflow Interface
 	};
 }
diff --git a/Engine/Plugins/Developer/Concert/ConcertApp/MultiUserClient/Source/MultiUserClient/Private/Widgets/ActiveSession/Replication/Client/ClientUtils.cpp b/Engine/Plugins/Developer/Concert/ConcertApp/MultiUserClient/Source/MultiUserClient/Private/Widgets/ActiveSession/Replication/Client/ClientUtils.cpp
--- a/Engine/Plugins/Developer/Concert/ConcertApp/MultiUserClient/Source/MultiUserClient/Private/Widgets/ActiveSession/Replication/Client/ClientUtils.cpp
+++ b/Engine/Plugins/Developer/Concert/ConcertApp/MultiUserClient/Source/MultiUserClient/Private/Widgets/ActiveSession/Replication/Client/ClientUtils.cpp
@@ -2,22 +2,24 @@
 
 #include "ClientUtils.h"
 
+#include "ConcertMessageData.h"
 #include "IConcertClient.h"
+#include "Replication/Client/ReplicationClientManager.h"
 #include "Widgets/ClientName/SClientName.h"
 
 namespace UE::MultiUserClient::ClientUtils
 {
-	FString GetClientDisplayName(const IConcertClient& InLocalClientInstance, const FGuid& InClientToGetName)
+	FString GetClientDisplayName(const IConcertClient& InLocalClientInstance, const FGuid& InClientEndpointId)
 	{
 		const TSharedPtr<IConcertClientSession> Session = InLocalClientInstance.GetCurrentSession();
-		const bool bIsLocalClient = ensure(Session) && Session->GetSessionClientEndpointId() == InClientToGetName;
+		const bool bIsLocalClient = ensure(Session) && Session->GetSessionClientEndpointId() == InClientEndpointId;
 		if (bIsLocalClient)
 		{
 			return ConcertClientSharedSlate::SClientName::GetDisplayText(Session->GetLocalClientInfo(), bIsLocalClient).ToString();
 		}
 
 		FConcertSessionClientInfo ClientInfo;
-		if (ensure(Session) && Session->FindSessionClient(InClientToGetName, ClientInfo))
+		if (ensure(Session) && Session->FindSessionClient(InClientEndpointId, ClientInfo))
 		{
 			return ConcertClientSharedSlate::SClientName::GetDisplayText(ClientInfo.ClientInfo, bIsLocalClient).ToString();
 		}
@@ -25,4 +27,51 @@ namespace UE::MultiUserClient::ClientUtils
 		ensureMsgf(false, TEXT("Bad args"));
 		return {};
 	}
+	
+	bool GetClientDisplayInfo(const IConcertClient& InLocalClientInstance, const FGuid& InClientEndpointId, FConcertClientInfo& OutClientInfo)
+	{
+		const TSharedPtr<IConcertClientSession> Session = InLocalClientInstance.GetCurrentSession();
+		const bool bIsLocalClient = ensure(Session) && Session->GetSessionClientEndpointId() == InClientEndpointId;
+		if (bIsLocalClient)
+		{
+			OutClientInfo = InLocalClientInstance.GetClientInfo();
+			return true;
+		}
+
+		FConcertSessionClientInfo ClientInfo;
+		if (ensure(Session) && Session->FindSessionClient(InClientEndpointId, ClientInfo))
+		{
+			OutClientInfo = MoveTemp(ClientInfo.ClientInfo);
+			return true;
+		}
+
+		return false;
+	}
+	
+	TArray<const FReplicationClient*> GetSortedClientList(const IConcertClient& InLocalClientInstance, const FReplicationClientManager& InReplicationManager)
+	{
+		TArray<const FReplicationClient*> Result;
+		TMap<const FReplicationClient*, FConcertClientInfo> ClientToDisplayInfo;
+		for (const TNonNullPtr<FRemoteReplicationClient> RemoteClient : InReplicationManager.GetRemoteClients())
+		{
+			FConcertClientInfo Info;
+			if (GetClientDisplayInfo(InLocalClientInstance, RemoteClient->GetEndpointId(), Info))
+			{
+				Result.Add(RemoteClient);
+				ClientToDisplayInfo.Add(RemoteClient, MoveTemp(Info));
+			}
+		}
+
+		Result.Sort([&ClientToDisplayInfo](const FReplicationClient& Left, const FReplicationClient& Right)
+		{
+			const FConcertClientInfo& LeftInfo = ClientToDisplayInfo[&Left];
+			const FConcertClientInfo& RightInfo = ClientToDisplayInfo[&Right];
+			const FText LeftDisplayName = ConcertClientSharedSlate::SClientName::GetDisplayText(LeftInfo, false);
+			const FText RightDisplayName = ConcertClientSharedSlate::SClientName::GetDisplayText(RightInfo, false);
+			return LeftDisplayName.ToString() <= RightDisplayName.ToString();
+		});
+		
+		Result.Insert(&InReplicationManager.GetLocalClient(), 0);
+		return Result;
+	}
 }
diff --git a/Engine/Plugins/Developer/Concert/ConcertApp/MultiUserClient/Source/MultiUserClient/Private/Widgets/ActiveSession/Replication/Client/ClientUtils.h b/Engine/Plugins/Developer/Concert/ConcertApp/MultiUserClient/Source/MultiUserClient/Private/Widgets/ActiveSession/Replication/Client/ClientUtils.h
--- a/Engine/Plugins/Developer/Concert/ConcertApp/MultiUserClient/Source/MultiUserClient/Private/Widgets/ActiveSession/Replication/Client/ClientUtils.h
+++ b/Engine/Plugins/Developer/Concert/ConcertApp/MultiUserClient/Source/MultiUserClient/Private/Widgets/ActiveSession/Replication/Client/ClientUtils.h
@@ -2,18 +2,45 @@
 
 #pragma once
 
+#include "Containers/Array.h"
+
 class FString;
 class IConcertClient;
+struct FConcertClientInfo;
 struct FGuid;
 
+namespace UE::MultiUserClient
+{
+	class FReplicationClient;
+	class FReplicationClientManager;
+}
+
 namespace UE::MultiUserClient::ClientUtils
 {
 	/**
 	 * Gets the display name for a client. Appends (me) if the client is local.
 	 * 
 	 * @param InLocalClientInstance Used to look up client display info
-	 * @param InClientToGetName The endpoint ID of the client whose name to get
+	 * @param InClientEndpointId The endpoint ID of the client whose name to get
 	 * @return The display name or empty
 	 */
-	FString GetClientDisplayName(const IConcertClient& InLocalClientInstance, const FGuid& InClientToGetName);
+	FString GetClientDisplayName(const IConcertClient& InLocalClientInstance, const FGuid& InClientEndpointId);
+
+	/**
+	 * Gets the display info for a given client
+	 *
+	 * @param InLocalClientInstance Used to look up client display info
+	 * @param InClientEndpointId The endpoint ID of the client whose name to get
+	 * @param OutClientInfo The client display info to get
+	 * @return Whether OutClientInfo holds a value
+	 */
+	bool GetClientDisplayInfo(const IConcertClient& InLocalClientInstance, const FGuid& InClientEndpointId, FConcertClientInfo& OutClientInfo);
+
+	/**
+	 * Gets all replication clients in a sorted array. The local client will always be first and then come all remote clients sorted alphabetically.
+	 * @param InLocalClientInstance Used to look up client display info
+	 * @param InReplicationManager Used to obtain all clients
+	 * @return Sorted client array
+	 */
+	TArray<const FReplicationClient*> GetSortedClientList(const IConcertClient& InLocalClientInstance, const FReplicationClientManager& InReplicationManager);
 }
diff --git a/Engine/Plugins/Developer/Concert/ConcertApp/MultiUserClient/Source/MultiUserClient/Private/Widgets/ActiveSession/Replication/Client/Columns/MultiStreamColumns.cpp b/Engine/Plugins/Developer/Concert/ConcertApp/MultiUserClient/Source/MultiUserClient/Private/Widgets/ActiveSession/Replication/Client/Columns/MultiStreamColumns.cpp
new file mode 100644
--- /dev/null
+++ b/Engine/Plugins/Developer/Concert/ConcertApp/MultiUserClient/Source/MultiUserClient/Private/Widgets/ActiveSession/Replication/Client/Columns/MultiStreamColumns.cpp
@@ -0,0 +1,114 @@
+// Copyright Epic Games, Inc. All Rights Reserved.
+
+#include "MultiStreamColumns.h"
+
+#include "IConcertClient.h"
+#include "Replication/Client/ReplicationClient.h"
+#include "Replication/Client/ReplicationClientManager.h"
+#include "Replication/Editor/Model/IEditableMultiReplicationStreamModel.h"
+#include "Replication/Editor/Model/IEditableReplicationStreamModel.h"
+#include "Replication/Editor/View/IMultiReplicationStreamEditor.h"
+#include "Replication/Editor/View/IReplicationStreamEditor.h"
+#include "SAssignPropertyComboBox.h"
+#include "Widgets/ActiveSession/Replication/Client/ClientUtils.h"
+
+#define LOCTEXT_NAMESPACE "AssignPropertyColumn"
+
+namespace UE::MultiUserClient::MultiStreamColumns
+{
+	const FName AssignPropertyColumnId(TEXT("AssignPropertyColumn"));
+
+	namespace AssignPropertyColumnUtils
+	{
+		static void ForEachStreamAssignedTo(
+			const ConcertClientSharedSlate::IMultiReplicationStreamEditor& MultiEditor,
+			const FConcertPropertyChain& Property,
+			TFunctionRef<void(const TSharedRef<ConcertClientSharedSlate::IEditableReplicationStreamModel>& Stream)> Consume
+		)
+		{
+			for (const TSharedRef<ConcertClientSharedSlate::IEditableReplicationStreamModel>& Stream : MultiEditor.GetModel().GetEditableStreams())
+			{
+				for (const FSoftObjectPath& SelectedObject : MultiEditor.GetEditorBase().GetObjectsBeingPropertyEdited())
+				{
+					const bool bStreamAssignedToProperty = Stream->HasProperty(SelectedObject, Property);
+					if (bStreamAssignedToProperty)
+					{
+						Consume(Stream);
+						// Each IEditableReplicationStreamModel should only be visited at most once.
+						break;
+					}
+				}
+			}
+		}
+
+		const FReplicationClient* FindClientByStream(const FReplicationClientManager& ClientManager, const ConcertClientSharedSlate::IReplicationStreamModel& StreamModel)
+		{
+			if (&ClientManager.GetLocalClient().GetClientEditModel().Get() == &StreamModel)
+			{
+				return &ClientManager.GetLocalClient();
+			}
+
+			for (const TNonNullPtr<FRemoteReplicationClient> Client : ClientManager.GetRemoteClients())
+			{
+				if (&Client->GetClientEditModel().Get() == &StreamModel)
+				{
+					return Client;
+				}
+			}
+			
+			return nullptr;
+		}
+		
+		static FString GetClientDisplayText(const IConcertClient& InConcertClient, const FReplicationClientManager& ClientManager, const ConcertClientSharedSlate::IReplicationStreamModel& StreamModel)
+		{
+			if (const FReplicationClient* Client = FindClientByStream(ClientManager, StreamModel))
+			{
+				return ClientUtils::GetClientDisplayName(InConcertClient, Client->GetEndpointId());
+			}
+
+			ensure(false);
+			return {};
+		}
+	}
+	
+	ConcertClientSharedSlate::ReplicationColumns::FReplicationPropertyColumn AssignPropertyColumn(
+		TAttribute<TSharedPtr<ConcertClientSharedSlate::IMultiReplicationStreamEditor>> MultiStreamEditor,
+		TSharedRef<IConcertClient> ConcertClient,
+		FReplicationClientManager& ClientManager,
+		const int32 ColumnsSortPriority
+		)
+	{
+		using namespace ConcertClientSharedSlate;
+		using namespace ConcertClientSharedSlate::ReplicationColumns;
+		check(MultiStreamEditor.IsBound() || MultiStreamEditor.IsSet());
+
+		const auto PopuluateSearch = [MultiStreamEditor, &ClientManager, ConcertClient](const FReplicatedPropertyData& Data, TArray<FString>& InOutSearchStrings)
+		{
+			AssignPropertyColumnUtils::ForEachStreamAssignedTo(*MultiStreamEditor.Get(), Data.GetProperty(),
+				[&ClientManager, &ConcertClient, &InOutSearchStrings](const TSharedRef<IEditableReplicationStreamModel>& Stream)
+				{
+					InOutSearchStrings.Add(AssignPropertyColumnUtils::GetClientDisplayText(*ConcertClient, ClientManager, *Stream));
+				});
+		};
+		const auto GenerateWidgetColumn = [MultiStreamEditor, ConcertClient, &ClientManager](const FReplicationPropertyColumn::FBuildArgs& Args)
+		{
+			const TArray<FSoftObjectPath> DisplayedObjects = MultiStreamEditor.Get()->GetEditorBase().GetObjectsBeingPropertyEdited();
+			return SNew(SAssignPropertyComboBox, MultiStreamEditor.Get().ToSharedRef(), ConcertClient, ClientManager)
+				.DisplayedProperty(Args.RowData.GetProperty())
+				.EditedObjects(DisplayedObjects)
+				.HighlightText(Args.HighlightText);
+		};
+		
+		return FReplicationPropertyColumn(
+			typename FReplicationPropertyColumn::FArguments()
+				.PopulateSearchItems_Lambda(PopuluateSearch)
+				.GenerateWidgetColumn_Lambda(GenerateWidgetColumn)
+				.ColumnSortOrder(ColumnsSortPriority),
+			SHeaderRow::Column(AssignPropertyColumnId)
+				.DefaultLabel(LOCTEXT("Owner", "Owner"))
+				.FillSized(200.f)
+			);
+	}
+}
+
+#undef LOCTEXT_NAMESPACE
\ No newline at end of file
diff --git a/Engine/Plugins/Developer/Concert/ConcertApp/MultiUserClient/Source/MultiUserClient/Private/Widgets/ActiveSession/Replication/Client/Columns/MultiStreamColumns.h b/Engine/Plugins/Developer/Concert/ConcertApp/MultiUserClient/Source/MultiUserClient/Private/Widgets/ActiveSession/Replication/Client/Columns/MultiStreamColumns.h
new file mode 100644
--- /dev/null
+++ b/Engine/Plugins/Developer/Concert/ConcertApp/MultiUserClient/Source/MultiUserClient/Private/Widgets/ActiveSession/Replication/Client/Columns/MultiStreamColumns.h
@@ -0,0 +1,44 @@
+// Copyright Epic Games, Inc. All Rights Reserved.
+
+#pragma once
+
+#include "Replication/Editor/View/PredefinedReplicationColumns.h"
+
+#include "Delegates/Delegate.h"
+#include "Misc/EnumClassFlags.h"
+
+class IConcertClient;
+
+namespace UE::MultiUserClient
+{
+	class FReplicationClientManager;
+}
+
+namespace UE::ConcertClientSharedSlate
+{
+	class IReplicationStreamModel;
+	class IMultiReplicationStreamEditor;
+	class IReplicationStreamViewer;
+}
+
+namespace UE::MultiUserClient::MultiStreamColumns
+{
+	const extern FName AssignPropertyColumnId;
+
+	/**
+	 * Creates a property column which assigns the property to the stream selected in the combo box.
+	 * 
+	 * @param MultiStreamEditor Used to determine the selected objects.
+	 * @param ConcertClient Used to look up client names
+	 * @param ClientManager Used to map streams back to client display info
+	 * @param ColumnsSortPriority The order relative to the other columns
+	 * 
+	 * @return A column that spawns a combo box for assigning properties
+	 */
+	ConcertClientSharedSlate::ReplicationColumns::FReplicationPropertyColumn AssignPropertyColumn(
+		TAttribute<TSharedPtr<ConcertClientSharedSlate::IMultiReplicationStreamEditor>> MultiStreamEditor,
+		TSharedRef<IConcertClient> ConcertClient,
+		FReplicationClientManager& ClientManager,
+		const int32 ColumnsSortPriority = 30 /* @see EReplicationPropertyColumnOrder */
+		);
+}
diff --git a/Engine/Plugins/Developer/Concert/ConcertApp/MultiUserClient/Source/MultiUserClient/Private/Widgets/ActiveSession/Replication/Client/Columns/SAssignPropertyComboBox.cpp b/Engine/Plugins/Developer/Concert/ConcertApp/MultiUserClient/Source/MultiUserClient/Private/Widgets/ActiveSession/Replication/Client/Columns/SAssignPropertyComboBox.cpp
new file mode 100644
--- /dev/null
+++ b/Engine/Plugins/Developer/Concert/ConcertApp/MultiUserClient/Source/MultiUserClient/Private/Widgets/ActiveSession/Replication/Client/Columns/SAssignPropertyComboBox.cpp
@@ -0,0 +1,290 @@
+// Copyright Epic Games, Inc. All Rights Reserved.
+
+#include "SAssignPropertyComboBox.h"
+
+#include "IConcertClient.h"
+#include "Replication/Client/ReplicationClient.h"
+#include "Replication/Client/ReplicationClientManager.h"
+#include "Replication/Editor/Model/IEditableReplicationStreamModel.h"
+#include "Widgets/ActiveSession/Replication/Client/ClientUtils.h"
+#include "Widgets/ClientName/SHorizontalClientList.h"
+#include "Widgets/ClientName/SLocalClientName.h"
+#include "Widgets/ClientName/SRemoteClientName.h"
+
+#include "Algo/AnyOf.h"
+#include "Framework/MultiBox/MultiBoxBuilder.h"
+#include "ScopedTransaction.h"
+#include "Styling/SlateTypes.h"
+#include "Widgets/Input/SComboButton.h"
+
+#define LOCTEXT_NAMESPACE "SAssignPropertyComboBox"
+
+namespace UE::MultiUserClient
+{
+	void SAssignPropertyComboBox::Construct(const FArguments& InArgs,
+        TSharedRef<ConcertClientSharedSlate::IMultiReplicationStreamEditor> InEditor,
+        TSharedRef<IConcertClient> InConcertClient,
+        FReplicationClientManager& InClientManager
+	)
+	{
+		Editor = InEditor;
+		ConcertClient = MoveTemp(InConcertClient);
+		ClientManager = &InClientManager;
+		
+		Property = InArgs._DisplayedProperty;
+		EditedObjects = InArgs._EditedObjects;
+		HighlightText = InArgs._HighlightText;
+		check(!EditedObjects.IsEmpty());
+		
+		ChildSlot
+		[
+			SNew(SComboButton)
+			.HasDownArrow(true)
+			.ContentPadding(FMargin(2.0f, 2.0f))
+			.ButtonContent()
+			[
+				SAssignNew(ClientListWidget, ConcertClientSharedSlate::SHorizontalClientList, ConcertClient.ToSharedRef())
+				.HighlightText_Lambda([this](){ return HighlightText ? *HighlightText : FText::GetEmpty(); })
+				.EmptyListSlot()
+				[
+					SNew(STextBlock)
+					.Text(LOCTEXT("None", " - "))
+				]
+			]
+			.OnGetMenuContent(this, &SAssignPropertyComboBox::GetMenuContent)
+		];
+
+		ClientManager->OnRemoteClientsChanged().AddSP(this, &SAssignPropertyComboBox::RebuildSubscriptionsAndRefresh);
+		RebuildSubscriptions();
+		RefreshContentBoxContent();
+	}
+	
+	void SAssignPropertyComboBox::RefreshContentBoxContent()
+	{
+		TArray<FGuid> Clients;
+		ClientManager->ForEachClient([this, &Clients](const FReplicationClient& Client)
+		{
+			const TMap<FSoftObjectPath, FReplicatedObjectInfo>& ObjectInfoMap = Client.GetStreamSynchronizer().GetServerState().ReplicatedObjects;
+			for (const FSoftObjectPath& ObjectPath : EditedObjects)
+			{
+				if (const FReplicatedObjectInfo* ObjectInfo = ObjectInfoMap.Find(ObjectPath)
+					; ObjectInfo && ObjectInfo->PropertySelection.ReplicatedProperties.Contains(Property))
+				{
+					Clients.Add(Client.GetEndpointId());
+					return EBreakBehavior::Continue;
+				}
+			}
+			return EBreakBehavior::Continue;
+		});
+
+		ClientListWidget->RefreshList(Clients);
+	}
+
+	TSharedRef<SWidget> SAssignPropertyComboBox::GetMenuContent()
+	{
+		using namespace ConcertClientSharedSlate;
+
+		const auto MakeWidget = [this](const FGuid& EndpointId) -> TSharedRef<SWidget>
+		{
+			const bool bIsLocalClient = EndpointId == ConcertClient->GetCurrentSession()->GetSessionClientEndpointId();
+			if (bIsLocalClient)
+			{
+				return SNew(SLocalClientName, ConcertClient.ToSharedRef())
+					.HighlightText_Lambda([this](){ return HighlightText ? *HighlightText : FText::GetEmpty(); });
+			}
+			return SNew(SRemoteClientName, ConcertClient.ToSharedRef())
+				.ClientEndpointId(EndpointId)
+				.HighlightText_Lambda([this](){ return HighlightText ? *HighlightText : FText::GetEmpty(); });
+		};
+		
+		FMenuBuilder MenuBuilder(true, nullptr);
+		for (const FReplicationClient* Client : ClientUtils::GetSortedClientList(*ConcertClient, *ClientManager))
+		{
+			TAttribute<FText> Tooltip = TAttribute<FText>::CreateLambda([this, EndpointId = Client->GetEndpointId()]()
+			{
+				FText Reason;
+				const bool bCanClick = CanClickOptionWithReason(EndpointId, &Reason);
+				if (!bCanClick)
+				{
+					return Reason;
+				}
+
+				switch (GetOptionCheckState(EndpointId))
+				{
+				case ECheckBoxState::Unchecked: return LOCTEXT("Action.Unchecked", "Assign property to client and remove it from all others.");
+				case ECheckBoxState::Undetermined:  return LOCTEXT("Action.Undetermined", "Assign property to client for all selected objects.");
+				case ECheckBoxState::Checked: return LOCTEXT("Action.Checked", "Remove property from client and remove it from all others.");
+				default: return FText::GetEmpty();
+				}
+			});
+			
+			MenuBuilder.AddMenuEntry(
+				FUIAction(
+					FExecuteAction::CreateSP(this, &SAssignPropertyComboBox::OnClickOption, Client->GetEndpointId()),
+					FCanExecuteAction::CreateSP(this, &SAssignPropertyComboBox::CanClickOption, Client->GetEndpointId()),
+					FGetActionCheckState::CreateSP(this, &SAssignPropertyComboBox::GetOptionCheckState, Client->GetEndpointId())
+					),
+				MakeWidget(Client->GetEndpointId()),
+				NAME_None,
+				Tooltip,
+				EUserInterfaceActionType::ToggleButton
+				);
+		}
+		
+		return MenuBuilder.MakeWidget();
+	}
+	
+	void SAssignPropertyComboBox::OnClickOption(const FGuid EndpointId) const
+	{
+		// Remote clients can disconnect after the combo-box is opened.
+		const FReplicationClient* Client = ClientManager->FindClient(EndpointId);
+		if (!Client)
+		{
+			return;
+		}
+		
+		const FText TransactionText = FText::Format(LOCTEXT("AllClientsAssignFmt", "Assign {0} property"), FText::FromString(Property.ToString(FConcertPropertyChain::EToStringMethod::LeafProperty)));
+		FScopedTransaction Transaction(TransactionText);
+
+		// Only one client is supposed to own the object: remove the other clients (if possible)
+		ClientManager->ForEachClient([this, Client](const FReplicationClient& ClientToRemoveFrom)
+		{
+			if (*Client != ClientToRemoveFrom
+				&& ClientToRemoveFrom.AllowsEditing())
+			{
+				for (const FSoftObjectPath& ObjectPath : EditedObjects)
+				{
+					ClientToRemoveFrom.GetClientEditModel()->RemoveProperties(ObjectPath, { Property });
+				}
+			}
+			
+			return EBreakBehavior::Continue;
+		});
+
+		const ECheckBoxState CheckBoxState = GetOptionCheckState(EndpointId);
+		const bool bRemoveProperty = CheckBoxState == ECheckBoxState::Checked; 
+		const TSharedRef<ConcertClientSharedSlate::IEditableReplicationStreamModel> EditModel = Client->GetClientEditModel();
+		for (const FSoftObjectPath& ObjectPath : EditedObjects)
+		{
+			if (bRemoveProperty)
+			{
+				EditModel->RemoveProperties(ObjectPath, { Property });
+				if (!EditModel->HasAnyPropertyAssigned(ObjectPath))
+				{
+					EditModel->RemoveObjects({ ObjectPath });
+				}
+			}
+			else
+			{
+				if (!EditModel->ContainsObjects({ ObjectPath }))
+				{
+					EditModel->AddObjects({ ObjectPath.ResolveObject() });
+				}
+				
+				if (EditModel->ContainsObjects({ ObjectPath }))
+				{
+					EditModel->AddProperties(ObjectPath, { Property });
+				}
+			}
+		}
+	}
+
+#define SET_REASON(Text) if (Reason) { *Reason = Text; }
+	bool SAssignPropertyComboBox::CanClickOptionWithReason(const FGuid& EndpointId, FText* Reason) const
+	{
+		const FReplicationClient* Client = ClientManager->FindClient(EndpointId);
+		// Remote clients can disconnect after the combo-box is opened.
+		if (!Client)
+		{
+			SET_REASON(LOCTEXT("ClientDisconnected", "Client disconnected."));
+			return false;
+		}
+
+		// The combo box assigns the property to the clicked client and removes from the others... check that the currently assigned clients allow it.
+		bool bCanRemoveFromOwners = true;
+		ClientManager->ForEachClient([this, &EndpointId, &Reason, Client, &bCanRemoveFromOwners](const FReplicationClient& ClientToRemoveFrom)
+		{
+			if (*Client != ClientToRemoveFrom && !ClientToRemoveFrom.AllowsEditing())
+			{
+				const bool bHasAnySelectedObject = Algo::AnyOf(EditedObjects, [this, &ClientToRemoveFrom](const FSoftObjectPath& ObjectPath)
+				{
+					return ClientToRemoveFrom.GetClientEditModel()->HasProperty(ObjectPath, Property);
+				});
+				bCanRemoveFromOwners = !bHasAnySelectedObject;
+				
+				SET_REASON(FText::Format(
+					LOCTEXT("OwningClientDoesNotAllow", "Client {0} does not allow remote editing of its properties but has registered this property."),
+					FText::FromString(ClientUtils::GetClientDisplayName(*ConcertClient, EndpointId))
+					));
+				return EBreakBehavior::Break;
+			}
+			return EBreakBehavior::Continue;
+		});
+
+		if (!bCanRemoveFromOwners)
+		{
+			return false;
+		}
+		
+		const bool bAllowsEditing = Client->AllowsEditing();
+		if (!bAllowsEditing)
+		{
+			SET_REASON(FText::Format(
+				LOCTEXT("RemoteEditingDisabled", "Client {0} does not allow remote editing of its properties."),
+				FText::FromString(ClientUtils::GetClientDisplayName(*ConcertClient, EndpointId))
+				));
+		}
+		return bAllowsEditing;
+	}
+#undef SET_REASON
+	
+	ECheckBoxState SAssignPropertyComboBox::GetOptionCheckState(const FGuid EndpointId) const
+	{
+		const FReplicationClient* Client = ClientManager->FindClient(EndpointId);
+		// Remote clients can disconnect after the combo-box is opened.
+		if (!Client)
+		{
+			return ECheckBoxState::Unchecked;
+		}
+
+		const TSharedRef<ConcertClientSharedSlate::IEditableReplicationStreamModel> Model = Client->GetClientEditModel();
+		ECheckBoxState CheckBoxState = ECheckBoxState::Undetermined;
+		for (const FSoftObjectPath& ObjectPath : EditedObjects)
+		{
+			const bool bHasProperty = Model->HasProperty(ObjectPath, Property);
+			switch (CheckBoxState)
+			{
+			case ECheckBoxState::Unchecked:
+				if (bHasProperty)
+				{
+					return ECheckBoxState::Undetermined;
+				}
+				break;
+			case ECheckBoxState::Checked:
+				if (!bHasProperty)
+				{
+					return ECheckBoxState::Undetermined;
+				}
+				break;
+			case ECheckBoxState::Undetermined:
+				CheckBoxState = bHasProperty ? ECheckBoxState::Checked : ECheckBoxState::Unchecked;
+				break;
+			default: ;
+			}
+		}
+
+		return CheckBoxState;
+	}
+	
+	void SAssignPropertyComboBox::RebuildSubscriptions()
+	{
+		ClientManager->ForEachClient([this](FReplicationClient& Client)
+		{
+			Client.OnModelChanged().RemoveAll(this);
+			Client.OnModelChanged().AddSP(this, &SAssignPropertyComboBox::RefreshContentBoxContent);
+			return EBreakBehavior::Continue;
+		});
+	}
+}
+
+#undef LOCTEXT_NAMESPACE
\ No newline at end of file
diff --git a/Engine/Plugins/Developer/Concert/ConcertApp/MultiUserClient/Source/MultiUserClient/Private/Widgets/ActiveSession/Replication/Client/Columns/SAssignPropertyComboBox.h b/Engine/Plugins/Developer/Concert/ConcertApp/MultiUserClient/Source/MultiUserClient/Private/Widgets/ActiveSession/Replication/Client/Columns/SAssignPropertyComboBox.h
new file mode 100644
--- /dev/null
+++ b/Engine/Plugins/Developer/Concert/ConcertApp/MultiUserClient/Source/MultiUserClient/Private/Widgets/ActiveSession/Replication/Client/Columns/SAssignPropertyComboBox.h
@@ -0,0 +1,90 @@
+// Copyright Epic Games, Inc. All Rights Reserved.
+
+#pragma once
+
+#include "Replication/Data/ConcertPropertySelection.h"
+
+#include "Containers/Array.h"
+#include "Widgets/DeclarativeSyntaxSupport.h"
+#include "Widgets/SCompoundWidget.h"
+
+class IConcertClient;
+
+enum class ECheckBoxState : uint8;
+
+struct FGuid;
+
+namespace UE::ConcertClientSharedSlate
+{
+	class SHorizontalClientList;
+	class IMultiReplicationStreamEditor;
+}
+
+namespace UE::MultiUserClient
+{
+	class FReplicationClientManager;
+	
+	/**
+	 * Placed in every property row column to assign properties to clients.
+	 * 
+	 * Can be used to assign a property to at most one client: multiple clients cannot be assigned because it is an advanced workflow
+	 * which might confuse basic users.
+	 */
+	class SAssignPropertyComboBox : public SCompoundWidget
+	{
+	public:
+
+		SLATE_BEGIN_ARGS(SAssignPropertyComboBox)
+		{}
+			SLATE_ARGUMENT(FConcertPropertyChain, DisplayedProperty)
+			SLATE_ARGUMENT(TArray<FSoftObjectPath>, EditedObjects)
+			SLATE_ARGUMENT(TSharedPtr<FText>, HighlightText)
+		SLATE_END_ARGS()
+
+		void Construct(const FArguments& InArgs,
+           TSharedRef<ConcertClientSharedSlate::IMultiReplicationStreamEditor> InEditor,
+           TSharedRef<IConcertClient> InConcertClient,
+           FReplicationClientManager& InClientManager
+		);
+
+	private:
+
+		FReplicationClientManager* ClientManager = nullptr;
+		
+		/** Used to obtain info about the streams */
+		TSharedPtr<ConcertClientSharedSlate::IMultiReplicationStreamEditor> Editor;
+		TSharedPtr<IConcertClient> ConcertClient;
+
+		/** The objects for which the property is being displayed. */
+		TArray<FSoftObjectPath> EditedObjects;
+		/** The property assigned to this column. */
+		FConcertPropertyChain Property;
+
+		/** The static menu content (when there is no drop-down). */
+		TSharedPtr<ConcertClientSharedSlate::SHorizontalClientList> ClientListWidget;
+
+		/** Passed to MakeListWidgetDelegate */
+		TSharedPtr<FText> HighlightText;
+
+		/** Updates the content of the combo box */
+		void RefreshContentBoxContent();
+		
+		/** Builds the drop-down menu */
+		TSharedRef<SWidget> GetMenuContent();
+
+		/** Called when a client is clicked in the drop-down menu */
+		void OnClickOption(const FGuid EndpointId) const;
+		
+		/** @return Whether an item in the drop-down menu should be clickable (returns true if the stream is writable) */
+		bool CanClickOption(const FGuid EndpointId) const { return CanClickOptionWithReason(EndpointId, nullptr); }
+		bool CanClickOptionWithReason(const FGuid& EndpointId, FText* Reason = nullptr) const;
+
+		/** @return Checkbox state for the item. Handles cases of multiple objects being property edited. */
+		ECheckBoxState GetOptionCheckState(const FGuid EndpointId) const;
+
+		/** Resubscribes to all clients changing. */
+		void RebuildSubscriptions();
+		void RebuildSubscriptionsAndRefresh() { RebuildSubscriptions(); RefreshContentBoxContent(); }
+	};
+}
+
diff --git a/Engine/Plugins/Developer/Concert/ConcertApp/MultiUserClient/Source/MultiUserClient/Private/Widgets/ActiveSession/Replication/Client/Multi/AllClientsSelectionModel.cpp b/Engine/Plugins/Developer/Concert/ConcertApp/MultiUserClient/Source/MultiUserClient/Private/Widgets/ActiveSession/Replication/Client/Multi/AllClientsSelectionModel.cpp
new file mode 100644
--- /dev/null
+++ b/Engine/Plugins/Developer/Concert/ConcertApp/MultiUserClient/Source/MultiUserClient/Private/Widgets/ActiveSession/Replication/Client/Multi/AllClientsSelectionModel.cpp
@@ -0,0 +1,29 @@
+// Copyright Epic Games, Inc. All Rights Reserved.
+
+#include "AllClientsSelectionModel.h"
+
+#include "Replication/Client/ReplicationClientManager.h"
+
+namespace UE::MultiUserClient
+{
+	FAllClientsSelectionModel::FAllClientsSelectionModel(FReplicationClientManager& InClientManager)
+		: ClientManager(InClientManager)
+	{
+		ClientManager.OnRemoteClientsChanged().AddRaw(this, &FAllClientsSelectionModel::OnRemoteClientsChanged);
+	}
+
+	FAllClientsSelectionModel::~FAllClientsSelectionModel()
+	{
+		ClientManager.OnRemoteClientsChanged().RemoveAll(this);
+	}
+
+	void FAllClientsSelectionModel::ForEachSelectedClient(TFunctionRef<EBreakBehavior(FReplicationClient&)> ProcessClient) const
+	{
+		ClientManager.ForEachClient(ProcessClient);
+	}
+
+	void FAllClientsSelectionModel::OnRemoteClientsChanged()
+	{
+		OnSelectionChangedDelegate.Broadcast();
+	}
+}
diff --git a/Engine/Plugins/Developer/Concert/ConcertApp/MultiUserClient/Source/MultiUserClient/Private/Widgets/ActiveSession/Replication/Client/Multi/AllClientsSelectionModel.h b/Engine/Plugins/Developer/Concert/ConcertApp/MultiUserClient/Source/MultiUserClient/Private/Widgets/ActiveSession/Replication/Client/Multi/AllClientsSelectionModel.h
new file mode 100644
--- /dev/null
+++ b/Engine/Plugins/Developer/Concert/ConcertApp/MultiUserClient/Source/MultiUserClient/Private/Widgets/ActiveSession/Replication/Client/Multi/AllClientsSelectionModel.h
@@ -0,0 +1,34 @@
+// Copyright Epic Games, Inc. All Rights Reserved.
+
+#pragma once
+
+#include "IClientSelectionModel.h"
+
+namespace UE::MultiUserClient
+{
+	class FReplicationClientManager;
+	
+	/** Exposes all clients and detects when clients disconnect. */
+	class FAllClientsSelectionModel : public IClientSelectionModel
+	{
+	public:
+		
+		FAllClientsSelectionModel(FReplicationClientManager& InClientManager);
+		virtual ~FAllClientsSelectionModel() override;
+
+		//~ Begin IClientSelectionModel Interface
+		virtual void ForEachSelectedClient(TFunctionRef<EBreakBehavior(FReplicationClient&)> ProcessClient) const override;
+		virtual FOnSelectionChanged& OnSelectionChanged() override { return OnSelectionChangedDelegate; }
+		//~ End IClientSelectionModel Interface
+
+	private:
+
+		/** Informs us when the list of clients changes */
+		FReplicationClientManager& ClientManager;
+		
+		/** Called when the clients ForEachSelectedClient enumerates has changed. */
+		FOnSelectionChanged OnSelectionChangedDelegate;
+		
+		void OnRemoteClientsChanged();
+	};
+}
diff --git a/Engine/Plugins/Developer/Concert/ConcertApp/MultiUserClient/Source/MultiUserClient/Private/Widgets/ActiveSession/Replication/Client/Multi/IClientSelectionModel.h b/Engine/Plugins/Developer/Concert/ConcertApp/MultiUserClient/Source/MultiUserClient/Private/Widgets/ActiveSession/Replication/Client/Multi/IClientSelectionModel.h
new file mode 100644
--- /dev/null
+++ b/Engine/Plugins/Developer/Concert/ConcertApp/MultiUserClient/Source/MultiUserClient/Private/Widgets/ActiveSession/Replication/Client/Multi/IClientSelectionModel.h
@@ -0,0 +1,29 @@
+// Copyright Epic Games, Inc. All Rights Reserved.
+
+#pragma once
+
+#include "Delegates/Delegate.h"
+#include "Templates/Function.h"
+
+enum class EBreakBehavior : uint8;
+
+namespace UE::MultiUserClient
+{
+	class FReplicationClient;
+	
+	/** Abstracts the concept of a client selection. */
+	class IClientSelectionModel
+	{
+	public:
+
+		/** Iterates through every client in the selection. */
+		virtual void ForEachSelectedClient(TFunctionRef<EBreakBehavior(FReplicationClient&)> ProcessClient) const = 0;
+
+		DECLARE_MULTICAST_DELEGATE(FOnSelectionChanged);
+		/** Called when the clients ForEachSelectedClient enumerates has changed. */
+		virtual FOnSelectionChanged& OnSelectionChanged() = 0;
+		
+		virtual ~IClientSelectionModel() = default;
+	};
+
+}
diff --git a/Engine/Plugins/Developer/Concert/ConcertApp/MultiUserClient/Source/MultiUserClient/Private/Widgets/ActiveSession/Replication/Client/Multi/MultiStreamModel.cpp b/Engine/Plugins/Developer/Concert/ConcertApp/MultiUserClient/Source/MultiUserClient/Private/Widgets/ActiveSession/Replication/Client/Multi/MultiStreamModel.cpp
new file mode 100644
--- /dev/null
+++ b/Engine/Plugins/Developer/Concert/ConcertApp/MultiUserClient/Source/MultiUserClient/Private/Widgets/ActiveSession/Replication/Client/Multi/MultiStreamModel.cpp
@@ -0,0 +1,104 @@
+// Copyright Epic Games, Inc. All Rights Reserved.
+
+#include "MultiStreamModel.h"
+
+#include "IClientSelectionModel.h"
+#include "Replication/Client/ReplicationClient.h"
+#include "Replication/Client/ReplicationClientManager.h"
+#include "Replication/Editor/Model/IEditableReplicationStreamModel.h"
+
+namespace UE::MultiUserClient
+{
+	FMultiStreamModel::FMultiStreamModel(IClientSelectionModel& InClientSelectionModel, FReplicationClientManager& InClientManager)
+		: ClientSelectionModel(InClientSelectionModel)
+		, ClientManager(InClientManager)
+	{
+		ClientSelectionModel.OnSelectionChanged().AddRaw(this, &FMultiStreamModel::RebuildStreamsSets);
+		RebuildStreamsSets();
+	}
+
+	void FMultiStreamModel::ForEachClient(TFunctionRef<EBreakBehavior(const FReplicationClient*)> ProcessClient) const
+	{
+		for (const FReplicationClient* ReadOnlyClient : GetCachedReadOnlyClients())
+		{
+			if (ProcessClient(ReadOnlyClient) == EBreakBehavior::Break)
+			{
+				return;
+			}
+		}
+		for (const FReplicationClient* WritableClient : GetCachedWritableClients())
+		{
+			if (ProcessClient(WritableClient) == EBreakBehavior::Break)
+			{
+				return;
+			}
+		}
+	}
+
+	TSet<TSharedRef<ConcertClientSharedSlate::IReplicationStreamModel>> FMultiStreamModel::GetReadOnlyStreams() const
+	{
+		TSet<TSharedRef<ConcertClientSharedSlate::IReplicationStreamModel>> Result;
+		Algo::Transform(CachedReadOnlyClients, Result, [](const FReplicationClient* Client){ return Client->GetClientEditModel(); });
+		return Result;
+	}
+
+	TSet<TSharedRef<ConcertClientSharedSlate::IEditableReplicationStreamModel>> FMultiStreamModel::GetEditableStreams() const
+	{
+		TSet<TSharedRef<ConcertClientSharedSlate::IEditableReplicationStreamModel>> Result;
+		Algo::Transform(CachedWritableClients, Result, [](const FReplicationClient* Client){ return Client->GetClientEditModel(); });
+		return Result;
+	}
+
+	void FMultiStreamModel::RebuildStreamsSets()
+	{
+		// Cannot just iterate through CachedReadOnlyClients because it may contain stale clients that were just removed
+		ClientManager.ForEachClient([this](FReplicationClient& Client)
+		{
+			Client.OnModelChanged().RemoveAll(this);
+			return EBreakBehavior::Continue;
+		});
+		
+		TSet<const FReplicationClient*> ReadOnlyClients;
+		TSet<const FReplicationClient*> WritableClients;
+		ClientSelectionModel.ForEachSelectedClient([this, &ReadOnlyClients, &WritableClients](FReplicationClient& Client)
+		{
+			const bool bIsUploadable = CanEverSubmit(Client.GetSubmissionWorkflow().GetUploadability());
+			if (bIsUploadable)
+			{
+				WritableClients.Add(&Client);
+			}
+			else
+			{
+				ReadOnlyClients.Add(&Client);
+				const TSharedRef<ConcertClientSharedSlate::IEditableReplicationStreamModel> Stream = Client.GetClientEditModel();
+				Client.OnModelChanged().AddRaw(this, &FMultiStreamModel::OnReadOnlyStreamChanged, Stream.ToWeakPtr());
+			}
+			
+			return EBreakBehavior::Continue;
+		});
+
+		const bool bReadOnlyStayedSame = CachedReadOnlyClients.Num() == ReadOnlyClients.Num() && CachedReadOnlyClients.Includes(ReadOnlyClients);
+		if (!bReadOnlyStayedSame)
+		{
+			CachedReadOnlyClients = MoveTemp(ReadOnlyClients);
+		}
+		const bool bWritableStayedSame = CachedWritableClients.Num() == WritableClients.Num() && CachedWritableClients.Includes(WritableClients);
+		if (!bWritableStayedSame)
+		{
+			CachedWritableClients = MoveTemp(WritableClients);
+		}
+
+		if (!bReadOnlyStayedSame || !bWritableStayedSame)
+		{
+			OnStreamSetChangedDelegate.Broadcast();
+		}
+	}
+
+	void FMultiStreamModel::OnReadOnlyStreamChanged(TWeakPtr<ConcertClientSharedSlate::IEditableReplicationStreamModel> ChangedStream)
+	{
+		if (const TSharedPtr<ConcertClientSharedSlate::IEditableReplicationStreamModel> ChangedStreamPin = ChangedStream.Pin())
+		{
+			OnReadOnlyStreamChangedDelegate.Broadcast(ChangedStreamPin.ToSharedRef());
+		}
+	}
+}
diff --git a/Engine/Plugins/Developer/Concert/ConcertApp/MultiUserClient/Source/MultiUserClient/Private/Widgets/ActiveSession/Replication/Client/Multi/MultiStreamModel.h b/Engine/Plugins/Developer/Concert/ConcertApp/MultiUserClient/Source/MultiUserClient/Private/Widgets/ActiveSession/Replication/Client/Multi/MultiStreamModel.h
new file mode 100644
--- /dev/null
+++ b/Engine/Plugins/Developer/Concert/ConcertApp/MultiUserClient/Source/MultiUserClient/Private/Widgets/ActiveSession/Replication/Client/Multi/MultiStreamModel.h
@@ -0,0 +1,56 @@
+// Copyright Epic Games, Inc. All Rights Reserved.
+
+#pragma once
+
+#include "Replication/Editor/Model/IEditableMultiReplicationStreamModel.h"
+
+#include "Containers/Set.h"
+#include "Templates/Function.h"
+#include "Templates/SharedPointer.h"
+
+enum class EBreakBehavior : uint8;
+
+namespace UE::MultiUserClient
+{
+	class FReplicationClient;
+	class FReplicationClientManager;
+	class IClientSelectionModel;
+
+	/** Checks whether clients accept remote changes and categorizes them into read-only and writable. */
+	class FMultiStreamModel : public ConcertClientSharedSlate::IEditableMultiReplicationStreamModel
+	{
+	public:
+		
+		FMultiStreamModel(IClientSelectionModel& InClientSelectionModel, FReplicationClientManager& InClientManager);
+		
+		const TSet<const FReplicationClient*>& GetCachedReadOnlyClients() const { return CachedReadOnlyClients; }
+		const TSet<const FReplicationClient*>& GetCachedWritableClients() const { return CachedWritableClients; }
+		void ForEachClient(TFunctionRef<EBreakBehavior(const FReplicationClient*)> ProcessClient) const;
+
+		//~ Begin IEditableMultiReplicationStreamModel Interface
+		virtual TSet<TSharedRef<ConcertClientSharedSlate::IReplicationStreamModel>> GetReadOnlyStreams() const override;
+		virtual TSet<TSharedRef<ConcertClientSharedSlate::IEditableReplicationStreamModel>> GetEditableStreams() const override;
+		virtual FOnReadOnlyStreamChanged& OnReadOnlyStreamChanged() override { return OnReadOnlyStreamChangedDelegate; }
+		virtual FOnStreamSetChanged& OnStreamSetChanged() override { return OnStreamSetChangedDelegate; }
+		//~ End IEditableMultiReplicationStreamModel Interface
+
+	private:
+		
+		/** Gets all clients to display and informs when they change. */
+		IClientSelectionModel& ClientSelectionModel;
+		/** Used to obtain a list of clients for unsubscribing */
+		FReplicationClientManager& ClientManager;
+
+		TSet<const FReplicationClient*> CachedReadOnlyClients;
+		TSet<const FReplicationClient*> CachedWritableClients;
+
+		FOnReadOnlyStreamChanged OnReadOnlyStreamChangedDelegate;
+		FOnStreamSetChanged OnStreamSetChangedDelegate;
+		
+		void RebuildStreamsSets();
+
+		/** Handle read-only streams changing */
+		void OnReadOnlyStreamChanged(TWeakPtr<ConcertClientSharedSlate::IEditableReplicationStreamModel> ChangedStream);
+	};
+}
+
diff --git a/Engine/Plugins/Developer/Concert/ConcertApp/MultiUserClient/Source/MultiUserClient/Private/Widgets/ActiveSession/Replication/Client/Multi/SAllClientsView.cpp b/Engine/Plugins/Developer/Concert/ConcertApp/MultiUserClient/Source/MultiUserClient/Private/Widgets/ActiveSession/Replication/Client/Multi/SAllClientsView.cpp
new file mode 100644
--- /dev/null
+++ b/Engine/Plugins/Developer/Concert/ConcertApp/MultiUserClient/Source/MultiUserClient/Private/Widgets/ActiveSession/Replication/Client/Multi/SAllClientsView.cpp
@@ -0,0 +1,35 @@
+// Copyright Epic Games, Inc. All Rights Reserved.
+
+#include "SAllClientsView.h"
+
+#include "AllClientsSelectionModel.h"
+#include "SMultiClientView.h"
+#include "Replication/Client/ReplicationClientManager.h"
+
+#include "Algo/Transform.h"
+
+namespace UE::MultiUserClient
+{
+	void SAllClientsView::Construct(const FArguments& InArgs, TSharedRef<IConcertClient> InConcertClient, FReplicationClientManager& InClientManager)
+	{
+		ClientManager = &InClientManager;
+		AllClientsModel = MakeUnique<FAllClientsSelectionModel>(InClientManager);
+		
+		ChildSlot
+		[
+			SNew(SMultiClientView, InConcertClient, InClientManager, *AllClientsModel)
+			.ViewSelectionArea() [ InArgs._ViewSelectionArea.Widget ]
+		];
+	}
+
+	TSet<const FReplicationClient*> SAllClientsView::GetAllClients() const
+	{
+		TSet<const FReplicationClient*> Result;
+		Algo::Transform(ClientManager->GetRemoteClients(), Result, [](const TNonNullPtr<FRemoteReplicationClient>& Client)
+		{
+			return Client.Get();
+		});
+		Result.Add(&ClientManager->GetLocalClient());
+		return Result;
+	}
+}
diff --git a/Engine/Plugins/Developer/Concert/ConcertApp/MultiUserClient/Source/MultiUserClient/Private/Widgets/ActiveSession/Replication/Client/Multi/SAllClientsView.h b/Engine/Plugins/Developer/Concert/ConcertApp/MultiUserClient/Source/MultiUserClient/Private/Widgets/ActiveSession/Replication/Client/Multi/SAllClientsView.h
new file mode 100644
--- /dev/null
+++ b/Engine/Plugins/Developer/Concert/ConcertApp/MultiUserClient/Source/MultiUserClient/Private/Widgets/ActiveSession/Replication/Client/Multi/SAllClientsView.h
@@ -0,0 +1,38 @@
+// Copyright Epic Games, Inc. All Rights Reserved.
+
+#pragma once
+
+#include "AllClientsSelectionModel.h"
+#include "Widgets/DeclarativeSyntaxSupport.h"
+#include "Widgets/SCompoundWidget.h"
+
+class IConcertClient;
+
+namespace UE::MultiUserClient
+{
+	class FReplicationClient;
+	class FReplicationClientManager;
+
+	/** Leverages SMultiClientView to display all replication clients. */
+	class SAllClientsView : public SCompoundWidget
+	{
+		SLATE_BEGIN_ARGS(SAllClientsView)
+			{}
+			/** Dedicated space for a widget with which to change the view. */
+			SLATE_NAMED_SLOT(FArguments, ViewSelectionArea)
+		SLATE_END_ARGS()
+
+		void Construct(const FArguments& InArgs, TSharedRef<IConcertClient> InConcertClient, FReplicationClientManager& InClientManager);
+
+	private:
+
+		/** Used to get all the replication clients and listen for client changes. */
+		FReplicationClientManager* ClientManager = nullptr;
+
+		/** Keeps the SMultiClientView updated of any changes to clients (e.g. disconnects, etc.) */
+		TUniquePtr<FAllClientsSelectionModel> AllClientsModel;
+
+		/** Gets all the clients to display */
+		TSet<const FReplicationClient*> GetAllClients() const;
+	};
+}
diff --git a/Engine/Plugins/Developer/Concert/ConcertApp/MultiUserClient/Source/MultiUserClient/Private/Widgets/ActiveSession/Replication/Client/Multi/SMultiClientView.cpp b/Engine/Plugins/Developer/Concert/ConcertApp/MultiUserClient/Source/MultiUserClient/Private/Widgets/ActiveSession/Replication/Client/Multi/SMultiClientView.cpp
new file mode 100644
--- /dev/null
+++ b/Engine/Plugins/Developer/Concert/ConcertApp/MultiUserClient/Source/MultiUserClient/Private/Widgets/ActiveSession/Replication/Client/Multi/SMultiClientView.cpp
@@ -0,0 +1,110 @@
+// Copyright Epic Games, Inc. All Rights Reserved.
+
+#include "SMultiClientView.h"
+
+#include "MultiStreamModel.h"
+#include "MultiUserReplicationSettings.h"
+#include "Replication/Client/ReplicationClient.h"
+#include "Replication/Client/ReplicationClientManager.h"
+#include "Replication/Editor/Model/Object/ActorSelectionSourceModel.h"
+#include "Replication/Editor/Model/Property/SelectPropertyFromUClassModel.h"
+#include "Replication/Editor/View/IMultiReplicationStreamEditor.h"
+#include "Replication/ReplicationWidgetFactories.h"
+#include "Widgets/ActiveSession/Replication/Client/Columns/MultiStreamColumns.h"
+#include "Widgets/ActiveSession/Replication/Client/SClientToolbar.h"
+
+#include "Algo/Transform.h"
+#include "UObject/Package.h"
+#include "Widgets/SBoxPanel.h"
+#include "Widgets/Text/STextBlock.h"
+
+#define LOCTEXT_NAMESPACE "SMultiClientView"
+
+namespace UE::MultiUserClient
+{
+	void SMultiClientView::Construct(const FArguments& InArgs, TSharedRef<IConcertClient> InConcertClient, FReplicationClientManager& ClientManager, IClientSelectionModel& InDisplayClientsModel)
+	{
+		StreamModel = MakeShared<FMultiStreamModel>(InDisplayClientsModel, ClientManager);
+
+		ChildSlot
+		[
+			SNew(SVerticalBox)
+
+			// Toolbar
+			+SVerticalBox::Slot()
+			.AutoHeight()
+			.Padding(2.f)
+			[
+				SAssignNew(Toolbar, SClientToolbar, ClientManager.GetAuthorityCache())
+				.ViewSelectionArea() [ InArgs._ViewSelectionArea.Widget ]
+				.DisplayedClients(this, &SMultiClientView::GetDisplayClientIds)
+				.ForEachReplicatedObject(this, &SMultiClientView::EnumerateObjectsInStreams)
+			]
+
+			// Editor
+			+SVerticalBox::Slot()
+			.FillHeight(1.f)
+			[
+				CreateEditorContent(InConcertClient, ClientManager)
+			]
+		];
+	}
+
+	TSharedRef<SWidget> SMultiClientView::CreateEditorContent(TSharedRef<IConcertClient> InConcertClient, FReplicationClientManager& InClientManager)
+	{
+		using namespace UE::ConcertClientSharedSlate;
+
+		TAttribute<const FConcertReplicationEditorSettings*> ReplicationSettingsAttribute =
+			TAttribute<const FConcertReplicationEditorSettings*>::CreateLambda([]()
+			{
+				return &UMultiUserReplicationSettings::Get()->ReplicationEditorSettings;
+			});
+		TAttribute<TSharedPtr<IMultiReplicationStreamEditor>> MultiStreamEditorAttribute =
+		   TAttribute<TSharedPtr<IMultiReplicationStreamEditor>>::CreateLambda([this]()
+		   {
+			   return StreamEditor;
+		   });
+		
+		FCreateMultiStreamEditorParams Params
+		{
+			.MultiStreamModel = StreamModel.ToSharedRef(),
+			.ObjectSource = MakeShared<FActorSelectionSourceModel>(),
+			.PropertySource = MakeShared<FSelectPropertyFromUClassModel>(),
+			.ReplicationSettingsAttribute = MoveTemp(ReplicationSettingsAttribute),
+			.ViewerParams =
+			{
+				.SubobjectModel = CreateDefaultComponentHierarchySubobjectModel(), // This makes actors have children in the top view
+				.AdditionalPropertyColumns = { MultiStreamColumns::AssignPropertyColumn(MultiStreamEditorAttribute, InConcertClient, InClientManager) }
+			}
+		};
+		StreamEditor = CreateBaseMultiStreamEditor(MoveTemp(Params));
+		check(StreamEditor);
+		return StreamEditor.ToSharedRef();
+	}
+
+	TSet<FGuid> SMultiClientView::GetDisplayClientIds() const
+	{
+		TSet<FGuid> ClientIds;
+		StreamModel->ForEachClient([&ClientIds](const FReplicationClient* Client)
+		{
+			ClientIds.Add(Client->GetEndpointId());
+			return EBreakBehavior::Continue;
+		});
+		return ClientIds;
+	}
+
+	void SMultiClientView::EnumerateObjectsInStreams(TFunctionRef<void(const FSoftObjectPath&)> Consumer)
+	{
+		StreamModel->ForEachClient([&Consumer](const FReplicationClient* Client)
+		{
+			Client->GetClientEditModel()->ForEachReplicatedObject([&Consumer](const FSoftObjectPath& Object)
+			{
+				Consumer(Object);
+				return EBreakBehavior::Continue;
+			});
+			return EBreakBehavior::Continue;
+		});
+	}
+}
+
+#undef LOCTEXT_NAMESPACE
\ No newline at end of file
diff --git a/Engine/Plugins/Developer/Concert/ConcertApp/MultiUserClient/Source/MultiUserClient/Private/Widgets/ActiveSession/Replication/Client/Multi/SMultiClientView.h b/Engine/Plugins/Developer/Concert/ConcertApp/MultiUserClient/Source/MultiUserClient/Private/Widgets/ActiveSession/Replication/Client/Multi/SMultiClientView.h
new file mode 100644
--- /dev/null
+++ b/Engine/Plugins/Developer/Concert/ConcertApp/MultiUserClient/Source/MultiUserClient/Private/Widgets/ActiveSession/Replication/Client/Multi/SMultiClientView.h
@@ -0,0 +1,58 @@
+// Copyright Epic Games, Inc. All Rights Reserved.
+
+#pragma once
+
+#include "Widgets/DeclarativeSyntaxSupport.h"
+#include "Widgets/SCompoundWidget.h"
+
+namespace UE::MultiUserClient
+{
+	class IClientSelectionModel;
+}
+
+class IConcertClient;
+
+namespace UE::ConcertClientSharedSlate
+{
+	class IMultiReplicationStreamEditor;
+	class IEditableReplicationStreamModel;
+}
+
+namespace UE::MultiUserClient
+{
+	class FGlobalAuthorityCache;
+	class FMultiStreamModel;
+	class FReplicationClient;
+	class FReplicationClientManager;
+	class SClientToolbar;
+
+	/** Displays a selection of clients. */
+	class SMultiClientView
+		: public SCompoundWidget
+	{
+		SLATE_BEGIN_ARGS(SMultiClientView)
+		{}
+			/** Dedicated space for a widget with which to change the view. */
+			SLATE_NAMED_SLOT(FArguments, ViewSelectionArea)
+		SLATE_END_ARGS()
+
+		void Construct(const FArguments& InArgs, TSharedRef<IConcertClient> InConcertClient, FReplicationClientManager& InClientManager, IClientSelectionModel& InDisplayClientsModel);
+
+	private:
+
+		/** Used to access and update the replication status widget. */
+		TSharedPtr<SClientToolbar> Toolbar;
+		
+		/** Combines the clients */
+		TSharedPtr<FMultiStreamModel> StreamModel;
+		/** Displayed in the UI. */
+		TSharedPtr<ConcertClientSharedSlate::IMultiReplicationStreamEditor> StreamEditor;
+
+		/** Creates this widget's editor content */
+		TSharedRef<SWidget> CreateEditorContent(TSharedRef<IConcertClient> InConcertClient, FReplicationClientManager& InClientManager);
+
+		// SClientToolbar attributes
+		TSet<FGuid> GetDisplayClientIds() const;
+		void EnumerateObjectsInStreams(TFunctionRef<void(const FSoftObjectPath&)> Consumer);
+	};
+}
diff --git a/Engine/Plugins/Developer/Concert/ConcertApp/MultiUserClient/Source/MultiUserClient/Private/Widgets/ActiveSession/Replication/Client/SClientToolbar.cpp b/Engine/Plugins/Developer/Concert/ConcertApp/MultiUserClient/Source/MultiUserClient/Private/Widgets/ActiveSession/Replication/Client/SClientToolbar.cpp
--- a/Engine/Plugins/Developer/Concert/ConcertApp/MultiUserClient/Source/MultiUserClient/Private/Widgets/ActiveSession/Replication/Client/SClientToolbar.cpp
+++ b/Engine/Plugins/Developer/Concert/ConcertApp/MultiUserClient/Source/MultiUserClient/Private/Widgets/ActiveSession/Replication/Client/SClientToolbar.cpp
@@ -9,7 +9,7 @@
 
 namespace UE::MultiUserClient
 {
-	void SClientToolbar::Construct(const FArguments& InArgs, const ConcertClientSharedSlate::IReplicationStreamModel& InObjectModel, FGlobalAuthorityCache& InAuthorityCache)
+	void SClientToolbar::Construct(const FArguments& InArgs, FGlobalAuthorityCache& InAuthorityCache)
 	{
 		ChildSlot
 		[
@@ -33,9 +33,15 @@ namespace UE::MultiUserClient
 			.VAlign(VAlign_Center)
 			.Padding(0.f, 0.f, 5.f, 0.f)
 			[
-				SNew(SReplicationStatus, InObjectModel, InAuthorityCache)
+				SAssignNew(ReplicationStatus, SReplicationStatus, InAuthorityCache)
 				.DisplayedClients(InArgs._DisplayedClients)
+				.ForEachReplicatedObject(InArgs._ForEachReplicatedObject)
 			]
 		];
 	}
-}
\ No newline at end of file
+
+	void SClientToolbar::RefreshStatusText()
+	{
+		ReplicationStatus->RefreshStatusText();
+	}
+}
diff --git a/Engine/Plugins/Developer/Concert/ConcertApp/MultiUserClient/Source/MultiUserClient/Private/Widgets/ActiveSession/Replication/Client/SClientToolbar.h b/Engine/Plugins/Developer/Concert/ConcertApp/MultiUserClient/Source/MultiUserClient/Private/Widgets/ActiveSession/Replication/Client/SClientToolbar.h
--- a/Engine/Plugins/Developer/Concert/ConcertApp/MultiUserClient/Source/MultiUserClient/Private/Widgets/ActiveSession/Replication/Client/SClientToolbar.h
+++ b/Engine/Plugins/Developer/Concert/ConcertApp/MultiUserClient/Source/MultiUserClient/Private/Widgets/ActiveSession/Replication/Client/SClientToolbar.h
@@ -2,6 +2,8 @@
 
 #pragma once
 
+#include "SReplicationStatus.h"
+
 #include "Widgets/DeclarativeSyntaxSupport.h"
 #include "Widgets/SCompoundWidget.h"
 
@@ -15,6 +17,7 @@ namespace UE::MultiUserClient
 	class FGlobalAuthorityCache;
 	class FStreamChangeTracker;
 	class ISubmissionWorkflow;
+	class SReplicationStatus;
 
 	/**
 	 * Contains a bunch of actions that can be performed on client view.
@@ -30,9 +33,20 @@ namespace UE::MultiUserClient
 		
 			/** The clients to show statistics for */
 			SLATE_ATTRIBUTE(TSet<FGuid>, DisplayedClients)
+		
+			/** Delegate which enumerates every replicated object. */
+			SLATE_EVENT(FForEachReplicatedObject, ForEachReplicatedObject)
 		SLATE_END_ARGS()
 
-		void Construct(const FArguments& InArgs, const ConcertClientSharedSlate::IReplicationStreamModel& InObjectModel, FGlobalAuthorityCache& InAuthorityCache);
+		void Construct(const FArguments& InArgs, FGlobalAuthorityCache& InAuthorityCache);
+
+		/** Updates the status text of how many objects are being replicated. */
+		void RefreshStatusText();
+		
+	private:
+
+		/** Displays how many objects are being replicated */
+		TSharedPtr<SReplicationStatus> ReplicationStatus;
 	};
 }
 
diff --git a/Engine/Plugins/Developer/Concert/ConcertApp/MultiUserClient/Source/MultiUserClient/Private/Widgets/ActiveSession/Replication/Client/SReplicationStatus.cpp b/Engine/Plugins/Developer/Concert/ConcertApp/MultiUserClient/Source/MultiUserClient/Private/Widgets/ActiveSession/Replication/Client/SReplicationStatus.cpp
--- a/Engine/Plugins/Developer/Concert/ConcertApp/MultiUserClient/Source/MultiUserClient/Private/Widgets/ActiveSession/Replication/Client/SReplicationStatus.cpp
+++ b/Engine/Plugins/Developer/Concert/ConcertApp/MultiUserClient/Source/MultiUserClient/Private/Widgets/ActiveSession/Replication/Client/SReplicationStatus.cpp
@@ -45,12 +45,14 @@ namespace UE::MultiUserClient
 		}
 	}
 	
-	void SReplicationStatus::Construct(const FArguments& InArgs, const ConcertClientSharedSlate::IReplicationStreamModel& InObjectModel, FGlobalAuthorityCache& InAuthorityCache)
+	void SReplicationStatus::Construct(const FArguments& InArgs, FGlobalAuthorityCache& InAuthorityCache)
 	{
-		ObjectModel = &InObjectModel;
 		AuthorityCache = &InAuthorityCache;
+		
 		DisplayedClientsAttribute = InArgs._DisplayedClients;
 		check(DisplayedClientsAttribute.IsBound() || DisplayedClientsAttribute.IsSet());
+		ForEachReplicatedObjectDelegate = InArgs._ForEachReplicatedObject;
+		check(ForEachReplicatedObjectDelegate.IsBound());
 
 		AuthorityCache->OnCacheChanged().AddSP(this, &SReplicationStatus::OnAuthorityCacheChanged);
 		ChildSlot
@@ -96,7 +98,7 @@ namespace UE::MultiUserClient
 
 		TSet<FSoftObjectPath> ReplicatedActors;
 		TSet<FSoftObjectPath> ReplicatedObjects;
-		ObjectModel->ForEachReplicatedObject([this, &DisplayedClients, &ReplicatedActors, &ReplicatedObjects](const FSoftObjectPath& Path)
+		ForEachReplicatedObjectDelegate.Execute([this, &DisplayedClients, &ReplicatedActors, &ReplicatedObjects](const FSoftObjectPath& Path)
 		{
 			const bool bIsReplicated = Algo::AnyOf(AuthorityCache->GetClientsWithAuthorityOverObject(Path), [&DisplayedClients](const FGuid& ClientId)
 			{
diff --git a/Engine/Plugins/Developer/Concert/ConcertApp/MultiUserClient/Source/MultiUserClient/Private/Widgets/ActiveSession/Replication/Client/SReplicationStatus.h b/Engine/Plugins/Developer/Concert/ConcertApp/MultiUserClient/Source/MultiUserClient/Private/Widgets/ActiveSession/Replication/Client/SReplicationStatus.h
--- a/Engine/Plugins/Developer/Concert/ConcertApp/MultiUserClient/Source/MultiUserClient/Private/Widgets/ActiveSession/Replication/Client/SReplicationStatus.h
+++ b/Engine/Plugins/Developer/Concert/ConcertApp/MultiUserClient/Source/MultiUserClient/Private/Widgets/ActiveSession/Replication/Client/SReplicationStatus.h
@@ -16,6 +16,8 @@ namespace UE::MultiUserClient
 {
 	class FGlobalAuthorityCache;
 
+	DECLARE_DELEGATE_OneParam(FForEachReplicatedObject, TFunctionRef<void(const FSoftObjectPath&)> Consumer);
+	
 	/**
 	 * Displays a text "Replicating x Objects for y Actors".
 	 *
@@ -31,20 +33,26 @@ namespace UE::MultiUserClient
 		{}
 			/** The clients to show statistics for */
 			SLATE_ATTRIBUTE(TSet<FGuid>, DisplayedClients)
+
+			/** Delegate which enumerates every replicated object. */
+			SLATE_EVENT(FForEachReplicatedObject, ForEachReplicatedObject)
 		SLATE_END_ARGS()
 
-		void Construct(const FArguments& InArgs, const ConcertClientSharedSlate::IReplicationStreamModel& InObjectModel, FGlobalAuthorityCache& InAuthorityCache);
+		void Construct(const FArguments& InArgs, FGlobalAuthorityCache& InAuthorityCache);
 		virtual ~SReplicationStatus() override;
 
+		/** Updates the status text after an external update has occured. */
+		void RefreshStatusText();
+
 	private:
 
-		/** Used to look up registered objects and subobjects. */
-		const ConcertClientSharedSlate::IReplicationStreamModel* ObjectModel;
 		/** Used to get authority state of objects and informs us when authority changes. */
-		FGlobalAuthorityCache* AuthorityCache;
+		FGlobalAuthorityCache* AuthorityCache = nullptr;
 
 		/** The clients to show statistics for */
 		TAttribute<TSet<FGuid>> DisplayedClientsAttribute;
+		/** Delegate which enumerates every replicated object. */
+		FForEachReplicatedObject ForEachReplicatedObjectDelegate;
 
 		/** Updated when authority changes. Displays subobjects in bold. */
 		TSharedPtr<STextBlock> ObjectsText;
@@ -52,7 +60,6 @@ namespace UE::MultiUserClient
 		TSharedPtr<STextBlock> ActorsText;
 
 		void OnAuthorityCacheChanged(const FGuid& ClientId) { RefreshStatusText(); }
-		void RefreshStatusText();
 	};
 }
 
diff --git a/Engine/Plugins/Developer/Concert/ConcertApp/MultiUserClient/Source/MultiUserClient/Private/Widgets/ActiveSession/Replication/Client/Single/SReplicationClientView.cpp b/Engine/Plugins/Developer/Concert/ConcertApp/MultiUserClient/Source/MultiUserClient/Private/Widgets/ActiveSession/Replication/Client/Single/SReplicationClientView.cpp
--- a/Engine/Plugins/Developer/Concert/ConcertApp/MultiUserClient/Source/MultiUserClient/Private/Widgets/ActiveSession/Replication/Client/Single/SReplicationClientView.cpp
+++ b/Engine/Plugins/Developer/Concert/ConcertApp/MultiUserClient/Source/MultiUserClient/Private/Widgets/ActiveSession/Replication/Client/Single/SReplicationClientView.cpp
@@ -4,7 +4,7 @@
 
 #include "MultiUserReplicationSettings.h"
 #include "Replication/Client/ReplicationClient.h"
-#include "Replication/Editor/Model/IReplicationStreamModel.h"
+#include "Replication/Editor/Model/IEditableReplicationStreamModel.h"
 #include "Replication/Editor/Model/Object/ActorSelectionSourceModel.h"
 #include "Replication/Editor/Model/Property/SelectPropertyFromUClassModel.h"
 #include "Replication/Editor/View/IReplicationStreamEditor.h"
@@ -37,7 +37,6 @@ namespace UE::MultiUserClient
 		FReplicationClient* ReplicationClient = GetReplicationClientAttribute.Get();
 		check(ReplicationClient);
 		
-		ConcertClientSharedSlate::IReplicationStreamModel& PropertyModel = *ReplicationClient->GetClientEditModel();
 		FGlobalAuthorityCache& AuthorityCache = ClientManager->GetAuthorityCache();
 		ChildSlot
 		[
@@ -48,9 +47,10 @@ namespace UE::MultiUserClient
 			.AutoHeight()
 			.Padding(2.f)
 			[
-				SNew(SClientToolbar, PropertyModel, AuthorityCache)
+				SNew(SClientToolbar, AuthorityCache)
 				.ViewSelectionArea() [ InArgs._ViewSelectionArea.Widget ]
 				.DisplayedClients(TSet{ ReplicationClient->GetEndpointId() })
+				.ForEachReplicatedObject(this, &SReplicationClientView::EnumerateReplicatedObjects)
 			]
 
 			// Editor
@@ -62,7 +62,7 @@ namespace UE::MultiUserClient
 		];
 		
 		// Refresh UI if streams change externally, e.g. a remote client changed what they sent
-		ReplicationClient->OnModelExternallyChanged().AddSP(this, &SReplicationClientView::OnModelChanged);
+		ReplicationClient->OnModelChanged().AddSP(this, &SReplicationClientView::OnModelChanged);
 
 		CVarReplicationClientViewMode->OnChangedDelegate().AddSP(this, &SReplicationClientView::OnConsoleVariableChanged);
 	}
@@ -94,9 +94,7 @@ namespace UE::MultiUserClient
 		}
 	}
 
-	TSharedRef<SWidget> SReplicationClientView::CreateThreeSectionedContent(
-		FReplicationClient& InReplicationClient
-		)
+	TSharedRef<SWidget> SReplicationClientView::CreateThreeSectionedContent(FReplicationClient& InReplicationClient)
 	{
 		using namespace ConcertClientSharedSlate;
 		IReplicationStreamModel& PropertyModel = *InReplicationClient.GetClientEditModel();
@@ -124,32 +122,32 @@ namespace UE::MultiUserClient
 			.DataModel = InReplicationClient.GetClientEditModel(),
 			.ObjectSource = MakeShared<FActorSelectionSourceModel>(),
 			.PropertySource = MakeShared<FSelectPropertyFromUClassModel>(),
-			.SubobjectView = CreateUnrealEditorSubobjectView(SubobjectViewParams),
-			.AdditionalObjectColumns =
-			{
-				SingleClientColumns::ToggleTopLevelAuthority(PropertyModel, AuthorityTracker, SubmissionWorkflow),
-				SingleClientColumns::OwnerOfTopLevelObject(ConcertClient.ToSharedRef(), AuthorityCache, PropertyModel)
-			},
-			.AdditionalPropertyColumns =
-			{
-				SingleClientColumns::OwnerOfProperty(ConcertClient.ToSharedRef(), AuthorityCache, GetReplicationViewerAttribute),
-				SingleClientColumns::ConflictWarningForProperty(ConcertClient.ToSharedRef(), GetReplicationViewerAttribute, AuthorityCache, InReplicationClient.GetEndpointId())
-			},
 			.IsEditingEnabled = TAttribute<bool>::CreateLambda([&SubmissionWorkflow](){ return SubmissionWorkflow.GetUploadability() != EChangeUploadability::NotImplemented; }),
 			.EditingDisabledToolTipText = LOCTEXT("Editing.NotImplemented", "Editing remote clients is not implemented. You can only edit the local client."),
-			.ReplicationSettingsAttribute = MoveTemp(ReplicationSettingsAttribute)
+			.ReplicationSettingsAttribute = MoveTemp(ReplicationSettingsAttribute),
+			.ViewerParams =
+			{
+				.SubobjectView = CreateUnrealEditorSubobjectView(SubobjectViewParams),
+				.AdditionalObjectColumns =
+				{
+					SingleClientColumns::ToggleTopLevelAuthority(PropertyModel, AuthorityTracker, SubmissionWorkflow),
+					SingleClientColumns::OwnerOfTopLevelObject(ConcertClient.ToSharedRef(), AuthorityCache, PropertyModel)
+				},
+				.AdditionalPropertyColumns =
+				{
+					SingleClientColumns::OwnerOfProperty(ConcertClient.ToSharedRef(), AuthorityCache, GetReplicationViewerAttribute),
+					SingleClientColumns::ConflictWarningForProperty(ConcertClient.ToSharedRef(), GetReplicationViewerAttribute, AuthorityCache, InReplicationClient.GetEndpointId())
+				}
+			}
 		};
 
 		EditorView_ThreeSectioned = CreateDefaultStreamEditor(ReplicationEditorCreationParams);
 		return EditorView_ThreeSectioned.ToSharedRef();
 	}
 
-	TSharedRef<SWidget> SReplicationClientView::CreateTwoSectionedContent(
-		FReplicationClient& InReplicationClient
-		)
+	TSharedRef<SWidget> SReplicationClientView::CreateTwoSectionedContent(FReplicationClient& InReplicationClient)
 	{
 		using namespace ConcertClientSharedSlate;
-		IReplicationStreamModel& PropertyModel = *InReplicationClient.GetClientEditModel();
 		FAuthorityChangeTracker& AuthorityTracker = InReplicationClient.GetAuthorityDiffer();
 		ISubmissionWorkflow& SubmissionWorkflow = InReplicationClient.GetSubmissionWorkflow();
 		FGlobalAuthorityCache& AuthorityCache = ClientManager->GetAuthorityCache();
@@ -165,22 +163,25 @@ namespace UE::MultiUserClient
 			.DataModel = InReplicationClient.GetClientEditModel(),
 			.ObjectSource = MakeShared<FActorSelectionSourceModel>(),
 			.PropertySource = MakeShared<FSelectPropertyFromUClassModel>(),
-			.SubobjectModel = CreateDefaultComponentHierarchySubobjectModel(), // This makes actors have children in the top view
-			.AdditionalObjectColumns =
-			{
-				// TODO DP: When we decide which UI version to choose, remove subobject from these function names.
-				SingleClientColumns::ToggleSubobjectAuthority(AuthorityTracker, SubmissionWorkflow),
-				SingleClientColumns::ConflictWarningForSubobject(ConcertClient.ToSharedRef(), AuthorityCache, InReplicationClient.GetEndpointId()),
-				SingleClientColumns::OwnerOfSubobject(ConcertClient.ToSharedRef(), AuthorityCache)
-			},
-			.AdditionalPropertyColumns =
-			{
-				SingleClientColumns::OwnerOfProperty(ConcertClient.ToSharedRef(), AuthorityCache, GetReplicationViewerAttribute),
-				SingleClientColumns::ConflictWarningForProperty(ConcertClient.ToSharedRef(), GetReplicationViewerAttribute, AuthorityCache, InReplicationClient.GetEndpointId())
-			},
 			.IsEditingEnabled = TAttribute<bool>::CreateLambda([&SubmissionWorkflow](){ return SubmissionWorkflow.GetUploadability() != EChangeUploadability::NotImplemented; }),
 			.EditingDisabledToolTipText = LOCTEXT("Editing.NotImplemented", "Editing remote clients is not implemented. You can only edit the local client."),
-			.ReplicationSettingsAttribute = MoveTemp(ReplicationSettingsAttribute)
+			.ReplicationSettingsAttribute = MoveTemp(ReplicationSettingsAttribute),
+			.ViewerParams =
+			{
+				.SubobjectModel = CreateDefaultComponentHierarchySubobjectModel(), // This makes actors have children in the top view
+				.AdditionalObjectColumns =
+				{
+					// TODO DP: When we decide which UI version to choose, remove subobject from these function names.
+					SingleClientColumns::ToggleSubobjectAuthority(AuthorityTracker, SubmissionWorkflow),
+					SingleClientColumns::ConflictWarningForSubobject(ConcertClient.ToSharedRef(), AuthorityCache, InReplicationClient.GetEndpointId()),
+					SingleClientColumns::OwnerOfSubobject(ConcertClient.ToSharedRef(), AuthorityCache)
+				},
+				.AdditionalPropertyColumns =
+				{
+					SingleClientColumns::OwnerOfProperty(ConcertClient.ToSharedRef(), AuthorityCache, GetReplicationViewerAttribute),
+					SingleClientColumns::ConflictWarningForProperty(ConcertClient.ToSharedRef(), GetReplicationViewerAttribute, AuthorityCache, InReplicationClient.GetEndpointId())
+				},
+			}
 		};
 
 		EditorView_TwoSectioned = CreateDefaultStreamEditor(ReplicationEditorCreationParams);
@@ -204,6 +205,18 @@ namespace UE::MultiUserClient
 	{
 		RebuildContent();
 	}
+
+	void SReplicationClientView::EnumerateReplicatedObjects(TFunctionRef<void(const FSoftObjectPath&)> Consumer) const
+	{
+		const FReplicationClient* ReplicationClient = GetReplicationClientAttribute.Get();
+		check(ReplicationClient);
+		
+		ReplicationClient->GetClientEditModel()->ForEachReplicatedObject([&Consumer](const FSoftObjectPath& Object)
+		{
+			Consumer(Object);
+			return EBreakBehavior::Continue;
+		});
+	}
 }
 
 #undef LOCTEXT_NAMESPACE
\ No newline at end of file
diff --git a/Engine/Plugins/Developer/Concert/ConcertApp/MultiUserClient/Source/MultiUserClient/Private/Widgets/ActiveSession/Replication/Client/Single/SReplicationClientView.h b/Engine/Plugins/Developer/Concert/ConcertApp/MultiUserClient/Source/MultiUserClient/Private/Widgets/ActiveSession/Replication/Client/Single/SReplicationClientView.h
--- a/Engine/Plugins/Developer/Concert/ConcertApp/MultiUserClient/Source/MultiUserClient/Private/Widgets/ActiveSession/Replication/Client/Single/SReplicationClientView.h
+++ b/Engine/Plugins/Developer/Concert/ConcertApp/MultiUserClient/Source/MultiUserClient/Private/Widgets/ActiveSession/Replication/Client/Single/SReplicationClientView.h
@@ -61,6 +61,9 @@ namespace UE::MultiUserClient
 		void OnModelChanged() const;
 		
 		void OnConsoleVariableChanged(IConsoleVariable* ConsoleVariable);
+
+		/** Lists all objects being replicated */
+		void EnumerateReplicatedObjects(TFunctionRef<void(const FSoftObjectPath&)> Consumer) const;
 	};
 }
 
diff --git a/Engine/Plugins/Developer/Concert/ConcertApp/MultiUserClient/Source/MultiUserClient/Private/Widgets/ActiveSession/Replication/Joined/EClientViewType.h b/Engine/Plugins/Developer/Concert/ConcertApp/MultiUserClient/Source/MultiUserClient/Private/Widgets/ActiveSession/Replication/Joined/EClientViewType.h
new file mode 100644
--- /dev/null
+++ b/Engine/Plugins/Developer/Concert/ConcertApp/MultiUserClient/Source/MultiUserClient/Private/Widgets/ActiveSession/Replication/Joined/EClientViewType.h
@@ -0,0 +1,15 @@
+// Copyright Epic Games, Inc. All Rights Reserved.
+
+#pragma once
+
+namespace UE::MultiUserClient
+{
+	/** Describes a type of view that is being displayed in SReplicationJoinedView and can get selected by the combo button. */
+	enum class EClientViewType 
+	{
+		// When changing this enum, make sure to update the slot order of SReplicationJoinedView's widget switcher
+		AllClients = 0,
+		LocalClient = 1,
+		RemoteClient = 2
+	};
+}
\ No newline at end of file
diff --git a/Engine/Plugins/Developer/Concert/ConcertApp/MultiUserClient/Source/MultiUserClient/Private/Widgets/ActiveSession/Replication/Joined/SReplicationJoinedView.cpp b/Engine/Plugins/Developer/Concert/ConcertApp/MultiUserClient/Source/MultiUserClient/Private/Widgets/ActiveSession/Replication/Joined/SReplicationJoinedView.cpp
--- a/Engine/Plugins/Developer/Concert/ConcertApp/MultiUserClient/Source/MultiUserClient/Private/Widgets/ActiveSession/Replication/Joined/SReplicationJoinedView.cpp
+++ b/Engine/Plugins/Developer/Concert/ConcertApp/MultiUserClient/Source/MultiUserClient/Private/Widgets/ActiveSession/Replication/Joined/SReplicationJoinedView.cpp
@@ -5,6 +5,7 @@
 #include "IConcertSyncClient.h"
 #include "Replication/MultiUserReplicationManager.h"
 #include "SSelectClientViewComboButton.h"
+#include "Widgets/ActiveSession/Replication/Client/Multi/SAllClientsView.h"
 #include "Widgets/ActiveSession/Replication/Client/Single/SReplicationClientView.h"
 
 #include "Replication/Client/RemoteReplicationClient.h"
@@ -29,6 +30,15 @@ namespace UE::MultiUserClient
 		ChildSlot
 		[
 			SAssignNew(ClientViewSwitcher, SWidgetSwitcher)
+			// These must be aligned with EClientViewType
+			+SWidgetSwitcher::Slot()
+			[
+				SNew(SAllClientsView, InClient->GetConcertClient(), *ReplicationManager->GetClientManager())
+				.ViewSelectionArea()
+				[
+					MakeClientSelectionArea()
+				]
+			]
 			+SWidgetSwitcher::Slot()
 			[
 				SNew(SReplicationClientView, InClient->GetConcertClient(), *ReplicationManager->GetClientManager())
@@ -70,19 +80,25 @@ namespace UE::MultiUserClient
 			OldClientWidgets.Add(Widget);
 			ClientViewSwitcher->RemoveSlot(Widget);
 		}
-		checkf(OldClientWidgets.Num() >= 1, TEXT("Was supposed to contain local client widget"));
+		checkf(OldClientWidgets.Num() >= 2, TEXT("Was supposed to contain all clients and local client widgets"));
 
+		// All clients
 		ClientViewSwitcher->AddSlot()
 			[
 				OldClientWidgets[0]
 			];
+		// Local client
+		ClientViewSwitcher->AddSlot()
+			[
+				OldClientWidgets[1]
+			];
 		RebuildClientViewSwitcherChildren(OldClientWidgets);
 
 		// Need to change view if the displayed client was removed
 		const bool bDisplayedClientWasRemoved = RemoteClientToWidgetSwitcherIndex.Contains(DisplayedClientId); 
 		if (bDisplayedClientWasRemoved)
 		{
-			// Will show the local client
+			// Will show all clients
 			ClientViewSwitcher->SetActiveWidgetIndex(0);
 		}
 		
@@ -182,13 +198,26 @@ namespace UE::MultiUserClient
 					? *RemoteClient
 					: Client->GetConcertClient()->GetCurrentSession()->GetSessionClientEndpointId();
 			})
+			.CurrentDisplayMode_Lambda([this]()
+			{
+				switch(ClientViewSwitcher->GetActiveWidgetIndex())
+				{
+				case 0: return EClientViewType::AllClients;
+				case 1: return EClientViewType::LocalClient;
+				default: return EClientViewType::RemoteClient;
+				}
+			})
 			.OnSelectClient_Lambda([this](const FGuid& InClientId)
 			{
 				const int32* RemoteClientIndex = RemoteClientToWidgetSwitcherIndex.Find(InClientId);
 				// If it's not a remote client, it's the local client
-				const int32 ActiveWidgetIndex = RemoteClientIndex ? *RemoteClientIndex : 0;
+				const int32 ActiveWidgetIndex = RemoteClientIndex ? *RemoteClientIndex : 1;
 				ClientViewSwitcher->SetActiveWidgetIndex(ActiveWidgetIndex);
 			})
+			.OnSelectAllClients_Lambda([this]()
+			{
+				ClientViewSwitcher->SetActiveWidgetIndex(0);
+			})
 		;
 	}
 
diff --git a/Engine/Plugins/Developer/Concert/ConcertApp/MultiUserClient/Source/MultiUserClient/Private/Widgets/ActiveSession/Replication/Joined/SSelectClientViewComboButton.cpp b/Engine/Plugins/Developer/Concert/ConcertApp/MultiUserClient/Source/MultiUserClient/Private/Widgets/ActiveSession/Replication/Joined/SSelectClientViewComboButton.cpp
--- a/Engine/Plugins/Developer/Concert/ConcertApp/MultiUserClient/Source/MultiUserClient/Private/Widgets/ActiveSession/Replication/Joined/SSelectClientViewComboButton.cpp
+++ b/Engine/Plugins/Developer/Concert/ConcertApp/MultiUserClient/Source/MultiUserClient/Private/Widgets/ActiveSession/Replication/Joined/SSelectClientViewComboButton.cpp
@@ -11,6 +11,8 @@
 #include "Widgets/Input/SComboButton.h"
 #include "Widgets/Layout/SWidgetSwitcher.h"
 
+#define LOCTEXT_NAMESPACE "SSelectClientViewComboButton"
+
 namespace UE::MultiUserClient
 {
 	void SSelectClientViewComboButton::Construct(const FArguments& InArgs)
@@ -18,7 +20,11 @@ namespace UE::MultiUserClient
 		Client = InArgs._Client;
 		ClientsAttribute = InArgs._SelectableClients;
 		CurrentSelection = InArgs._CurrentSelection;
+		CurrentDisplayMode = InArgs._CurrentDisplayMode;
+		
 		OnSelectClientDelegate = InArgs._OnSelectClient;
+		OnSelectAllClients = InArgs._OnSelectAllClients;
+		check(OnSelectClientDelegate.IsBound() && OnSelectAllClients.IsBound());
 		
 		ChildSlot
 		[
@@ -29,6 +35,10 @@ namespace UE::MultiUserClient
 				SAssignNew(ButtonContent, SWidgetSwitcher)
 				.WidgetIndex(this, &SSelectClientViewComboButton::GetActiveWidgetIndex)
 
+				+SWidgetSwitcher::Slot()
+				[
+					MakeAllClientsDisplayWidget()
+				]
 				+SWidgetSwitcher::Slot()
 				[
 					SNew(ConcertClientSharedSlate::SLocalClientName, Client.ToSharedRef())
@@ -46,12 +56,23 @@ namespace UE::MultiUserClient
 	{
 		FMenuBuilder MenuBuilder(true, nullptr);
 
+		MenuBuilder.AddMenuEntry(FUIAction(
+			FExecuteAction::CreateLambda([this](){ OnSelectAllClients.Execute(); }),
+				FCanExecuteAction::CreateLambda([this](){ return GetActiveWidgetIndex() != static_cast<int32>(EClientViewType::AllClients); })
+			),
+			MakeAllClientsDisplayWidget()
+			);
+
 		for (const FGuid& ClientId : ClientsAttribute.Get())
 		{
 			const bool bIsLocalClient = ClientId == Client->GetCurrentSession()->GetSessionClientEndpointId();
 			FUIAction UIAction(
 				FExecuteAction::CreateLambda([this, ClientId](){ OnSelectClientDelegate.Execute(ClientId); }),
-				FCanExecuteAction::CreateLambda([this, ClientId](){ return CurrentSelection.Get() != ClientId; })
+				FCanExecuteAction::CreateLambda([this, ClientId]()
+				{
+					return GetActiveWidgetIndex() == static_cast<int32>(EClientViewType::AllClients)
+						|| CurrentSelection.Get() != ClientId;
+				})
 				);
 			
 			if (bIsLocalClient)
@@ -67,11 +88,19 @@ namespace UE::MultiUserClient
 		return MenuBuilder.MakeWidget();
 	}
 
+	TSharedRef<SWidget> SSelectClientViewComboButton::MakeAllClientsDisplayWidget() const
+	{
+		return SNew(SHorizontalBox)
+			+SHorizontalBox::Slot()
+			[
+				SNew(STextBlock)
+				.Text(LOCTEXT("AllClients", "All Clients"))
+			];
+	}
+
 	int32 SSelectClientViewComboButton::GetActiveWidgetIndex() const
 	{
-		return CurrentSelection.Get() == Client->GetCurrentSession()->GetSessionClientEndpointId()
-			? static_cast<int32>(EButtonContent::LocalClient)
-			: static_cast<uint32>(EButtonContent::RemoteClient);
+		return static_cast<int32>(CurrentDisplayMode.Get());
 	}
 
 	FGuid SSelectClientViewComboButton::GetSelectedClientEndpointId() const
@@ -79,3 +108,5 @@ namespace UE::MultiUserClient
 		return CurrentSelection.Get();
 	}
 }
+
+#undef LOCTEXT_NAMESPACE
\ No newline at end of file
diff --git a/Engine/Plugins/Developer/Concert/ConcertApp/MultiUserClient/Source/MultiUserClient/Private/Widgets/ActiveSession/Replication/Joined/SSelectClientViewComboButton.h b/Engine/Plugins/Developer/Concert/ConcertApp/MultiUserClient/Source/MultiUserClient/Private/Widgets/ActiveSession/Replication/Joined/SSelectClientViewComboButton.h
--- a/Engine/Plugins/Developer/Concert/ConcertApp/MultiUserClient/Source/MultiUserClient/Private/Widgets/ActiveSession/Replication/Joined/SSelectClientViewComboButton.h
+++ b/Engine/Plugins/Developer/Concert/ConcertApp/MultiUserClient/Source/MultiUserClient/Private/Widgets/ActiveSession/Replication/Joined/SSelectClientViewComboButton.h
@@ -2,6 +2,7 @@
 
 #pragma once
 
+#include "EClientViewType.h"
 #include "Widgets/DeclarativeSyntaxSupport.h"
 #include "Widgets/SCompoundWidget.h"
 
@@ -18,8 +19,9 @@ namespace UE::MultiUserClient
 	class SSelectClientViewComboButton : public SCompoundWidget
 	{
 	public:
-
+		
 		DECLARE_DELEGATE_OneParam(FSelectClient, const FGuid&);
+		DECLARE_DELEGATE(FSelectAllClients);
 		
 		SLATE_BEGIN_ARGS(SSelectClientViewComboButton)
 		{}
@@ -28,11 +30,15 @@ namespace UE::MultiUserClient
 		
 			/** Remote clients that can be selected from in the order that they should be displayed. */
 			SLATE_ATTRIBUTE(TArray<FGuid>, SelectableClients)
-			/** The client that is currently selected */
+			/** The client that is currently selected. Only valid if EButtonContent == LocalClient or RemoteClient. */
 			SLATE_ATTRIBUTE(FGuid, CurrentSelection)
+			/** Determines what is being displayed in the main view right now. */
+			SLATE_ATTRIBUTE(EClientViewType, CurrentDisplayMode)
 
 			/** Called when a client is selected */
 			SLATE_EVENT(FSelectClient, OnSelectClient)
+			/** Called when all clients are supposed to be displayed */
+			SLATE_EVENT(FSelectAllClients, OnSelectAllClients)
 		SLATE_END_ARGS()
 
 		void Construct(const FArguments& InArgs);
@@ -46,19 +52,19 @@ namespace UE::MultiUserClient
 		TAttribute<TArray<FGuid>> ClientsAttribute;
 		/** The current client selection */
 		TAttribute<FGuid> CurrentSelection;
+		/** Determines what is being displayed in the main view right now. */
+		TAttribute<EClientViewType> CurrentDisplayMode;
 		
 		/** Called when a client is selected */
 		FSelectClient OnSelectClientDelegate;
+		/** Called when all clients are supposed to be displayed */
+		FSelectAllClients OnSelectAllClients;
 
-		enum class EButtonContent 
-		{
-			LocalClient = 0,
-			RemoteClient = 1
-		};
 		/** Content for the button */
 		TSharedPtr<SWidgetSwitcher> ButtonContent;
 		
 		TSharedRef<SWidget> MakeMenuContent();
+		TSharedRef<SWidget> MakeAllClientsDisplayWidget() const;
 		
 		int32 GetActiveWidgetIndex() const;
 		FGuid GetSelectedClientEndpointId() const;
diff --git a/Engine/Plugins/Developer/Concert/ConcertUI/ConcertClientSharedSlate/Source/ConcertClientSharedSlate/Private/Replication/Editor/Model/GenericPropertySelectionModel.cpp b/Engine/Plugins/Developer/Concert/ConcertUI/ConcertClientSharedSlate/Source/ConcertClientSharedSlate/Private/Replication/Editor/Model/GenericPropertySelectionModel.cpp
deleted file mode 100644
diff --git a/Engine/Plugins/Developer/Concert/ConcertUI/ConcertClientSharedSlate/Source/ConcertClientSharedSlate/Private/Replication/Editor/Model/GenericPropertySelectionModel.h b/Engine/Plugins/Developer/Concert/ConcertUI/ConcertClientSharedSlate/Source/ConcertClientSharedSlate/Private/Replication/Editor/Model/GenericPropertySelectionModel.h
deleted file mode 100644
diff --git a/Engine/Plugins/Developer/Concert/ConcertUI/ConcertClientSharedSlate/Source/ConcertClientSharedSlate/Private/Replication/Editor/Model/GenericReplicationStreamModel.cpp b/Engine/Plugins/Developer/Concert/ConcertUI/ConcertClientSharedSlate/Source/ConcertClientSharedSlate/Private/Replication/Editor/Model/GenericReplicationStreamModel.cpp
new file mode 100644
--- /dev/null
+++ b/Engine/Plugins/Developer/Concert/ConcertUI/ConcertClientSharedSlate/Source/ConcertClientSharedSlate/Private/Replication/Editor/Model/GenericReplicationStreamModel.cpp
@@ -0,0 +1,261 @@
+// Copyright Epic Games, Inc. All Rights Reserved.
+
+#include "GenericReplicationStreamModel.h"
+
+#include "Replication/PropertyChainUtils.h"
+#include "Replication/Data/ObjectReplicationMap.h"
+#include "Replication/Settings/ConcertReplicationEditorSettings.h"
+
+#include "Algo/AllOf.h"
+
+namespace UE::ConcertClientSharedSlate
+{
+	FGenericReplicationStreamModel::FGenericReplicationStreamModel(TAttribute<FObjectReplicationMap*> ReplicationMapAttribute)
+		: ReplicationMapAttribute(MoveTemp(ReplicationMapAttribute))
+	{}
+
+	FSoftClassPath FGenericReplicationStreamModel::GetObjectClass(const FSoftObjectPath& Object) const
+	{
+		const FObjectReplicationMap* ReplicationMap = ReplicationMapAttribute.Get();
+		if (!ensure(ReplicationMap))
+		{
+			return {};
+		}
+		
+		const FReplicatedObjectInfo* AssignedProperties = ReplicationMap->ReplicatedObjects.Find(Object);
+		return AssignedProperties
+			? AssignedProperties->ClassPath
+			: FSoftClassPath{};
+	}
+
+	bool FGenericReplicationStreamModel::ContainsObjects(const TSet<FSoftObjectPath>& Objects) const
+	{
+		const FObjectReplicationMap* ReplicationMap = ReplicationMapAttribute.Get();
+		return ensure(ReplicationMap)
+			&& Algo::AllOf(Objects, [this, ReplicationMap](const FSoftObjectPath& ObjectPath){ return ReplicationMap->ReplicatedObjects.Contains(ObjectPath); });
+	}
+
+	bool FGenericReplicationStreamModel::ContainsProperties(const FSoftObjectPath& Object, const TSet<FConcertPropertyChain>& Properties) const
+	{
+		const FObjectReplicationMap* ReplicationMap = ReplicationMapAttribute.Get();
+		if (!ensure(ReplicationMap))
+		{
+			return false;
+		}
+
+		const FReplicatedObjectInfo* ObjectInfo = ReplicationMap->ReplicatedObjects.Find(Object);
+		return ObjectInfo
+			&& Algo::AllOf(Properties, [ObjectInfo](const FConcertPropertyChain& Property){ return ObjectInfo->PropertySelection.ReplicatedProperties.Contains(Property); });
+	}
+
+	bool FGenericReplicationStreamModel::ForEachReplicatedObject(TFunctionRef<EBreakBehavior(const FSoftObjectPath& Object)> Delegate) const
+	{
+		const FObjectReplicationMap* ReplicationMap = ReplicationMapAttribute.Get();
+		if (!ensure(ReplicationMap))
+		{
+			return false;
+		}
+
+		for (const TPair<FSoftObjectPath, FReplicatedObjectInfo>& ObjectMap: ReplicationMap->ReplicatedObjects)
+		{
+			if (Delegate(ObjectMap.Key) == EBreakBehavior::Break)
+			{
+				return true;
+			}
+		}
+
+		return !ReplicationMap->ReplicatedObjects.IsEmpty();
+	}
+
+	bool FGenericReplicationStreamModel::ForEachProperty(const FSoftObjectPath& Object, TFunctionRef<EBreakBehavior(const FConcertPropertyChain& Parent)> Delegate) const
+	{
+		const FObjectReplicationMap* ReplicationMap = ReplicationMapAttribute.Get();
+		if (!ensure(ReplicationMap))
+		{
+			return false;
+		}
+
+		const FReplicatedObjectInfo* AssignedProperties = ReplicationMap->ReplicatedObjects.Find(Object);
+		if (!AssignedProperties)
+		{
+			return false;
+		}
+
+		for (const FConcertPropertyChain& ReplicatedPropertyInfo : AssignedProperties->PropertySelection.ReplicatedProperties)
+		{
+			if (Delegate(ReplicatedPropertyInfo) == EBreakBehavior::Break)
+			{
+				return true;
+			}
+		}
+		return !AssignedProperties->PropertySelection.ReplicatedProperties.IsEmpty();
+	}
+
+	void FGenericReplicationStreamModel::AddObjects(TConstArrayView<UObject*> Objects)
+	{
+		FObjectReplicationMap* ReplicationMap = ReplicationMapAttribute.Get();
+		if (!ensure(ReplicationMap) || Objects.IsEmpty())
+		{
+			return;
+		}
+		
+		TSet<UObject*> ObjectsNotAdded;
+		for (UObject* Object : Objects)
+		{
+			const FSoftObjectPath ObjectPath = Object;
+			if (ensureAlways(Object) && !ReplicationMap->ReplicatedObjects.Contains(ObjectPath))
+			{
+				FReplicatedObjectInfo& ObjectInfo = ReplicationMap->ReplicatedObjects.Add(ObjectPath);
+				ObjectInfo.ClassPath = Object->GetClass();
+			}
+			else
+			{
+				ObjectsNotAdded.Add(Object);
+			}
+		}
+		
+		if (LIKELY(ObjectsNotAdded.IsEmpty()))
+		{
+			OnObjectsChangedDelegate.Broadcast(Objects, {}, EReplicatedObjectChangeReason::ChangedDirectly);
+		}
+		else if (ObjectsNotAdded.Num() < Objects.Num())
+		{
+			// Uncommon case so it is ok if it is suboptimal
+			TArray<UObject*> Added;
+			for (UObject* Object : Objects)
+			{
+				if (!ObjectsNotAdded.Contains(Object))
+				{
+					Added.Add(Object);
+				}
+			}
+			OnObjectsChangedDelegate.Broadcast(Added, {}, EReplicatedObjectChangeReason::ChangedDirectly);
+		}
+	}
+
+	void FGenericReplicationStreamModel::RemoveObjects(TConstArrayView<FSoftObjectPath> Objects)
+	{
+		FObjectReplicationMap* ReplicationMap = ReplicationMapAttribute.Get();
+		if (!ensure(ReplicationMap) || Objects.IsEmpty())
+		{
+			return;
+		}
+		
+		TSet<FSoftObjectPath> ObjectsNotRemoved;
+		for (const FSoftObjectPath& Object : Objects)
+		{
+			const bool bRemoved = ReplicationMap->ReplicatedObjects.Remove(Object) != 0;
+			if (!bRemoved)
+			{
+				ObjectsNotRemoved.Add(Object);
+			}
+		}
+		
+		if (LIKELY(ObjectsNotRemoved.IsEmpty()))
+		{
+			OnObjectsChangedDelegate.Broadcast({}, Objects, EReplicatedObjectChangeReason::ChangedDirectly);
+		}
+		else if (ObjectsNotRemoved.Num() < Objects.Num())
+		{
+			// Uncommon case so it is ok if it is suboptimal
+			TArray<FSoftObjectPath> Removed;
+			for (const FSoftObjectPath& Object : Objects)
+			{
+				if (!ObjectsNotRemoved.Contains(Object))
+				{
+					Removed.Add(Object);
+				}
+			}
+			OnObjectsChangedDelegate.Broadcast({}, Removed, EReplicatedObjectChangeReason::ChangedDirectly);
+		}
+	}
+
+	void FGenericReplicationStreamModel::AddProperties(const FSoftObjectPath& Object, TConstArrayView<FConcertPropertyChain> Properties)
+	{
+		FObjectReplicationMap* ReplicationMap = ReplicationMapAttribute.Get();
+		if (!ensure(ReplicationMap))
+		{
+			return;
+		}
+		
+		FReplicatedObjectInfo* AssignedProperties = ReplicationMap->ReplicatedObjects.Find(Object);
+		if (!AssignedProperties)
+		{
+			return;
+		}
+
+		UClass* Class = AssignedProperties->ClassPath.TryLoadClass<UObject>();
+		if (!Class)
+		{
+			return;
+		}
+
+		bool bAddedAtLeastOne = false;
+		TArray<FConcertPropertyChain>& ReplicatedProperties = AssignedProperties->PropertySelection.ReplicatedProperties;
+		ReplicatedProperties.Reserve(ReplicatedProperties.Num() + Properties.Num());
+		for (const FConcertPropertyChain& AddedProperty : Properties)
+		{
+			bAddedAtLeastOne |= ReplicatedProperties.AddUnique(AddedProperty) != INDEX_NONE;
+
+			// Parent properties must also be added
+			// Not exactly efficient to iterate through the hierarchy for every removed item but it should be fine... Properties.Num() == 1 is the most common case
+			ConcertSyncCore::PropertyChain::ForEachReplicatableConcertProperty(*Class, [&bAddedAtLeastOne, &AssignedProperties, &AddedProperty](FConcertPropertyChain&& Property)
+			{
+				if (Property.IsParentOf(AddedProperty))
+				{
+					bAddedAtLeastOne |= AssignedProperties->PropertySelection.ReplicatedProperties.AddUnique(Property) != INDEX_NONE;
+				}
+				return EBreakBehavior::Continue;
+			});
+		}
+
+		if (bAddedAtLeastOne)
+		{
+			OnPropertiesChangedDelegate.Broadcast();
+		}
+	}
+
+	void FGenericReplicationStreamModel::RemoveProperties(const FSoftObjectPath& Object, TConstArrayView<FConcertPropertyChain> Properties)
+	{
+		FObjectReplicationMap* ReplicationMap = ReplicationMapAttribute.Get();
+		if (!ensure(ReplicationMap))
+		{
+			return;
+		}
+		
+		FReplicatedObjectInfo* AssignedProperties = ReplicationMap->ReplicatedObjects.Find(Object);
+		if (!AssignedProperties)
+		{
+			return;
+		}
+
+		UClass* Class = AssignedProperties->ClassPath.TryLoadClass<UObject>();
+		int32 NumRemoved = 0;
+		for (const FConcertPropertyChain& RemovedProperty : Properties)
+		{
+			NumRemoved += AssignedProperties->PropertySelection.ReplicatedProperties.Remove(RemovedProperty);
+
+			// Removal should not fail if the class is not available
+			if (!Class)
+			{
+				continue;
+			}
+
+			// Child properties must also be removed
+			// Not exactly efficient to iterate through the hierarchy for every removed item but it should be fine... Properties.Num() == 1 is the most common case
+			ConcertSyncCore::PropertyChain::ForEachReplicatableConcertProperty(*Class, [&NumRemoved, &AssignedProperties, &RemovedProperty](FConcertPropertyChain&& Property)
+			{
+				if (Property.IsChildOf(RemovedProperty))
+				{
+					NumRemoved += AssignedProperties->PropertySelection.ReplicatedProperties.Remove(Property);
+				}
+				return EBreakBehavior::Continue;
+			});
+		}
+		
+		if (NumRemoved > 0)
+		{
+			OnPropertiesChangedDelegate.Broadcast();
+		}
+	}
+}
diff --git a/Engine/Plugins/Developer/Concert/ConcertUI/ConcertClientSharedSlate/Source/ConcertClientSharedSlate/Private/Replication/Editor/Model/GenericReplicationStreamModel.h b/Engine/Plugins/Developer/Concert/ConcertUI/ConcertClientSharedSlate/Source/ConcertClientSharedSlate/Private/Replication/Editor/Model/GenericReplicationStreamModel.h
new file mode 100644
--- /dev/null
+++ b/Engine/Plugins/Developer/Concert/ConcertUI/ConcertClientSharedSlate/Source/ConcertClientSharedSlate/Private/Replication/Editor/Model/GenericReplicationStreamModel.h
@@ -0,0 +1,48 @@
+// Copyright Epic Games, Inc. All Rights Reserved.
+
+#pragma once
+
+#include "Replication/Editor/Model/IEditableReplicationStreamModel.h"
+#include "Misc/Attribute.h"
+#include "UObject/WeakObjectPtrTemplates.h"
+
+struct FConcertReplicationEditorSettings;
+struct FObjectReplicationMap;
+
+namespace UE::ConcertClientSharedSlate
+{
+	/** Implements logic for editing a FObjectReplicationMap contained in an UObject. */
+	class FGenericReplicationStreamModel
+		: public IEditableReplicationStreamModel
+	{
+	public:
+		
+		FGenericReplicationStreamModel(TAttribute<FObjectReplicationMap*> ReplicationMapAttribute);
+		
+		//~ Begin IReplicationStreamModel Interface
+		virtual FSoftClassPath GetObjectClass(const FSoftObjectPath& Object) const override;
+		virtual bool ContainsObjects(const TSet<FSoftObjectPath>& Objects) const override;
+		virtual bool ContainsProperties(const FSoftObjectPath& Object, const TSet<FConcertPropertyChain>& Properties) const override;
+		virtual bool ForEachReplicatedObject(TFunctionRef<EBreakBehavior(const FSoftObjectPath& Object)> Delegate) const override;
+		virtual bool ForEachProperty(const FSoftObjectPath& Object, TFunctionRef<EBreakBehavior(const FConcertPropertyChain& Parent)> Delegate) const override;
+		//~ End IReplicationStreamModel Interface
+		
+		//~ Begin IEditableReplicationStreamModel Interface
+		virtual void AddObjects(TConstArrayView<UObject*> Objects) override;
+		virtual void RemoveObjects(TConstArrayView<FSoftObjectPath> Objects) override;
+		virtual void AddProperties(const FSoftObjectPath&, TConstArrayView<FConcertPropertyChain> Properties) override;
+		virtual void RemoveProperties(const FSoftObjectPath&, TConstArrayView<FConcertPropertyChain> Properties) override;
+		virtual FOnObjectsChanged& OnObjectsChanged() override { return OnObjectsChangedDelegate; }
+		virtual FOnPropertiesChanged& OnPropertiesChanged() override { return OnPropertiesChangedDelegate; }
+		//~ End IEditableReplicationStreamModel Interface
+
+	private:
+
+		/** Returns the replication map that is supposed to be edited. */
+		TAttribute<FObjectReplicationMap*> ReplicationMapAttribute;
+
+		FOnObjectsChanged OnObjectsChangedDelegate;
+		FOnPropertiesChanged OnPropertiesChangedDelegate;
+	};
+}
+
diff --git a/Engine/Plugins/Developer/Concert/ConcertUI/ConcertClientSharedSlate/Source/ConcertClientSharedSlate/Private/Replication/Editor/Model/ReplicationStreamObject.h b/Engine/Plugins/Developer/Concert/ConcertUI/ConcertClientSharedSlate/Source/ConcertClientSharedSlate/Private/Replication/Editor/Model/ReplicationStreamObject.h
new file mode 100644
--- /dev/null
+++ b/Engine/Plugins/Developer/Concert/ConcertUI/ConcertClientSharedSlate/Source/ConcertClientSharedSlate/Private/Replication/Editor/Model/ReplicationStreamObject.h
@@ -0,0 +1,18 @@
+// Copyright Epic Games, Inc. All Rights Reserved.
+
+#pragma once
+
+#include "Replication/Data/ObjectReplicationMap.h"
+#include "UObject/Object.h"
+#include "ReplicationStreamObject.generated.h"
+
+/** UObject wrapper to allow transactions on FObjectReplicationMap. */
+UCLASS()
+class UReplicationStreamObject : public UObject
+{
+	GENERATED_BODY()
+public:
+	
+	UPROPERTY()
+	FObjectReplicationMap ReplicationMap;
+};
\ No newline at end of file
diff --git a/Engine/Plugins/Developer/Concert/ConcertUI/ConcertClientSharedSlate/Source/ConcertClientSharedSlate/Private/Replication/Editor/Model/TransactionalPropertySelectionModel.cpp b/Engine/Plugins/Developer/Concert/ConcertUI/ConcertClientSharedSlate/Source/ConcertClientSharedSlate/Private/Replication/Editor/Model/TransactionalPropertySelectionModel.cpp
deleted file mode 100644
diff --git a/Engine/Plugins/Developer/Concert/ConcertUI/ConcertClientSharedSlate/Source/ConcertClientSharedSlate/Private/Replication/Editor/Model/TransactionalPropertySelectionModel.h b/Engine/Plugins/Developer/Concert/ConcertUI/ConcertClientSharedSlate/Source/ConcertClientSharedSlate/Private/Replication/Editor/Model/TransactionalPropertySelectionModel.h
deleted file mode 100644
diff --git a/Engine/Plugins/Developer/Concert/ConcertUI/ConcertClientSharedSlate/Source/ConcertClientSharedSlate/Private/Replication/Editor/Model/TransactionalReplicationStreamModel.cpp b/Engine/Plugins/Developer/Concert/ConcertUI/ConcertClientSharedSlate/Source/ConcertClientSharedSlate/Private/Replication/Editor/Model/TransactionalReplicationStreamModel.cpp
new file mode 100644
--- /dev/null
+++ b/Engine/Plugins/Developer/Concert/ConcertUI/ConcertClientSharedSlate/Source/ConcertClientSharedSlate/Private/Replication/Editor/Model/TransactionalReplicationStreamModel.cpp
@@ -0,0 +1,84 @@
+// Copyright Epic Games, Inc. All Rights Reserved.
+
+#include "TransactionalReplicationStreamModel.h"
+
+#include "ScopedTransaction.h"
+#include "Misc/TransactionObjectEvent.h"
+
+#define LOCTEXT_NAMESPACE "FTransactionalReplicationStreamModel"
+
+namespace UE::ConcertClientSharedSlate
+{
+	FTransactionalReplicationStreamModel::FTransactionalReplicationStreamModel(
+		UObject& OwningObject,
+		TAttribute<FObjectReplicationMap*> ReplicationMapAttribute
+		)
+		: FGenericReplicationStreamModel(MoveTemp(ReplicationMapAttribute))
+		, OwningObject(&OwningObject)
+	{}
+
+	void FTransactionalReplicationStreamModel::AddObjects(TConstArrayView<UObject*> Objects)
+	{
+		const FScopedTransaction Transaction(LOCTEXT("AddObjects", "Add replicated objects"));
+		OwningObject->Modify();
+		FGenericReplicationStreamModel::AddObjects(Objects);
+	}
+
+	void FTransactionalReplicationStreamModel::RemoveObjects(TConstArrayView<FSoftObjectPath> Objects)
+	{
+		const FScopedTransaction Transaction(LOCTEXT("RemoveObjects", "Remove replicated objects"));
+		OwningObject->Modify();
+		FGenericReplicationStreamModel::RemoveObjects(Objects);
+	}
+
+	void FTransactionalReplicationStreamModel::AddProperties(const FSoftObjectPath& SoftObjectPath, TConstArrayView<FConcertPropertyChain> Properties)
+	{
+		const FScopedTransaction Transaction(LOCTEXT("AddProperties", "Add replicated properties"));
+		OwningObject->Modify();
+		FGenericReplicationStreamModel::AddProperties(SoftObjectPath, Properties);
+	}
+
+	void FTransactionalReplicationStreamModel::RemoveProperties(const FSoftObjectPath& SoftObjectPath, TConstArrayView<FConcertPropertyChain> Properties)
+	{
+		const FScopedTransaction Transaction(LOCTEXT("RemoveProperties", "Remove replicated properties"));
+		OwningObject->Modify();
+		FGenericReplicationStreamModel::RemoveProperties(SoftObjectPath, Properties);
+	}
+
+	bool FTransactionalReplicationStreamModel::MatchesContext(const FTransactionContext& InContext, const TArray<TPair<UObject*, FTransactionObjectEvent>>& TransactionObjectContexts) const
+	{
+		if (!OwningObject.IsValid())
+		{
+			return false;
+		}
+		
+		for (const TPair<UObject*, FTransactionObjectEvent>& TransactionObjectPair : TransactionObjectContexts)
+		{
+			UObject* Object = TransactionObjectPair.Key;
+			while (Object != nullptr)
+			{
+				if (Object == OwningObject)
+				{
+					return true;
+				}
+				Object = Object->GetOuter();
+			}
+		}
+
+		return false;
+	}
+
+	void FTransactionalReplicationStreamModel::PostUndo(bool bSuccess)
+	{
+		OnObjectsChanged().Broadcast({}, {}, EReplicatedObjectChangeReason::ExternalChange);
+		OnPropertiesChanged().Broadcast();
+	}
+
+	void FTransactionalReplicationStreamModel::PostRedo(bool bSuccess)
+	{
+		OnObjectsChanged().Broadcast({}, {}, EReplicatedObjectChangeReason::ExternalChange);
+		OnPropertiesChanged().Broadcast();
+	}
+}
+
+#undef LOCTEXT_NAMESPACE
\ No newline at end of file
diff --git a/Engine/Plugins/Developer/Concert/ConcertUI/ConcertClientSharedSlate/Source/ConcertClientSharedSlate/Private/Replication/Editor/Model/TransactionalReplicationStreamModel.h b/Engine/Plugins/Developer/Concert/ConcertUI/ConcertClientSharedSlate/Source/ConcertClientSharedSlate/Private/Replication/Editor/Model/TransactionalReplicationStreamModel.h
new file mode 100644
--- /dev/null
+++ b/Engine/Plugins/Developer/Concert/ConcertUI/ConcertClientSharedSlate/Source/ConcertClientSharedSlate/Private/Replication/Editor/Model/TransactionalReplicationStreamModel.h
@@ -0,0 +1,42 @@
+// Copyright Epic Games, Inc. All Rights Reserved.
+
+#pragma once
+
+#include "GenericReplicationStreamModel.h"
+#include "EditorUndoClient.h"
+
+namespace UE::ConcertClientSharedSlate
+{
+	/** Special case of FGenericPropertySelectionModel where the edited FObjectReplicationMap lives in an UObject that is RF_Transactional. */
+	class FTransactionalReplicationStreamModel
+		: public FGenericReplicationStreamModel
+		, public FSelfRegisteringEditorUndoClient
+	{
+	public:
+
+		FTransactionalReplicationStreamModel(
+			UObject& OwningObject,
+			TAttribute<FObjectReplicationMap*> ReplicationMapAttribute
+			);
+		
+		//~ Begin IEditableReplicationStreamModel Interface
+		virtual void AddObjects(TConstArrayView<UObject*> Objects) override;
+		virtual void RemoveObjects(TConstArrayView<FSoftObjectPath> Objects) override;
+		virtual void AddProperties(const FSoftObjectPath&, TConstArrayView<FConcertPropertyChain> Properties) override;
+		virtual void RemoveProperties(const FSoftObjectPath&, TConstArrayView<FConcertPropertyChain> Properties) override;
+		//~ End IEditableReplicationStreamModel Interface
+
+		//~ Begin FEditorUndoClient Interface
+		virtual bool MatchesContext(const FTransactionContext& InContext, const TArray<TPair<UObject*, FTransactionObjectEvent>>& TransactionObjectContexts) const override;
+		virtual void PostUndo(bool bSuccess) override;
+		virtual void PostRedo(bool bSuccess) override;
+		//~ End FEditorUndoClient Interface
+
+	private:
+
+		/** User of FTransactionalPropertySelectionModel is responsible for keeping OwningObject alive, e.g. via an asset editor. */
+		TWeakObjectPtr<UObject> OwningObject;
+	};
+}
+
+
diff --git a/Engine/Plugins/Developer/Concert/ConcertUI/ConcertClientSharedSlate/Source/ConcertClientSharedSlate/Private/Replication/Editor/View/MultiEditor/ConsolidatedMultiStreamModel.cpp b/Engine/Plugins/Developer/Concert/ConcertUI/ConcertClientSharedSlate/Source/ConcertClientSharedSlate/Private/Replication/Editor/View/MultiEditor/ConsolidatedMultiStreamModel.cpp
new file mode 100644
--- /dev/null
+++ b/Engine/Plugins/Developer/Concert/ConcertUI/ConcertClientSharedSlate/Source/ConcertClientSharedSlate/Private/Replication/Editor/View/MultiEditor/ConsolidatedMultiStreamModel.cpp
@@ -0,0 +1,181 @@
+// Copyright Epic Games, Inc. All Rights Reserved.
+
+#include "ConsolidatedMultiStreamModel.h"
+
+#include "Replication/Editor/Model/IEditableMultiReplicationStreamModel.h"
+#include "Replication/Editor/Model/IEditableReplicationStreamModel.h"
+#include "Replication/Editor/Model/ReplicationStreamObject.h"
+#include "Replication/Editor/Model/TransactionalReplicationStreamModel.h"
+
+#include "UObject/Package.h"
+
+namespace UE::ConcertClientSharedSlate
+{
+	namespace Private
+	{
+		static UReplicationStreamObject* MakeObject(bool bShouldSupportTransactions)
+		{
+			const EObjectFlags Flags = RF_Transient | (bShouldSupportTransactions ? RF_Transactional : RF_NoFlags);
+			return NewObject<UReplicationStreamObject>(GetTransientPackage(), NAME_None, Flags);
+		}
+
+		static TSharedRef<IEditableReplicationStreamModel> MakeModel(UReplicationStreamObject* Object, bool bShouldSupportTransactions)
+		{
+			TAttribute<FObjectReplicationMap*> Attribute = TAttribute<FObjectReplicationMap*>::CreateLambda([WeakPtr = TWeakObjectPtr<UReplicationStreamObject>(Object)]() -> FObjectReplicationMap* 
+			{
+				if (WeakPtr.IsValid())
+				{
+					return &WeakPtr->ReplicationMap;
+				}
+				return nullptr;
+			});
+			
+			if (bShouldSupportTransactions)
+			{
+				return MakeShared<FTransactionalReplicationStreamModel>(*Object, MoveTemp(Attribute));
+			}
+			return MakeShared<FGenericReplicationStreamModel>(MoveTemp(Attribute));
+		}
+	}
+	
+	FConsolidatedMultiStreamModel::FConsolidatedMultiStreamModel(
+		TSharedRef<IEditableMultiReplicationStreamModel> InMultiStreamModel,
+		bool bShouldSupportTransactions
+		)
+		: StreamForAddingObject(Private::MakeObject(bShouldSupportTransactions))
+		, StreamForAdding(Private::MakeModel(StreamForAddingObject, bShouldSupportTransactions))
+		, MultiStreamModel(MoveTemp(InMultiStreamModel))
+	{
+		MultiStreamModel->OnReadOnlyStreamChanged().AddRaw(this, &FConsolidatedMultiStreamModel::OnReadOnlyStreamChanged);
+		MultiStreamModel->OnStreamSetChanged().AddRaw(this, &FConsolidatedMultiStreamModel::RebuildStreamSubscriptions);
+		StreamForAdding->OnObjectsChanged().AddRaw(this, &FConsolidatedMultiStreamModel::OnObjectsChanged_StreamForAdding);
+		RebuildStreamSubscriptions();
+	}
+
+	FConsolidatedMultiStreamModel::~FConsolidatedMultiStreamModel()
+	{
+		MultiStreamModel->OnReadOnlyStreamChanged().RemoveAll(this);
+		MultiStreamModel->OnStreamSetChanged().RemoveAll(this);
+		StreamForAdding->OnObjectsChanged().RemoveAll(this);
+		ClearStreamSubscriptions();
+	}
+
+	FSoftClassPath FConsolidatedMultiStreamModel::GetObjectClass(const FSoftObjectPath& Object) const
+	{
+		FSoftClassPath Result = StreamForAdding->GetObjectClass(Object);
+		if (Result.IsValid())
+		{
+			return Result;
+		}
+
+		MultiStreamModel->ForEachStream([&Object, &Result](const TSharedRef<IReplicationStreamModel>& Stream)
+        {
+        	Result = Stream->GetObjectClass(Object);
+        	return Result.IsValid() ? EBreakBehavior::Break : EBreakBehavior::Continue;
+        });
+		return Result;
+	}
+
+	bool FConsolidatedMultiStreamModel::ContainsObjects(const TSet<FSoftObjectPath>& Objects) const
+	{
+		if (StreamForAdding->ContainsObjects(Objects))
+		{
+			return true;
+		}
+
+		bool bContains = false;
+		MultiStreamModel->ForEachStream([&Objects, &bContains](const TSharedRef<IReplicationStreamModel>& Stream)
+		{
+			bContains = Stream->ContainsObjects(Objects);
+			return bContains ? EBreakBehavior::Break : EBreakBehavior::Continue;
+		});
+		return bContains;
+	}		
+
+	bool FConsolidatedMultiStreamModel::ForEachReplicatedObject(TFunctionRef<EBreakBehavior(const FSoftObjectPath& Object)> Delegate) const
+	{
+		TSet<FSoftObjectPath> UniquePathsOnly;
+		EBreakBehavior BreakBehavior = EBreakBehavior::Continue;
+		const auto ProcessObjects = [&Delegate, &UniquePathsOnly, &BreakBehavior](const FSoftObjectPath& Object)
+		{
+			bool bAlreadyInSet = false;
+			UniquePathsOnly.Add(Object, &bAlreadyInSet);
+			if (!bAlreadyInSet)
+			{
+				BreakBehavior = Delegate(Object);
+			}
+			
+			return BreakBehavior;
+		};
+		
+		bool bAnyMappings = StreamForAdding->ForEachReplicatedObject(ProcessObjects);
+		if (BreakBehavior == EBreakBehavior::Break)
+		{
+			return bAnyMappings;
+		}
+
+		MultiStreamModel->ForEachStream([&BreakBehavior, &ProcessObjects, &bAnyMappings](const TSharedRef<IReplicationStreamModel>& Stream)
+		{
+			bAnyMappings |= Stream->ForEachReplicatedObject(ProcessObjects);
+			return BreakBehavior;
+		});
+		
+		return bAnyMappings;
+	}
+
+	void FConsolidatedMultiStreamModel::AddObjects(TConstArrayView<UObject*> Objects)
+	{
+		StreamForAdding->AddObjects(Objects);
+	}
+
+	void FConsolidatedMultiStreamModel::RemoveObjects(TConstArrayView<FSoftObjectPath> Objects)
+	{
+		StreamForAdding->RemoveObjects(Objects);
+		for (const TSharedRef<IEditableReplicationStreamModel>& Model : MultiStreamModel->GetEditableStreams())
+		{
+			Model->RemoveObjects(Objects);
+		}
+	}
+
+	void FConsolidatedMultiStreamModel::AddReferencedObjects(FReferenceCollector& Collector)
+	{
+		Collector.AddReferencedObject(StreamForAddingObject);
+	}
+
+	void FConsolidatedMultiStreamModel::OnObjectsChanged_StreamForAdding(TConstArrayView<UObject*> AddedObjects, TConstArrayView<FSoftObjectPath> RemovedObjects, EReplicatedObjectChangeReason ChangeReason)
+	{
+		OnObjectsChangedDelegate.Broadcast(AddedObjects, RemovedObjects, ChangeReason);
+	}
+
+	void FConsolidatedMultiStreamModel::OnObjectsChanged_ClientStream(TConstArrayView<UObject*> AddedObjects, TConstArrayView<FSoftObjectPath> RemovedObjects, EReplicatedObjectChangeReason ChangeReason)
+	{
+		OnObjectsChangedDelegate.Broadcast(AddedObjects, RemovedObjects, ChangeReason);
+	}
+
+	void FConsolidatedMultiStreamModel::OnReadOnlyStreamChanged(TSharedRef<IReplicationStreamModel> Stream)
+	{
+		OnObjectsChangedDelegate.Broadcast({}, {}, EReplicatedObjectChangeReason::ExternalChange);
+	}
+
+	void FConsolidatedMultiStreamModel::RebuildStreamSubscriptions()
+	{
+		ClearStreamSubscriptions();
+
+		for (const TSharedRef<IEditableReplicationStreamModel>& Model : MultiStreamModel->GetEditableStreams())
+		{
+			Model->OnObjectsChanged().AddRaw(this, &FConsolidatedMultiStreamModel::OnObjectsChanged_ClientStream);
+		}
+	}
+
+	void FConsolidatedMultiStreamModel::ClearStreamSubscriptions()
+	{
+		for (const TWeakPtr<IEditableReplicationStreamModel>& Model : SubscribedToViewers)
+		{
+			if (const TSharedPtr<IEditableReplicationStreamModel> ModelPin = Model.Pin())
+			{
+				ModelPin->OnObjectsChanged().RemoveAll(this);
+			}
+		}
+		SubscribedToViewers.Empty();
+	}
+}
diff --git a/Engine/Plugins/Developer/Concert/ConcertUI/ConcertClientSharedSlate/Source/ConcertClientSharedSlate/Private/Replication/Editor/View/MultiEditor/ConsolidatedMultiStreamModel.h b/Engine/Plugins/Developer/Concert/ConcertUI/ConcertClientSharedSlate/Source/ConcertClientSharedSlate/Private/Replication/Editor/View/MultiEditor/ConsolidatedMultiStreamModel.h
new file mode 100644
--- /dev/null
+++ b/Engine/Plugins/Developer/Concert/ConcertUI/ConcertClientSharedSlate/Source/ConcertClientSharedSlate/Private/Replication/Editor/View/MultiEditor/ConsolidatedMultiStreamModel.h
@@ -0,0 +1,99 @@
+// Copyright Epic Games, Inc. All Rights Reserved.
+
+#pragma once
+
+#include "Replication/Editor/Model/IEditableReplicationStreamModel.h"
+#include "UObject/GCObject.h"
+
+class UReplicationStreamObject;
+
+namespace UE::ConcertClientSharedSlate
+{
+	class IEditableMultiReplicationStreamModel;
+}
+
+namespace UE::ConcertClientSharedSlate
+{
+	/**
+	 * A mock object that implements functions in a specific way to allow editing of multiple streams in a single SBaseReplicationStreamEditor UI.
+	 * Coupled to SMultiReplicationStreamEditor.
+	 *
+	 * Function summary:
+	 * - AddObject & RemoveObject calls are forwarded to WrittenToStream
+	 * - Object getters, like ForEachReplicatedObject, query WrittenToStream AND the multiple streams in MultiStreamModel
+	 * - Property editing is disabled
+	 * 
+	 * SMultiReplicationStreamEditor's responsibility is to consolidate all replicated objects in a single UI and display all properties in the class.
+	 * The creator of SMultiReplicationStreamEditor is supposed to inject custom columns into the property section for assigning specific properties
+	 * to the streams in the IEditableReplicationStreamModel. 
+	 */
+	class FConsolidatedMultiStreamModel
+		: public IEditableReplicationStreamModel
+		, public FGCObject
+	{
+	public:
+
+		FConsolidatedMultiStreamModel(
+			TSharedRef<IEditableMultiReplicationStreamModel> InMultiStreamModel,
+			bool bShouldSupportTransactions
+			);
+		virtual ~FConsolidatedMultiStreamModel();
+		
+		// The functions related to properties are checkNoEntry() because SBaseReplicationStreamEditor's FFakeObjectToPropertiesEditorModel masks those functions 
+		
+		//~ Begin IReplicationStreamModel Interface
+		virtual FSoftClassPath GetObjectClass(const FSoftObjectPath& Object) const override;
+		virtual bool ContainsObjects(const TSet<FSoftObjectPath>& Objects) const override;
+		virtual bool ContainsProperties(const FSoftObjectPath& Object, const TSet<FConcertPropertyChain>& Properties) const override { checkNoEntry(); return false; }
+		virtual bool ForEachReplicatedObject(TFunctionRef<EBreakBehavior(const FSoftObjectPath& Object)> Delegate) const override;
+		virtual bool ForEachProperty(const FSoftObjectPath& Object, TFunctionRef<EBreakBehavior(const FConcertPropertyChain& Property)> Delegate) const override { checkNoEntry(); return false; }
+		//~ End IReplicationStreamModel Interface
+
+		//~ Begin IEditableReplicationStreamModel Interface
+		virtual void AddObjects(TConstArrayView<UObject*> Objects) override;
+		virtual void RemoveObjects(TConstArrayView<FSoftObjectPath> Objects) override;
+		virtual void AddProperties(const FSoftObjectPath& Object, TConstArrayView<FConcertPropertyChain> Properties) override { checkNoEntry(); }
+		virtual void RemoveProperties(const FSoftObjectPath& Object, TConstArrayView<FConcertPropertyChain> Properties) override { checkNoEntry(); }
+		virtual FOnObjectsChanged& OnObjectsChanged() override { return OnObjectsChangedDelegate; }
+		virtual FOnPropertiesChanged& OnPropertiesChanged() override { return OnPropertiesChangedDelegate; }
+		//~ End IEditableReplicationStreamModel Interface
+
+		//~ Begin FGCObject Interface
+		virtual void AddReferencedObjects(FReferenceCollector& Collector) override;
+		virtual FString GetReferencerName() const override { return TEXT("FConsolidatedMultiStreamModel"); }
+		//~ End FGCObject Interface
+
+	private:
+
+		/** What StreamForAdding is editing. Optionally transacted. */
+		TObjectPtr<UReplicationStreamObject> StreamForAddingObject;
+		/**
+		 * AddObjects calls are redirected to this model.
+		 * That allows objects added via the add button to show up in the UI regardless of what the real streams contain.
+		 *
+		 * RemoveObjects are applied to both StreamForAdding and to the editable streams in MultiStreamModel.
+		 */
+		TSharedRef<IEditableReplicationStreamModel> StreamForAdding;
+
+		/** Contains all the streams that are being displayed */
+		TSharedRef<IEditableMultiReplicationStreamModel> MultiStreamModel;
+
+		/** Called the equivalent event of WrittenToStream or any of the streams in MultiStreamModel is broadcast. */
+		FOnObjectsChanged OnObjectsChangedDelegate;
+		/** Never actually called, just to implement OnPropertiesChanged. */
+		FOnPropertiesChanged OnPropertiesChangedDelegate;
+		
+		/** All streams we've subscribed to for changes. Used so we can unsubscribe when MultiStreamModel->OnStreamsChanged triggers. */
+		TArray<TWeakPtr<IEditableReplicationStreamModel>> SubscribedToViewers;
+
+		// Trigger our delegates when any of the substreams change.
+		void OnObjectsChanged_StreamForAdding(TConstArrayView<UObject*> AddedObjects, TConstArrayView<FSoftObjectPath> RemovedObjects, EReplicatedObjectChangeReason ChangeReason);
+		void OnObjectsChanged_ClientStream(TConstArrayView<UObject*> AddedObjects, TConstArrayView<FSoftObjectPath> RemovedObjects, EReplicatedObjectChangeReason ChangeReason);
+		
+		/** Trigger subscription rebuild when streams are added / removed. */
+		void OnReadOnlyStreamChanged(TSharedRef<IReplicationStreamModel> Stream);
+		void RebuildStreamSubscriptions();
+		void ClearStreamSubscriptions();
+	};
+}
+
diff --git a/Engine/Plugins/Developer/Concert/ConcertUI/ConcertClientSharedSlate/Source/ConcertClientSharedSlate/Private/Replication/Editor/View/MultiEditor/SMultiReplicationStreamEditor.cpp b/Engine/Plugins/Developer/Concert/ConcertUI/ConcertClientSharedSlate/Source/ConcertClientSharedSlate/Private/Replication/Editor/View/MultiEditor/SMultiReplicationStreamEditor.cpp
new file mode 100644
--- /dev/null
+++ b/Engine/Plugins/Developer/Concert/ConcertUI/ConcertClientSharedSlate/Source/ConcertClientSharedSlate/Private/Replication/Editor/View/MultiEditor/SMultiReplicationStreamEditor.cpp
@@ -0,0 +1,33 @@
+// Copyright Epic Games, Inc. All Rights Reserved.
+
+#include "SMultiReplicationStreamEditor.h"
+
+#include "ConsolidatedMultiStreamModel.h"
+#include "Replication/ReplicationWidgetFactories.h"
+#include "Replication/Editor/Model/IEditableMultiReplicationStreamModel.h"
+#include "Replication/Editor/View/IReplicationStreamEditor.h"
+
+namespace UE::ConcertClientSharedSlate
+{
+	void SMultiReplicationStreamEditor::Construct(const FArguments& InArgs, FCreateMultiStreamEditorParams Params)
+	{
+		MultiStreamModel = MoveTemp(Params.MultiStreamModel);
+
+		const bool bConsolidatedModelShouldTransact = EnumHasAnyFlags(Params.Flags, EMultiStreamEditorFlags::Transactional);
+		ConsolidatedModel = MakeShared<FConsolidatedMultiStreamModel>(MoveTemp(Params.MultiStreamModel), bConsolidatedModelShouldTransact);
+		
+		const FCreateEditorParams BaseEditorParams
+		{
+			.DataModel = ConsolidatedModel.ToSharedRef(),
+			.ObjectSource = MoveTemp(Params.ObjectSource),
+			.PropertySource = MoveTemp(Params.PropertySource),
+			.ViewerParams = MoveTemp(Params.ViewerParams)
+		};
+		EditorView = CreateBaseStreamEditor(BaseEditorParams);
+		
+		ChildSlot
+		[
+			EditorView.ToSharedRef()
+		];
+	}
+}
diff --git a/Engine/Plugins/Developer/Concert/ConcertUI/ConcertClientSharedSlate/Source/ConcertClientSharedSlate/Private/Replication/Editor/View/MultiEditor/SMultiReplicationStreamEditor.h b/Engine/Plugins/Developer/Concert/ConcertUI/ConcertClientSharedSlate/Source/ConcertClientSharedSlate/Private/Replication/Editor/View/MultiEditor/SMultiReplicationStreamEditor.h
new file mode 100644
--- /dev/null
+++ b/Engine/Plugins/Developer/Concert/ConcertUI/ConcertClientSharedSlate/Source/ConcertClientSharedSlate/Private/Replication/Editor/View/MultiEditor/SMultiReplicationStreamEditor.h
@@ -0,0 +1,56 @@
+// Copyright Epic Games, Inc. All Rights Reserved.
+
+#pragma once
+
+#include "Replication/ReplicationWidgetFactories.h"
+#include "Replication/Editor/View/IMultiReplicationStreamEditor.h"
+
+#include "Widgets/DeclarativeSyntaxSupport.h"
+#include "Widgets/SCompoundWidget.h"
+
+namespace UE::ConcertClientSharedSlate
+{
+	enum class EReplicatedObjectChangeReason : uint8;
+	class FConsolidatedMultiStreamModel;
+	class IEditableReplicationStreamModel;
+}
+
+namespace UE::ConcertClientSharedSlate
+{
+	class FAggregateReadOnlyStreamModel;
+	struct FCreateMultiStreamEditorParams;
+
+	/**
+	 * Displays multiple streams in a single widget.
+	 * 
+	 * SMultiReplicationStreamEditor's responsibility is to consolidate all replicated objects in a single UI and display all properties in the class.
+	 * The creator of SMultiReplicationStreamEditor is supposed to inject custom columns into the property section for assigning specific properties
+	 * to the streams in the IEditableReplicationStreamModel. 
+	 */
+	class SMultiReplicationStreamEditor : public IMultiReplicationStreamEditor
+	{
+		using Super = IMultiReplicationStreamEditor;
+	public:
+
+		SLATE_BEGIN_ARGS(SMultiReplicationStreamEditor)
+		{}
+		SLATE_END_ARGS()
+
+		void Construct(const FArguments& InArgs, FCreateMultiStreamEditorParams Params);
+
+		virtual IReplicationStreamEditor& GetEditorBase() const override { return *EditorView; }
+		virtual IEditableMultiReplicationStreamModel& GetModel() const override { return *MultiStreamModel; }
+
+	private:
+
+		/** This model combines all of the user provided ones so they are all displayed in one view. */
+		TSharedPtr<FConsolidatedMultiStreamModel> ConsolidatedModel;
+		
+		/** Tells us which streams there are. */
+		TSharedPtr<IEditableMultiReplicationStreamModel> MultiStreamModel;
+
+		/** The main view, which is displaying ConsolidatedModel. */
+		TSharedPtr<IReplicationStreamEditor> EditorView;
+		
+	};
+}
diff --git a/Engine/Plugins/Developer/Concert/ConcertUI/ConcertClientSharedSlate/Source/ConcertClientSharedSlate/Private/Replication/Editor/View/ObjectEditor/FakeObjectToPropertiesEditorModel.h b/Engine/Plugins/Developer/Concert/ConcertUI/ConcertClientSharedSlate/Source/ConcertClientSharedSlate/Private/Replication/Editor/View/ObjectEditor/FakeObjectToPropertiesEditorModel.h
--- a/Engine/Plugins/Developer/Concert/ConcertUI/ConcertClientSharedSlate/Source/ConcertClientSharedSlate/Private/Replication/Editor/View/ObjectEditor/FakeObjectToPropertiesEditorModel.h
+++ b/Engine/Plugins/Developer/Concert/ConcertUI/ConcertClientSharedSlate/Source/ConcertClientSharedSlate/Private/Replication/Editor/View/ObjectEditor/FakeObjectToPropertiesEditorModel.h
@@ -46,8 +46,6 @@ namespace UE::ConcertClientSharedSlate
 
 		//~ Begin IReplicationStreamModel Interface
 		// Technically these functions should be also be wrapped but the SObjectToPropertyView does not use them so let's not for now.
-		virtual uint32 GetNumReplicatedObjects() const override { return RealModel->GetNumReplicatedObjects(); }
-		virtual uint32 GetNumProperties(const FSoftObjectPath& Object) const override { return RealModel->GetNumProperties(Object); }
 		virtual bool ContainsObjects(const TSet<FSoftObjectPath>& Objects) const override { return RealModel->ContainsObjects(Objects); }
 		virtual bool ContainsProperties(const FSoftObjectPath& Object, const TSet<FConcertPropertyChain>& Properties) const override { return RealModel->ContainsProperties(Object, Properties); }
 		
diff --git a/Engine/Plugins/Developer/Concert/ConcertUI/ConcertClientSharedSlate/Source/ConcertClientSharedSlate/Private/Replication/Editor/View/ObjectEditor/SBaseReplicationStreamEditor.cpp b/Engine/Plugins/Developer/Concert/ConcertUI/ConcertClientSharedSlate/Source/ConcertClientSharedSlate/Private/Replication/Editor/View/ObjectEditor/SBaseReplicationStreamEditor.cpp
--- a/Engine/Plugins/Developer/Concert/ConcertUI/ConcertClientSharedSlate/Source/ConcertClientSharedSlate/Private/Replication/Editor/View/ObjectEditor/SBaseReplicationStreamEditor.cpp
+++ b/Engine/Plugins/Developer/Concert/ConcertUI/ConcertClientSharedSlate/Source/ConcertClientSharedSlate/Private/Replication/Editor/View/ObjectEditor/SBaseReplicationStreamEditor.cpp
@@ -84,9 +84,7 @@ namespace UE::ConcertClientSharedSlate
 
 	void SBaseReplicationStreamEditor::Refresh()
 	{
-		ReplicationViewer->RefreshObjectData();
-		ReplicationViewer->RefreshSubobjectData();
-		ReplicationViewer->RefreshPropertyData();
+		ReplicationViewer->Refresh();
 	}
 
 	TArray<FSoftObjectPath> SBaseReplicationStreamEditor::GetSelectedTopLevelObjects() const
@@ -227,7 +225,6 @@ namespace UE::ConcertClientSharedSlate
 			return EBreakBehavior::Continue;
 		});
 		
-		checkf(ObjectAndChildren.Num() >= ObjectsToDelete.Num(), TEXT("Above algorithm is broken."));
 		EditablePropertiesModel->RemoveObjects(ObjectAndChildren);
 
 		// Viewer has an internal cache of the selected object. Setting the displayed object to the cached one is a no-op.
diff --git a/Engine/Plugins/Developer/Concert/ConcertUI/ConcertClientSharedSlate/Source/ConcertClientSharedSlate/Private/Replication/Editor/View/ObjectViewer/SReplicationStreamViewer.cpp b/Engine/Plugins/Developer/Concert/ConcertUI/ConcertClientSharedSlate/Source/ConcertClientSharedSlate/Private/Replication/Editor/View/ObjectViewer/SReplicationStreamViewer.cpp
--- a/Engine/Plugins/Developer/Concert/ConcertUI/ConcertClientSharedSlate/Source/ConcertClientSharedSlate/Private/Replication/Editor/View/ObjectViewer/SReplicationStreamViewer.cpp
+++ b/Engine/Plugins/Developer/Concert/ConcertUI/ConcertClientSharedSlate/Source/ConcertClientSharedSlate/Private/Replication/Editor/View/ObjectViewer/SReplicationStreamViewer.cpp
@@ -30,10 +30,7 @@ namespace UE::ConcertClientSharedSlate
 			CreateContentWidget(InArgs)
 		];
 
-		RefreshObjectData();
-		RefreshSubobjectData();
-		RefreshPropertyData();
-
+		Refresh();
 		PropertyArea->SetExpanded(true);
 	}
 
@@ -61,10 +58,9 @@ namespace UE::ConcertClientSharedSlate
 
 	void SReplicationStreamViewer::RefreshObjectData()
 	{
-		const int32 NumElements = PropertiesModel->GetNumReplicatedObjects();
 		// Re-using existing instances is tricky: we cannot update the object path in an item because the list view will no detect this change;
 		// list view only looks at the shared ptr address. So the UI will not be refreshed. Since the number of items will be small, just reallocate... 
-		AllObjectRowData.Empty(NumElements);
+		AllObjectRowData.Empty();
 
 		// Try to re-use old instances by using the old PathToObjectDataCache. This is also done so the expansion states restore correctly in the tree view.
 		TMap<FSoftObjectPath, TSharedPtr<FReplicatedObjectData>> NewPathToObjectDataCache;
diff --git a/Engine/Plugins/Developer/Concert/ConcertUI/ConcertClientSharedSlate/Source/ConcertClientSharedSlate/Private/Replication/Editor/View/ObjectViewer/SSubobjectAndPropertySection.cpp b/Engine/Plugins/Developer/Concert/ConcertUI/ConcertClientSharedSlate/Source/ConcertClientSharedSlate/Private/Replication/Editor/View/ObjectViewer/SSubobjectAndPropertySection.cpp
--- a/Engine/Plugins/Developer/Concert/ConcertUI/ConcertClientSharedSlate/Source/ConcertClientSharedSlate/Private/Replication/Editor/View/ObjectViewer/SSubobjectAndPropertySection.cpp
+++ b/Engine/Plugins/Developer/Concert/ConcertUI/ConcertClientSharedSlate/Source/ConcertClientSharedSlate/Private/Replication/Editor/View/ObjectViewer/SSubobjectAndPropertySection.cpp
@@ -95,7 +95,7 @@ namespace UE::ConcertClientSharedSlate
 		// Try to re-use old instances by using the old ChainToPropertyDataCache. This is also done so the expansion states restore correctly in the tree view.
 		TMap<FConcertPropertyChain, TSharedPtr<FReplicatedPropertyData>> NewChainToPropertyDataCache;
 		
-		PropertyRowData.Empty(PropertiesModel->GetNumProperties(SelectedObjects[0]));
+		PropertyRowData.Empty();
 		for (const FConcertPropertyChain& PropertyChain : SharedProperties)
 		{
 			const TSharedPtr<FReplicatedPropertyData>* ExistingItem = ChainToPropertyDataCache.Find(PropertyChain);
diff --git a/Engine/Plugins/Developer/Concert/ConcertUI/ConcertClientSharedSlate/Source/ConcertClientSharedSlate/Private/Replication/Editor/View/ObjectViewer/Tree/SelectionViewerColumns.cpp b/Engine/Plugins/Developer/Concert/ConcertUI/ConcertClientSharedSlate/Source/ConcertClientSharedSlate/Private/Replication/Editor/View/ObjectViewer/Tree/SelectionViewerColumns.cpp
--- a/Engine/Plugins/Developer/Concert/ConcertUI/ConcertClientSharedSlate/Source/ConcertClientSharedSlate/Private/Replication/Editor/View/ObjectViewer/Tree/SelectionViewerColumns.cpp
+++ b/Engine/Plugins/Developer/Concert/ConcertUI/ConcertClientSharedSlate/Source/ConcertClientSharedSlate/Private/Replication/Editor/View/ObjectViewer/Tree/SelectionViewerColumns.cpp
@@ -259,7 +259,7 @@ namespace UE::ConcertClientSharedSlate::ReplicationColumns::Property
 					Model.RemoveProperties(Path, Properties);
 					// If this does not resolve it is not too bad if the subobject is not removed...
 					UObject* Object = Path.ResolveObject();
-					const bool bNeedsToRemoveNonRoot = Object && !Object->IsA<AActor>() && Model.GetNumProperties(Path) == 0;
+					const bool bNeedsToRemoveNonRoot = Object && !Object->IsA<AActor>() && !Model.ContainsObjects({Path});
 					if (bNeedsToRemoveNonRoot)
 					{
 						Model.RemoveObjects({ Path });
diff --git a/Engine/Plugins/Developer/Concert/ConcertUI/ConcertClientSharedSlate/Source/ConcertClientSharedSlate/Private/Replication/ReplicationWidgetFactories.cpp b/Engine/Plugins/Developer/Concert/ConcertUI/ConcertClientSharedSlate/Source/ConcertClientSharedSlate/Private/Replication/ReplicationWidgetFactories.cpp
--- a/Engine/Plugins/Developer/Concert/ConcertUI/ConcertClientSharedSlate/Source/ConcertClientSharedSlate/Private/Replication/ReplicationWidgetFactories.cpp
+++ b/Engine/Plugins/Developer/Concert/ConcertUI/ConcertClientSharedSlate/Source/ConcertClientSharedSlate/Private/Replication/ReplicationWidgetFactories.cpp
@@ -2,8 +2,9 @@
 
 #include "Replication/ReplicationWidgetFactories.h"
 
-#include "Editor/Model/TransactionalPropertySelectionModel.h"
+#include "Editor/Model/TransactionalReplicationStreamModel.h"
 #include "Editor/Model/Subobject/ComponentHierarchySubobjectModel.h"
+#include "Editor/View/MultiEditor/SMultiReplicationStreamEditor.h"
 #include "Editor/View/ObjectEditor/SDefaultReplicationStreamEditor.h"
 #include "Editor/View/ObjectEditor/SBaseReplicationStreamEditor.h"
 #include "Editor/View/SubobjectView/SSubobjectView.h"
@@ -21,17 +22,17 @@ namespace UE::ConcertClientSharedSlate
 		return MakeShared<FComponentHierarchySubobjectModel>();
 	}
 	
-	TSharedRef<IReplicationStreamEditor> CreateEditor(FCreateEditorParams Params)
+	TSharedRef<IReplicationStreamEditor> CreateBaseStreamEditor(FCreateEditorParams Params)
 	{
 		return SNew(SBaseReplicationStreamEditor, Params.DataModel, Params.ObjectSource, Params.PropertySource)
-			.AdditionalObjectColumns(Params.AdditionalObjectColumns)
-			.AdditionalPropertyColumns(Params.AdditionalPropertyColumns)
-			.SubobjectView(Params.SubobjectView)
-			.SubobjectModel(Params.SubobjectModel)
-			.OnExtendObjectsContextMenu(Params.OnExtendObjectsContextMenu)
-			.SortPropertyRowPredicate(Params.SortPropertyRowPredicate)
-			.LeftOfObjectSearchBar() [ Params.LeftOfObjectSearchBar.Widget ]
-			.LeftOfPropertySearchBar() [ Params.LeftOfPropertySearchBar.Widget ]
+			.AdditionalObjectColumns(Params.ViewerParams.AdditionalObjectColumns)
+			.AdditionalPropertyColumns(Params.ViewerParams.AdditionalPropertyColumns)
+			.SubobjectView(Params.ViewerParams.SubobjectView)
+			.SubobjectModel(Params.ViewerParams.SubobjectModel)
+			.OnExtendObjectsContextMenu(Params.ViewerParams.OnExtendObjectsContextMenu)
+			.SortPropertyRowPredicate(Params.ViewerParams.SortPropertyRowPredicate)
+			.LeftOfObjectSearchBar() [ Params.ViewerParams.LeftOfObjectSearchBar.Widget ]
+			.LeftOfPropertySearchBar() [ Params.ViewerParams.LeftOfPropertySearchBar.Widget ]
 			.IsEditingEnabled(Params.IsEditingEnabled)
 			.EditingDisabledToolTipText(Params.EditingDisabledToolTipText)
 			.ReplicationSettings(Params.ReplicationSettingsAttribute);
@@ -40,14 +41,14 @@ namespace UE::ConcertClientSharedSlate
 	TSharedRef<IReplicationStreamEditor> CreateDefaultStreamEditor(FCreateEditorParams Params)
 	{
 		return SNew(SDefaultReplicationStreamEditor, Params.DataModel, Params.ObjectSource, Params.PropertySource)
-			.AdditionalObjectColumns(Params.AdditionalObjectColumns)
-			.AdditionalPropertyColumns(Params.AdditionalPropertyColumns)
-			.SubobjectView(Params.SubobjectView)
-			.SubobjectModel(Params.SubobjectModel)
-			.OnExtendObjectsContextMenu(Params.OnExtendObjectsContextMenu)
-			.SortPropertyRowPredicate(Params.SortPropertyRowPredicate)
-			.LeftOfObjectSearchBar() [ Params.LeftOfObjectSearchBar.Widget ]
-			.LeftOfPropertySearchBar() [ Params.LeftOfPropertySearchBar.Widget ]
+			.AdditionalObjectColumns(Params.ViewerParams.AdditionalObjectColumns)
+			.AdditionalPropertyColumns(Params.ViewerParams.AdditionalPropertyColumns)
+			.SubobjectView(Params.ViewerParams.SubobjectView)
+			.SubobjectModel(Params.ViewerParams.SubobjectModel)
+			.OnExtendObjectsContextMenu(Params.ViewerParams.OnExtendObjectsContextMenu)
+			.SortPropertyRowPredicate(Params.ViewerParams.SortPropertyRowPredicate)
+			.LeftOfObjectSearchBar() [ Params.ViewerParams.LeftOfObjectSearchBar.Widget ]
+			.LeftOfPropertySearchBar() [ Params.ViewerParams.LeftOfPropertySearchBar.Widget ]
 			.IsEditingEnabled(Params.IsEditingEnabled)
 			.EditingDisabledToolTipText(Params.EditingDisabledToolTipText)
 			.ReplicationSettings(Params.ReplicationSettingsAttribute);
@@ -58,10 +59,15 @@ namespace UE::ConcertClientSharedSlate
 		TAttribute<FObjectReplicationMap*> ReplicationMapAttribute
 		)
 	{
-		return MakeShared<FTransactionalPropertySelectionModel>(
+		return MakeShared<FTransactionalReplicationStreamModel>(
 			OwnerObject,
 			MoveTemp(ReplicationMapAttribute)
 			);
 	}
+
+	TSharedRef<IMultiReplicationStreamEditor> CreateBaseMultiStreamEditor(FCreateMultiStreamEditorParams Params)
+	{
+		return SNew(SMultiReplicationStreamEditor, MoveTemp(Params));
+	}
 }
 
diff --git a/Engine/Plugins/Developer/Concert/ConcertUI/ConcertClientSharedSlate/Source/ConcertClientSharedSlate/Public/Replication/Editor/Model/IEditableMultiReplicationStreamModel.h b/Engine/Plugins/Developer/Concert/ConcertUI/ConcertClientSharedSlate/Source/ConcertClientSharedSlate/Public/Replication/Editor/Model/IEditableMultiReplicationStreamModel.h
new file mode 100644
--- /dev/null
+++ b/Engine/Plugins/Developer/Concert/ConcertUI/ConcertClientSharedSlate/Source/ConcertClientSharedSlate/Public/Replication/Editor/Model/IEditableMultiReplicationStreamModel.h
@@ -0,0 +1,57 @@
+// Copyright Epic Games, Inc. All Rights Reserved.
+
+#pragma once
+
+#include "IMultiReplicationStreamModel.h"
+
+#include "Delegates/Delegate.h"
+#include "Misc/EBreakBehavior.h"
+#include "Templates/SharedPointer.h"
+
+namespace UE::ConcertClientSharedSlate
+{
+	class IEditableReplicationStreamModel;
+	
+	/**
+	 * Model for UI editing multiple streams at the same time.
+	 *
+	 * Some streams can be read but not edited:
+	 * - the set of readable streams includes the set of editable streams, but
+	 * - the number of editable streams can be less than or equal to the number of read-only streams.
+	 */
+	class IEditableMultiReplicationStreamModel : public IMultiReplicationStreamModel
+	{
+	public:
+
+		/**
+		 * @return Gets the streams that can be written to.
+		 * 
+		 * @note There may be streams that are readable but not editable.
+		 * Formally, the constraint is GetReadableStreams().Includes(GetEditableStreams()) and GetEditableStreams().Num() <= GetReadOnlyStreams().Num().
+		 */
+		virtual TSet<TSharedRef<IEditableReplicationStreamModel>> GetEditableStreams() const = 0;
+
+		/** Util for iterating through both read-only and  */
+		void ForEachStream(TFunctionRef<EBreakBehavior(const TSharedRef<IReplicationStreamModel>& Model)> Callback) const
+		{
+			const TSet<TSharedRef<IReplicationStreamModel>> ReadableModels = GetReadOnlyStreams();
+			for (const TSharedRef<IReplicationStreamModel>& Model : GetReadOnlyStreams())
+			{
+				if (Callback(Model) == EBreakBehavior::Break)
+				{
+					return;
+				}
+			}
+			
+			for (const TSharedRef<IEditableReplicationStreamModel>& Model : GetEditableStreams())
+			{
+				const TSharedRef<IReplicationStreamModel> CastModel = StaticCastSharedRef<IReplicationStreamModel, IEditableReplicationStreamModel>(Model);
+				if (!ReadableModels.Contains(CastModel) // Already iterated it if it is contained
+					&& Callback(CastModel) == EBreakBehavior::Break)
+				{
+					return;
+				}
+			}
+		}
+	};
+}
\ No newline at end of file
diff --git a/Engine/Plugins/Developer/Concert/ConcertUI/ConcertClientSharedSlate/Source/ConcertClientSharedSlate/Public/Replication/Editor/Model/IEditableReplicationStreamModel.h b/Engine/Plugins/Developer/Concert/ConcertUI/ConcertClientSharedSlate/Source/ConcertClientSharedSlate/Public/Replication/Editor/Model/IEditableReplicationStreamModel.h
--- a/Engine/Plugins/Developer/Concert/ConcertUI/ConcertClientSharedSlate/Source/ConcertClientSharedSlate/Public/Replication/Editor/Model/IEditableReplicationStreamModel.h
+++ b/Engine/Plugins/Developer/Concert/ConcertUI/ConcertClientSharedSlate/Source/ConcertClientSharedSlate/Public/Replication/Editor/Model/IEditableReplicationStreamModel.h
@@ -16,8 +16,14 @@ namespace UE::ConcertClientSharedSlate
 {
 	enum class EReplicatedObjectChangeReason : uint8
 	{
+		/** A function call on the model changed the data. */
 		ChangedDirectly,
-		Transacted
+		
+		/**
+		 * An external change, such as transaction, changed the data out from under us.
+		 * The AddedObjects and RemovedObjects fields are empty because we don't know exactly what changed, if anything.
+		 */
+		ExternalChange
 	};
 	
 	/**
@@ -26,7 +32,7 @@ namespace UE::ConcertClientSharedSlate
 	 * Models may not always be writable. When editing a UAsset, it will be writable. However, if we join a multi-user
 	 * session we do not want to edit the objects nor properties - only read.
 	 */
-	class CONCERTCLIENTSHAREDSLATE_API IEditableReplicationStreamModel : public IReplicationStreamModel
+	class IEditableReplicationStreamModel : public IReplicationStreamModel
 	{
 	public:
 
diff --git a/Engine/Plugins/Developer/Concert/ConcertUI/ConcertClientSharedSlate/Source/ConcertClientSharedSlate/Public/Replication/Editor/Model/IMultiReplicationStreamModel.h b/Engine/Plugins/Developer/Concert/ConcertUI/ConcertClientSharedSlate/Source/ConcertClientSharedSlate/Public/Replication/Editor/Model/IMultiReplicationStreamModel.h
new file mode 100644
--- /dev/null
+++ b/Engine/Plugins/Developer/Concert/ConcertUI/ConcertClientSharedSlate/Source/ConcertClientSharedSlate/Public/Replication/Editor/Model/IMultiReplicationStreamModel.h
@@ -0,0 +1,30 @@
+// Copyright Epic Games, Inc. All Rights Reserved.
+
+#pragma once
+
+#include "Delegates/Delegate.h"
+#include "Templates/SharedPointer.h"
+
+namespace UE::ConcertClientSharedSlate
+{
+	class IReplicationStreamModel;
+	
+	/** Model for UI displaying multiple streams at the same time. */
+	class IMultiReplicationStreamModel : public TSharedFromThis<IMultiReplicationStreamModel>
+	{
+	public:
+
+		/** @return Streams in this model that can only be read from. */
+		virtual TSet<TSharedRef<IReplicationStreamModel>> GetReadOnlyStreams() const = 0;
+
+		DECLARE_MULTICAST_DELEGATE_OneParam(FOnReadOnlyStreamChanged, TSharedRef<IReplicationStreamModel> /* ChangedStream */);
+		/** Called when a read-only stream has changed. */
+		virtual FOnReadOnlyStreamChanged& OnReadOnlyStreamChanged() = 0;
+		
+		DECLARE_MULTICAST_DELEGATE(FOnStreamSetChanged);
+		/** Broadcasts when the result of GetStreams has changed. */
+		virtual FOnStreamSetChanged& OnStreamSetChanged() = 0;
+
+		virtual ~IMultiReplicationStreamModel() = default;
+	};
+}
\ No newline at end of file
diff --git a/Engine/Plugins/Developer/Concert/ConcertUI/ConcertClientSharedSlate/Source/ConcertClientSharedSlate/Public/Replication/Editor/Model/IReplicationStreamModel.h b/Engine/Plugins/Developer/Concert/ConcertUI/ConcertClientSharedSlate/Source/ConcertClientSharedSlate/Public/Replication/Editor/Model/IReplicationStreamModel.h
--- a/Engine/Plugins/Developer/Concert/ConcertUI/ConcertClientSharedSlate/Source/ConcertClientSharedSlate/Public/Replication/Editor/Model/IReplicationStreamModel.h
+++ b/Engine/Plugins/Developer/Concert/ConcertUI/ConcertClientSharedSlate/Source/ConcertClientSharedSlate/Public/Replication/Editor/Model/IReplicationStreamModel.h
@@ -13,15 +13,10 @@ namespace UE::ConcertClientSharedSlate
 	 * Abstracts the concept of mapping objects to properties. This is a read-only interface.
 	 * @see IEditableReplicationStreamModel 
 	 */
-	class CONCERTCLIENTSHAREDSLATE_API IReplicationStreamModel : public TSharedFromThis<IReplicationStreamModel>
+	class IReplicationStreamModel : public TSharedFromThis<IReplicationStreamModel>
 	{
 	public:
 
-		/** @return The number of replicated objects */
-		virtual uint32 GetNumReplicatedObjects() const = 0;
-		/** @return The number of properties assigned to this object */
-		virtual uint32 GetNumProperties(const FSoftObjectPath& Object) const = 0;
-
 		/** @return Gets the class of the replication object */
 		virtual FSoftClassPath GetObjectClass(const FSoftObjectPath& Object) const = 0;
 
@@ -42,6 +37,18 @@ namespace UE::ConcertClientSharedSlate
 		 * @param Delegate The callback
 		 */
 		virtual bool ForEachProperty(const FSoftObjectPath& Object, TFunctionRef<EBreakBehavior(const FConcertPropertyChain& Property)> Delegate) const = 0;
+
+		/** @return Gets all replicated objects */
+		TSet<FSoftObjectPath> GetReplicatedObjects() const
+		{
+			TSet<FSoftObjectPath> Result;
+			ForEachReplicatedObject([&Result](const FSoftObjectPath& ObjectPath)
+			{
+				Result.Add(ObjectPath);
+				return EBreakBehavior::Continue;
+			});
+			return Result;
+		}
 		
 		/**
 		 * Iterates all the direct child properties of a given property in the mapping.
@@ -66,6 +73,30 @@ namespace UE::ConcertClientSharedSlate
 			return bFoundAtLeastOne;
 		}
 
+		/** @return Whether the given SearchedProperty is assigned to ObjectPath. */
+		bool HasProperty(const FSoftObjectPath& ObjectPath, const FConcertPropertyChain& SearchedProperty) const
+		{
+			bool bFound = false;
+			ForEachProperty(ObjectPath, [&SearchedProperty, &bFound](const FConcertPropertyChain& Property)
+			{
+				bFound = SearchedProperty == Property;
+				return bFound ? EBreakBehavior::Break : EBreakBehavior::Continue;
+			});
+			return bFound;
+		}
+
+		/** @return Whether there is at least one property assigned to ObjectPath. */
+		bool HasAnyPropertyAssigned(const FSoftObjectPath& ObjectPath) const
+		{
+			bool bFound = false;
+			ForEachProperty(ObjectPath, [&bFound](const FConcertPropertyChain& Property)
+			{
+				bFound = true;
+				return EBreakBehavior::Break;
+			});
+			return bFound;
+		}
+
 		/** Util for getting properties as array */
 		TSet<FConcertPropertyChain> GetAllProperties(const FSoftObjectPath& Object) const
 		{
diff --git a/Engine/Plugins/Developer/Concert/ConcertUI/ConcertClientSharedSlate/Source/ConcertClientSharedSlate/Public/Replication/Editor/View/IMultiReplicationStreamEditor.h b/Engine/Plugins/Developer/Concert/ConcertUI/ConcertClientSharedSlate/Source/ConcertClientSharedSlate/Public/Replication/Editor/View/IMultiReplicationStreamEditor.h
new file mode 100644
--- /dev/null
+++ b/Engine/Plugins/Developer/Concert/ConcertUI/ConcertClientSharedSlate/Source/ConcertClientSharedSlate/Public/Replication/Editor/View/IMultiReplicationStreamEditor.h
@@ -0,0 +1,26 @@
+// Copyright Epic Games, Inc. All Rights Reserved.
+
+#pragma once
+
+#include "Widgets/SCompoundWidget.h"
+
+namespace UE::ConcertClientSharedSlate
+{
+	class IEditableMultiReplicationStreamModel;
+	class IReplicationStreamEditor;
+	
+	/**
+	 * Widget which edits multiple replication stream.
+	 * @see ReplicationWidgetFactories.h
+	 */
+	class CONCERTCLIENTSHAREDSLATE_API IMultiReplicationStreamEditor : public SCompoundWidget
+	{
+	public:
+
+		/** @return The widget drawing the consolidate model. */
+		virtual IReplicationStreamEditor& GetEditorBase() const = 0;
+
+		/** @return The underlying model */
+		virtual IEditableMultiReplicationStreamModel& GetModel() const = 0;
+	};
+}
diff --git a/Engine/Plugins/Developer/Concert/ConcertUI/ConcertClientSharedSlate/Source/ConcertClientSharedSlate/Public/Replication/ReplicationWidgetFactories.h b/Engine/Plugins/Developer/Concert/ConcertUI/ConcertClientSharedSlate/Source/ConcertClientSharedSlate/Public/Replication/ReplicationWidgetFactories.h
--- a/Engine/Plugins/Developer/Concert/ConcertUI/ConcertClientSharedSlate/Source/ConcertClientSharedSlate/Public/Replication/ReplicationWidgetFactories.h
+++ b/Engine/Plugins/Developer/Concert/ConcertUI/ConcertClientSharedSlate/Source/ConcertClientSharedSlate/Public/Replication/ReplicationWidgetFactories.h
@@ -5,6 +5,7 @@
 #include "Editor/View/PredefinedReplicationColumns.h"
 #include "ReplicationWidgetDelegates.h"
 
+#include "Delegates/Delegate.h"
 #include "Misc/Attribute.h"
 #include "Templates/SharedPointer.h"
 
@@ -16,7 +17,10 @@ struct FObjectReplicationMap;
 
 namespace UE::ConcertClientSharedSlate
 {
+	class IEditableMultiReplicationStreamModel;
 	class IEditableReplicationStreamModel;
+	class IMultiReplicationStreamEditor;
+	
 	class IReplicationStreamModel;
 	class IObjectSelectionSourceModel;
 	class IReplicationStreamEditor;
@@ -57,25 +61,9 @@ namespace UE::ConcertClientSharedSlate
 		FCreateSubobjectViewParams Params
 		);
 
-	/** Params for creating an IReplicationStreamEditor */
-	struct FCreateEditorParams
+	/** Params for creating a IReplicationStreamViewer. */
+	struct FCreateViewerParams
 	{
-		/**
-		 * The model that the editor is displaying.
-		 * @note The view will keep a strong reference to this.
-		 */
-		TSharedRef<IEditableReplicationStreamModel> DataModel;
-		/**
-		 * Determines the objects that can be added to the object list. 
-		 * @note The view will keep a strong reference to this.
-		 */
-		TSharedRef<IObjectSelectionSourceModel> ObjectSource;
-		/**
-		 * Determines the properties that can be added to the property list.
-		 * @note The view will keep a strong reference to this.
-		 */
-		TSharedRef<IPropertySelectionSourceModel> PropertySource;
-
 		/**
 		 * Optional. This is inserted between the root object outliner and property view.
 		 * It e.g. displays the components of the actor selected in the root object view.
@@ -106,7 +94,30 @@ namespace UE::ConcertClientSharedSlate
 		TAlwaysValidWidget LeftOfObjectSearchBar;
 		/** Optional widget to add to the left of the property list search bar. */
 		TAlwaysValidWidget LeftOfPropertySearchBar;
+	};
+
+	// TODO DP 5.5: Create factory function that uses FCreateViewerParams and creates an IReplicationStreamViewer
 
+	/** Params for creating an IReplicationStreamEditor */
+	struct FCreateEditorParams
+	{
+		/**
+		 * The model that the editor is displaying.
+		 * @note The view will keep a strong reference to this.
+		 */
+		TSharedRef<IEditableReplicationStreamModel> DataModel;
+		
+		/**
+		 * Determines the objects that can be added to the object list. 
+		 * @note The view will keep a strong reference to this.
+		 */
+		TSharedRef<IObjectSelectionSourceModel> ObjectSource;
+		/**
+		 * Determines the properties that can be added to the property list.
+		 * @note The view will keep a strong reference to this.
+		 */
+		TSharedRef<IPropertySelectionSourceModel> PropertySource;
+		
 		/** Optional. Determines whether all UI for changing the model should be disabled. */
 		TAttribute<bool> IsEditingEnabled;
 		/** Optional. Whenever IsEditingEnabled returns true, this tooltip is displayed for relevant, disabled UI. */
@@ -114,6 +125,9 @@ namespace UE::ConcertClientSharedSlate
 
 		/** Optional settings for auto adding common properties and objects. */
 		TAttribute<const FConcertReplicationEditorSettings*> ReplicationSettingsAttribute;
+
+		/** Base params for customizing the viewing part of the editor */
+		FCreateViewerParams ViewerParams;
 	};
 
 	/**
@@ -134,7 +148,7 @@ namespace UE::ConcertClientSharedSlate
 	 *		- Similar to details panel
 	 *		- Shows properties of the selected root object and / or subobjects.
 	 */
-	CONCERTCLIENTSHAREDSLATE_API TSharedRef<IReplicationStreamEditor> CreateEditor(FCreateEditorParams Params);
+	CONCERTCLIENTSHAREDSLATE_API TSharedRef<IReplicationStreamEditor> CreateBaseStreamEditor(FCreateEditorParams Params);
 
 	/**
 	 * Creates a default IReplicationStreamEditor.
@@ -142,11 +156,61 @@ namespace UE::ConcertClientSharedSlate
 	 * This editor adds a checkbox to the start of every property row.
 	 * - Checking adds the property to the selected objects' property mappings.
 	 * - Unchecking removes the property to the selected objects' property mappings
-	 *
-	 * This requires a subobject view. If you do not pass it in, the default view, which looks like the SSubobjectEditor,
-	 * is created.
 	 * 
 	 * @see CreateViewer for a description of the UI layout.
 	 */
 	CONCERTCLIENTSHAREDSLATE_API TSharedRef<IReplicationStreamEditor> CreateDefaultStreamEditor(FCreateEditorParams Params);
-}
\ No newline at end of file
+
+	enum class EMultiStreamEditorFlags
+	{
+		None = 0,
+		
+		/** The passed in stream models transact and so should the implementation of the underlying editor. */
+		Transactional = 1 << 0
+	};
+	ENUM_CLASS_FLAGS(EMultiStreamEditorFlags)
+	
+	/** Params for creating an IMultiReplicationStreamEditor */
+	struct FCreateMultiStreamEditorParams
+	{
+		/**
+		 * The source of underlying streams that should be edited.
+		 * Must NOT contain ConsolidationDataModel (this is check()-ed).
+		 * @note The view will keep a strong reference to this.
+		 */
+		TSharedRef<IEditableMultiReplicationStreamModel> MultiStreamModel;
+		
+		/**
+		 * Determines the objects that can be added to the object list. 
+		 * @note The view will keep a strong reference to this.
+		 */
+		TSharedRef<IObjectSelectionSourceModel> ObjectSource;
+		/**
+		 * Determines the properties that can be added to the property list.
+		 * @note The view will keep a strong reference to this.
+		 */
+		TSharedRef<IPropertySelectionSourceModel> PropertySource;
+
+		/** These flags modify the behaviour */
+		EMultiStreamEditorFlags Flags = EMultiStreamEditorFlags::Transactional;
+		
+		/** Optional settings for auto adding common properties and objects. */
+		TAttribute<const FConcertReplicationEditorSettings*> ReplicationSettingsAttribute;
+		
+		/** Base params for customizing the viewing part of the editor */
+		FCreateViewerParams ViewerParams;
+	};
+
+	/**
+	 * Creates an editor that displays multiple streams in a single widget.
+	 * 
+	 * This editor consolidates all replicated objects into a single UI and display all properties in the class.
+	 * You are supposed to inject custom columns into the property section for assigning specific properties
+	 * to the streams in the IEditableReplicationStreamModel.
+	 *
+	 * The add & remove functionality of this editor adds to a
+	 *
+	 * @note This editor does NOT create any UI for changing the assigned properties. You are supposed to inject column widgets to the property rows for this.
+	 */
+	CONCERTCLIENTSHAREDSLATE_API TSharedRef<IMultiReplicationStreamEditor> CreateBaseMultiStreamEditor(FCreateMultiStreamEditorParams Params);
+}
