Split single and multi columns up into folders

diff --git a/Engine/Plugins/Developer/Concert/ConcertApp/MultiUserClient/Source/MultiUserClient/Private/Widgets/ActiveSession/Replication/Client/Columns/MultiStreamColumns.cpp b/Engine/Plugins/Developer/Concert/ConcertApp/MultiUserClient/Source/MultiUserClient/Private/Widgets/ActiveSession/Replication/Client/Columns/MultiStreamColumns.cpp
deleted file mode 100644
diff --git a/Engine/Plugins/Developer/Concert/ConcertApp/MultiUserClient/Source/MultiUserClient/Private/Widgets/ActiveSession/Replication/Client/Columns/MultiStreamColumns.h b/Engine/Plugins/Developer/Concert/ConcertApp/MultiUserClient/Source/MultiUserClient/Private/Widgets/ActiveSession/Replication/Client/Columns/MultiStreamColumns.h
deleted file mode 100644
diff --git a/Engine/Plugins/Developer/Concert/ConcertApp/MultiUserClient/Source/MultiUserClient/Private/Widgets/ActiveSession/Replication/Client/Columns/SAssignPropertyComboBox.cpp b/Engine/Plugins/Developer/Concert/ConcertApp/MultiUserClient/Source/MultiUserClient/Private/Widgets/ActiveSession/Replication/Client/Columns/SAssignPropertyComboBox.cpp
deleted file mode 100644
diff --git a/Engine/Plugins/Developer/Concert/ConcertApp/MultiUserClient/Source/MultiUserClient/Private/Widgets/ActiveSession/Replication/Client/Columns/SAssignPropertyComboBox.h b/Engine/Plugins/Developer/Concert/ConcertApp/MultiUserClient/Source/MultiUserClient/Private/Widgets/ActiveSession/Replication/Client/Columns/SAssignPropertyComboBox.h
deleted file mode 100644
diff --git a/Engine/Plugins/Developer/Concert/ConcertApp/MultiUserClient/Source/MultiUserClient/Private/Widgets/ActiveSession/Replication/Client/Columns/SOwnerClientList.cpp b/Engine/Plugins/Developer/Concert/ConcertApp/MultiUserClient/Source/MultiUserClient/Private/Widgets/ActiveSession/Replication/Client/Columns/SOwnerClientList.cpp
deleted file mode 100644
diff --git a/Engine/Plugins/Developer/Concert/ConcertApp/MultiUserClient/Source/MultiUserClient/Private/Widgets/ActiveSession/Replication/Client/Columns/SOwnerClientList.h b/Engine/Plugins/Developer/Concert/ConcertApp/MultiUserClient/Source/MultiUserClient/Private/Widgets/ActiveSession/Replication/Client/Columns/SOwnerClientList.h
deleted file mode 100644
diff --git a/Engine/Plugins/Developer/Concert/ConcertApp/MultiUserClient/Source/MultiUserClient/Private/Widgets/ActiveSession/Replication/Client/Columns/SWarningIcon.cpp b/Engine/Plugins/Developer/Concert/ConcertApp/MultiUserClient/Source/MultiUserClient/Private/Widgets/ActiveSession/Replication/Client/Columns/SWarningIcon.cpp
deleted file mode 100644
diff --git a/Engine/Plugins/Developer/Concert/ConcertApp/MultiUserClient/Source/MultiUserClient/Private/Widgets/ActiveSession/Replication/Client/Columns/SWarningIcon.h b/Engine/Plugins/Developer/Concert/ConcertApp/MultiUserClient/Source/MultiUserClient/Private/Widgets/ActiveSession/Replication/Client/Columns/SWarningIcon.h
deleted file mode 100644
diff --git a/Engine/Plugins/Developer/Concert/ConcertApp/MultiUserClient/Source/MultiUserClient/Private/Widgets/ActiveSession/Replication/Client/Columns/SingleClientColumns.h b/Engine/Plugins/Developer/Concert/ConcertApp/MultiUserClient/Source/MultiUserClient/Private/Widgets/ActiveSession/Replication/Client/Columns/SingleClientColumns.h
deleted file mode 100644
diff --git a/Engine/Plugins/Developer/Concert/ConcertApp/MultiUserClient/Source/MultiUserClient/Private/Widgets/ActiveSession/Replication/Client/Columns/SingleClientColumns_ConflictWarning.cpp b/Engine/Plugins/Developer/Concert/ConcertApp/MultiUserClient/Source/MultiUserClient/Private/Widgets/ActiveSession/Replication/Client/Columns/SingleClientColumns_ConflictWarning.cpp
deleted file mode 100644
diff --git a/Engine/Plugins/Developer/Concert/ConcertApp/MultiUserClient/Source/MultiUserClient/Private/Widgets/ActiveSession/Replication/Client/Columns/SingleClientColumns_Owner.cpp b/Engine/Plugins/Developer/Concert/ConcertApp/MultiUserClient/Source/MultiUserClient/Private/Widgets/ActiveSession/Replication/Client/Columns/SingleClientColumns_Owner.cpp
deleted file mode 100644
diff --git a/Engine/Plugins/Developer/Concert/ConcertApp/MultiUserClient/Source/MultiUserClient/Private/Widgets/ActiveSession/Replication/Client/Columns/SingleClientColumns_ToggleAuthority.cpp b/Engine/Plugins/Developer/Concert/ConcertApp/MultiUserClient/Source/MultiUserClient/Private/Widgets/ActiveSession/Replication/Client/Columns/SingleClientColumns_ToggleAuthority.cpp
deleted file mode 100644
diff --git a/Engine/Plugins/Developer/Concert/ConcertApp/MultiUserClient/Source/MultiUserClient/Private/Widgets/ActiveSession/Replication/Client/Multi/Columns/MultiStreamColumns.h b/Engine/Plugins/Developer/Concert/ConcertApp/MultiUserClient/Source/MultiUserClient/Private/Widgets/ActiveSession/Replication/Client/Multi/Columns/MultiStreamColumns.h
new file mode 100644
--- /dev/null
+++ b/Engine/Plugins/Developer/Concert/ConcertApp/MultiUserClient/Source/MultiUserClient/Private/Widgets/ActiveSession/Replication/Client/Multi/Columns/MultiStreamColumns.h
@@ -0,0 +1,44 @@
+// Copyright Epic Games, Inc. All Rights Reserved.
+
+#pragma once
+
+#include "Replication/Editor/View/PredefinedReplicationColumns.h"
+
+#include "Delegates/Delegate.h"
+#include "Misc/EnumClassFlags.h"
+
+class IConcertClient;
+
+namespace UE::MultiUserClient
+{
+	class FReplicationClientManager;
+}
+
+namespace UE::ConcertClientSharedSlate
+{
+	class IReplicationStreamModel;
+	class IMultiReplicationStreamEditor;
+	class IReplicationStreamViewer;
+}
+
+namespace UE::MultiUserClient::MultiStreamColumns
+{
+	const extern FName AssignPropertyColumnId;
+
+	/**
+	 * Creates a property column which assigns the property to the stream selected in the combo box.
+	 * 
+	 * @param MultiStreamEditor Used to determine the selected objects.
+	 * @param ConcertClient Used to look up client names
+	 * @param ClientManager Used to map streams back to client display info
+	 * @param ColumnsSortPriority The order relative to the other columns
+	 * 
+	 * @return A column that spawns a combo box for assigning properties
+	 */
+	ConcertClientSharedSlate::ReplicationColumns::FReplicationPropertyColumn AssignPropertyColumn(
+		TAttribute<TSharedPtr<ConcertClientSharedSlate::IMultiReplicationStreamEditor>> MultiStreamEditor,
+		TSharedRef<IConcertClient> ConcertClient,
+		FReplicationClientManager& ClientManager,
+		const int32 ColumnsSortPriority = 30 /* @see EReplicationPropertyColumnOrder */
+		);
+}
diff --git a/Engine/Plugins/Developer/Concert/ConcertApp/MultiUserClient/Source/MultiUserClient/Private/Widgets/ActiveSession/Replication/Client/Multi/Columns/MultiStreamColumns_AssignProperties.cpp b/Engine/Plugins/Developer/Concert/ConcertApp/MultiUserClient/Source/MultiUserClient/Private/Widgets/ActiveSession/Replication/Client/Multi/Columns/MultiStreamColumns_AssignProperties.cpp
new file mode 100644
--- /dev/null
+++ b/Engine/Plugins/Developer/Concert/ConcertApp/MultiUserClient/Source/MultiUserClient/Private/Widgets/ActiveSession/Replication/Client/Multi/Columns/MultiStreamColumns_AssignProperties.cpp
@@ -0,0 +1,114 @@
+// Copyright Epic Games, Inc. All Rights Reserved.
+
+#include "MultiStreamColumns.h"
+
+#include "IConcertClient.h"
+#include "Replication/Client/ReplicationClient.h"
+#include "Replication/Client/ReplicationClientManager.h"
+#include "Replication/Editor/Model/IEditableMultiReplicationStreamModel.h"
+#include "Replication/Editor/Model/IEditableReplicationStreamModel.h"
+#include "Replication/Editor/View/IMultiReplicationStreamEditor.h"
+#include "Replication/Editor/View/IReplicationStreamEditor.h"
+#include "SAssignPropertyComboBox.h"
+#include "Widgets/ActiveSession/Replication/Client/ClientUtils.h"
+
+#define LOCTEXT_NAMESPACE "AssignPropertyColumn"
+
+namespace UE::MultiUserClient::MultiStreamColumns
+{
+	const FName AssignPropertyColumnId(TEXT("AssignPropertyColumn"));
+
+	namespace AssignPropertyColumnUtils
+	{
+		static void ForEachStreamAssignedTo(
+			const ConcertClientSharedSlate::IMultiReplicationStreamEditor& MultiEditor,
+			const FConcertPropertyChain& Property,
+			TFunctionRef<void(const TSharedRef<ConcertClientSharedSlate::IEditableReplicationStreamModel>& Stream)> Consume
+		)
+		{
+			for (const TSharedRef<ConcertClientSharedSlate::IEditableReplicationStreamModel>& Stream : MultiEditor.GetModel().GetEditableStreams())
+			{
+				for (const FSoftObjectPath& SelectedObject : MultiEditor.GetEditorBase().GetObjectsBeingPropertyEdited())
+				{
+					const bool bStreamAssignedToProperty = Stream->HasProperty(SelectedObject, Property);
+					if (bStreamAssignedToProperty)
+					{
+						Consume(Stream);
+						// Each IEditableReplicationStreamModel should only be visited at most once.
+						break;
+					}
+				}
+			}
+		}
+
+		const FReplicationClient* FindClientByStream(const FReplicationClientManager& ClientManager, const ConcertClientSharedSlate::IReplicationStreamModel& StreamModel)
+		{
+			if (&ClientManager.GetLocalClient().GetClientEditModel().Get() == &StreamModel)
+			{
+				return &ClientManager.GetLocalClient();
+			}
+
+			for (const TNonNullPtr<FRemoteReplicationClient> Client : ClientManager.GetRemoteClients())
+			{
+				if (&Client->GetClientEditModel().Get() == &StreamModel)
+				{
+					return Client;
+				}
+			}
+			
+			return nullptr;
+		}
+		
+		static FString GetClientDisplayText(const IConcertClient& InConcertClient, const FReplicationClientManager& ClientManager, const ConcertClientSharedSlate::IReplicationStreamModel& StreamModel)
+		{
+			if (const FReplicationClient* Client = FindClientByStream(ClientManager, StreamModel))
+			{
+				return ClientUtils::GetClientDisplayName(InConcertClient, Client->GetEndpointId());
+			}
+
+			ensure(false);
+			return {};
+		}
+	}
+	
+	ConcertClientSharedSlate::ReplicationColumns::FReplicationPropertyColumn AssignPropertyColumn(
+		TAttribute<TSharedPtr<ConcertClientSharedSlate::IMultiReplicationStreamEditor>> MultiStreamEditor,
+		TSharedRef<IConcertClient> ConcertClient,
+		FReplicationClientManager& ClientManager,
+		const int32 ColumnsSortPriority
+		)
+	{
+		using namespace ConcertClientSharedSlate;
+		using namespace ConcertClientSharedSlate::ReplicationColumns;
+		check(MultiStreamEditor.IsBound() || MultiStreamEditor.IsSet());
+
+		const auto PopuluateSearch = [MultiStreamEditor, &ClientManager, ConcertClient](const FReplicatedPropertyData& Data, TArray<FString>& InOutSearchStrings)
+		{
+			AssignPropertyColumnUtils::ForEachStreamAssignedTo(*MultiStreamEditor.Get(), Data.GetProperty(),
+				[&ClientManager, &ConcertClient, &InOutSearchStrings](const TSharedRef<IEditableReplicationStreamModel>& Stream)
+				{
+					InOutSearchStrings.Add(AssignPropertyColumnUtils::GetClientDisplayText(*ConcertClient, ClientManager, *Stream));
+				});
+		};
+		const auto GenerateWidgetColumn = [MultiStreamEditor, ConcertClient, &ClientManager](const FReplicationPropertyColumn::FBuildArgs& Args)
+		{
+			const TArray<FSoftObjectPath> DisplayedObjects = MultiStreamEditor.Get()->GetEditorBase().GetObjectsBeingPropertyEdited();
+			return SNew(SAssignPropertyComboBox, MultiStreamEditor.Get().ToSharedRef(), ConcertClient, ClientManager)
+				.DisplayedProperty(Args.RowData.GetProperty())
+				.EditedObjects(DisplayedObjects)
+				.HighlightText(Args.HighlightText);
+		};
+		
+		return FReplicationPropertyColumn(
+			typename FReplicationPropertyColumn::FArguments()
+				.PopulateSearchItems_Lambda(PopuluateSearch)
+				.GenerateWidgetColumn_Lambda(GenerateWidgetColumn)
+				.ColumnSortOrder(ColumnsSortPriority),
+			SHeaderRow::Column(AssignPropertyColumnId)
+				.DefaultLabel(LOCTEXT("Owner", "Owner"))
+				.FillSized(200.f)
+			);
+	}
+}
+
+#undef LOCTEXT_NAMESPACE
\ No newline at end of file
diff --git a/Engine/Plugins/Developer/Concert/ConcertApp/MultiUserClient/Source/MultiUserClient/Private/Widgets/ActiveSession/Replication/Client/Multi/Columns/SAssignPropertyComboBox.cpp b/Engine/Plugins/Developer/Concert/ConcertApp/MultiUserClient/Source/MultiUserClient/Private/Widgets/ActiveSession/Replication/Client/Multi/Columns/SAssignPropertyComboBox.cpp
new file mode 100644
--- /dev/null
+++ b/Engine/Plugins/Developer/Concert/ConcertApp/MultiUserClient/Source/MultiUserClient/Private/Widgets/ActiveSession/Replication/Client/Multi/Columns/SAssignPropertyComboBox.cpp
@@ -0,0 +1,290 @@
+// Copyright Epic Games, Inc. All Rights Reserved.
+
+#include "SAssignPropertyComboBox.h"
+
+#include "IConcertClient.h"
+#include "Replication/Client/ReplicationClient.h"
+#include "Replication/Client/ReplicationClientManager.h"
+#include "Replication/Editor/Model/IEditableReplicationStreamModel.h"
+#include "Widgets/ActiveSession/Replication/Client/ClientUtils.h"
+#include "Widgets/ClientName/SHorizontalClientList.h"
+#include "Widgets/ClientName/SLocalClientName.h"
+#include "Widgets/ClientName/SRemoteClientName.h"
+
+#include "Algo/AnyOf.h"
+#include "Framework/MultiBox/MultiBoxBuilder.h"
+#include "ScopedTransaction.h"
+#include "Styling/SlateTypes.h"
+#include "Widgets/Input/SComboButton.h"
+
+#define LOCTEXT_NAMESPACE "SAssignPropertyComboBox"
+
+namespace UE::MultiUserClient
+{
+	void SAssignPropertyComboBox::Construct(const FArguments& InArgs,
+        TSharedRef<ConcertClientSharedSlate::IMultiReplicationStreamEditor> InEditor,
+        TSharedRef<IConcertClient> InConcertClient,
+        FReplicationClientManager& InClientManager
+	)
+	{
+		Editor = InEditor;
+		ConcertClient = MoveTemp(InConcertClient);
+		ClientManager = &InClientManager;
+		
+		Property = InArgs._DisplayedProperty;
+		EditedObjects = InArgs._EditedObjects;
+		HighlightText = InArgs._HighlightText;
+		check(!EditedObjects.IsEmpty());
+		
+		ChildSlot
+		[
+			SNew(SComboButton)
+			.HasDownArrow(true)
+			.ContentPadding(FMargin(2.0f, 2.0f))
+			.ButtonContent()
+			[
+				SAssignNew(ClientListWidget, ConcertClientSharedSlate::SHorizontalClientList, ConcertClient.ToSharedRef())
+				.HighlightText_Lambda([this](){ return HighlightText ? *HighlightText : FText::GetEmpty(); })
+				.EmptyListSlot()
+				[
+					SNew(STextBlock)
+					.Text(LOCTEXT("None", " - "))
+				]
+			]
+			.OnGetMenuContent(this, &SAssignPropertyComboBox::GetMenuContent)
+		];
+
+		ClientManager->OnRemoteClientsChanged().AddSP(this, &SAssignPropertyComboBox::RebuildSubscriptionsAndRefresh);
+		RebuildSubscriptions();
+		RefreshContentBoxContent();
+	}
+	
+	void SAssignPropertyComboBox::RefreshContentBoxContent()
+	{
+		TArray<FGuid> Clients;
+		ClientManager->ForEachClient([this, &Clients](const FReplicationClient& Client)
+		{
+			const TMap<FSoftObjectPath, FReplicatedObjectInfo>& ObjectInfoMap = Client.GetStreamSynchronizer().GetServerState().ReplicatedObjects;
+			for (const FSoftObjectPath& ObjectPath : EditedObjects)
+			{
+				if (const FReplicatedObjectInfo* ObjectInfo = ObjectInfoMap.Find(ObjectPath)
+					; ObjectInfo && ObjectInfo->PropertySelection.ReplicatedProperties.Contains(Property))
+				{
+					Clients.Add(Client.GetEndpointId());
+					return EBreakBehavior::Continue;
+				}
+			}
+			return EBreakBehavior::Continue;
+		});
+
+		ClientListWidget->RefreshList(Clients);
+	}
+
+	TSharedRef<SWidget> SAssignPropertyComboBox::GetMenuContent()
+	{
+		using namespace ConcertClientSharedSlate;
+
+		const auto MakeWidget = [this](const FGuid& EndpointId) -> TSharedRef<SWidget>
+		{
+			const bool bIsLocalClient = EndpointId == ConcertClient->GetCurrentSession()->GetSessionClientEndpointId();
+			if (bIsLocalClient)
+			{
+				return SNew(SLocalClientName, ConcertClient.ToSharedRef())
+					.HighlightText_Lambda([this](){ return HighlightText ? *HighlightText : FText::GetEmpty(); });
+			}
+			return SNew(SRemoteClientName, ConcertClient.ToSharedRef())
+				.ClientEndpointId(EndpointId)
+				.HighlightText_Lambda([this](){ return HighlightText ? *HighlightText : FText::GetEmpty(); });
+		};
+		
+		FMenuBuilder MenuBuilder(true, nullptr);
+		for (const FReplicationClient* Client : ClientUtils::GetSortedClientList(*ConcertClient, *ClientManager))
+		{
+			TAttribute<FText> Tooltip = TAttribute<FText>::CreateLambda([this, EndpointId = Client->GetEndpointId()]()
+			{
+				FText Reason;
+				const bool bCanClick = CanClickOptionWithReason(EndpointId, &Reason);
+				if (!bCanClick)
+				{
+					return Reason;
+				}
+
+				switch (GetOptionCheckState(EndpointId))
+				{
+				case ECheckBoxState::Unchecked: return LOCTEXT("Action.Unchecked", "Assign property to client and remove it from all others.");
+				case ECheckBoxState::Undetermined:  return LOCTEXT("Action.Undetermined", "Assign property to client for all selected objects.");
+				case ECheckBoxState::Checked: return LOCTEXT("Action.Checked", "Remove property from client and remove it from all others.");
+				default: return FText::GetEmpty();
+				}
+			});
+			
+			MenuBuilder.AddMenuEntry(
+				FUIAction(
+					FExecuteAction::CreateSP(this, &SAssignPropertyComboBox::OnClickOption, Client->GetEndpointId()),
+					FCanExecuteAction::CreateSP(this, &SAssignPropertyComboBox::CanClickOption, Client->GetEndpointId()),
+					FGetActionCheckState::CreateSP(this, &SAssignPropertyComboBox::GetOptionCheckState, Client->GetEndpointId())
+					),
+				MakeWidget(Client->GetEndpointId()),
+				NAME_None,
+				Tooltip,
+				EUserInterfaceActionType::ToggleButton
+				);
+		}
+		
+		return MenuBuilder.MakeWidget();
+	}
+	
+	void SAssignPropertyComboBox::OnClickOption(const FGuid EndpointId) const
+	{
+		// Remote clients can disconnect after the combo-box is opened.
+		const FReplicationClient* Client = ClientManager->FindClient(EndpointId);
+		if (!Client)
+		{
+			return;
+		}
+		
+		const FText TransactionText = FText::Format(LOCTEXT("AllClientsAssignFmt", "Assign {0} property"), FText::FromString(Property.ToString(FConcertPropertyChain::EToStringMethod::LeafProperty)));
+		FScopedTransaction Transaction(TransactionText);
+
+		// Only one client is supposed to own the object: remove the other clients (if possible)
+		ClientManager->ForEachClient([this, Client](const FReplicationClient& ClientToRemoveFrom)
+		{
+			if (*Client != ClientToRemoveFrom
+				&& ClientToRemoveFrom.AllowsEditing())
+			{
+				for (const FSoftObjectPath& ObjectPath : EditedObjects)
+				{
+					ClientToRemoveFrom.GetClientEditModel()->RemoveProperties(ObjectPath, { Property });
+				}
+			}
+			
+			return EBreakBehavior::Continue;
+		});
+
+		const ECheckBoxState CheckBoxState = GetOptionCheckState(EndpointId);
+		const bool bRemoveProperty = CheckBoxState == ECheckBoxState::Checked; 
+		const TSharedRef<ConcertClientSharedSlate::IEditableReplicationStreamModel> EditModel = Client->GetClientEditModel();
+		for (const FSoftObjectPath& ObjectPath : EditedObjects)
+		{
+			if (bRemoveProperty)
+			{
+				EditModel->RemoveProperties(ObjectPath, { Property });
+				if (!EditModel->HasAnyPropertyAssigned(ObjectPath))
+				{
+					EditModel->RemoveObjects({ ObjectPath });
+				}
+			}
+			else
+			{
+				if (!EditModel->ContainsObjects({ ObjectPath }))
+				{
+					EditModel->AddObjects({ ObjectPath.ResolveObject() });
+				}
+				
+				if (EditModel->ContainsObjects({ ObjectPath }))
+				{
+					EditModel->AddProperties(ObjectPath, { Property });
+				}
+			}
+		}
+	}
+
+#define SET_REASON(Text) if (Reason) { *Reason = Text; }
+	bool SAssignPropertyComboBox::CanClickOptionWithReason(const FGuid& EndpointId, FText* Reason) const
+	{
+		const FReplicationClient* Client = ClientManager->FindClient(EndpointId);
+		// Remote clients can disconnect after the combo-box is opened.
+		if (!Client)
+		{
+			SET_REASON(LOCTEXT("ClientDisconnected", "Client disconnected."));
+			return false;
+		}
+
+		// The combo box assigns the property to the clicked client and removes from the others... check that the currently assigned clients allow it.
+		bool bCanRemoveFromOwners = true;
+		ClientManager->ForEachClient([this, &EndpointId, &Reason, Client, &bCanRemoveFromOwners](const FReplicationClient& ClientToRemoveFrom)
+		{
+			if (*Client != ClientToRemoveFrom && !ClientToRemoveFrom.AllowsEditing())
+			{
+				const bool bHasAnySelectedObject = Algo::AnyOf(EditedObjects, [this, &ClientToRemoveFrom](const FSoftObjectPath& ObjectPath)
+				{
+					return ClientToRemoveFrom.GetClientEditModel()->HasProperty(ObjectPath, Property);
+				});
+				bCanRemoveFromOwners = !bHasAnySelectedObject;
+				
+				SET_REASON(FText::Format(
+					LOCTEXT("OwningClientDoesNotAllow", "Client {0} does not allow remote editing of its properties but has registered this property."),
+					FText::FromString(ClientUtils::GetClientDisplayName(*ConcertClient, EndpointId))
+					));
+				return EBreakBehavior::Break;
+			}
+			return EBreakBehavior::Continue;
+		});
+
+		if (!bCanRemoveFromOwners)
+		{
+			return false;
+		}
+		
+		const bool bAllowsEditing = Client->AllowsEditing();
+		if (!bAllowsEditing)
+		{
+			SET_REASON(FText::Format(
+				LOCTEXT("RemoteEditingDisabled", "Client {0} does not allow remote editing of its properties."),
+				FText::FromString(ClientUtils::GetClientDisplayName(*ConcertClient, EndpointId))
+				));
+		}
+		return bAllowsEditing;
+	}
+#undef SET_REASON
+	
+	ECheckBoxState SAssignPropertyComboBox::GetOptionCheckState(const FGuid EndpointId) const
+	{
+		const FReplicationClient* Client = ClientManager->FindClient(EndpointId);
+		// Remote clients can disconnect after the combo-box is opened.
+		if (!Client)
+		{
+			return ECheckBoxState::Unchecked;
+		}
+
+		const TSharedRef<ConcertClientSharedSlate::IEditableReplicationStreamModel> Model = Client->GetClientEditModel();
+		ECheckBoxState CheckBoxState = ECheckBoxState::Undetermined;
+		for (const FSoftObjectPath& ObjectPath : EditedObjects)
+		{
+			const bool bHasProperty = Model->HasProperty(ObjectPath, Property);
+			switch (CheckBoxState)
+			{
+			case ECheckBoxState::Unchecked:
+				if (bHasProperty)
+				{
+					return ECheckBoxState::Undetermined;
+				}
+				break;
+			case ECheckBoxState::Checked:
+				if (!bHasProperty)
+				{
+					return ECheckBoxState::Undetermined;
+				}
+				break;
+			case ECheckBoxState::Undetermined:
+				CheckBoxState = bHasProperty ? ECheckBoxState::Checked : ECheckBoxState::Unchecked;
+				break;
+			default: ;
+			}
+		}
+
+		return CheckBoxState;
+	}
+	
+	void SAssignPropertyComboBox::RebuildSubscriptions()
+	{
+		ClientManager->ForEachClient([this](FReplicationClient& Client)
+		{
+			Client.OnModelChanged().RemoveAll(this);
+			Client.OnModelChanged().AddSP(this, &SAssignPropertyComboBox::RefreshContentBoxContent);
+			return EBreakBehavior::Continue;
+		});
+	}
+}
+
+#undef LOCTEXT_NAMESPACE
\ No newline at end of file
diff --git a/Engine/Plugins/Developer/Concert/ConcertApp/MultiUserClient/Source/MultiUserClient/Private/Widgets/ActiveSession/Replication/Client/Multi/Columns/SAssignPropertyComboBox.h b/Engine/Plugins/Developer/Concert/ConcertApp/MultiUserClient/Source/MultiUserClient/Private/Widgets/ActiveSession/Replication/Client/Multi/Columns/SAssignPropertyComboBox.h
new file mode 100644
--- /dev/null
+++ b/Engine/Plugins/Developer/Concert/ConcertApp/MultiUserClient/Source/MultiUserClient/Private/Widgets/ActiveSession/Replication/Client/Multi/Columns/SAssignPropertyComboBox.h
@@ -0,0 +1,90 @@
+// Copyright Epic Games, Inc. All Rights Reserved.
+
+#pragma once
+
+#include "Replication/Data/ConcertPropertySelection.h"
+
+#include "Containers/Array.h"
+#include "Widgets/DeclarativeSyntaxSupport.h"
+#include "Widgets/SCompoundWidget.h"
+
+class IConcertClient;
+
+enum class ECheckBoxState : uint8;
+
+struct FGuid;
+
+namespace UE::ConcertClientSharedSlate
+{
+	class SHorizontalClientList;
+	class IMultiReplicationStreamEditor;
+}
+
+namespace UE::MultiUserClient
+{
+	class FReplicationClientManager;
+	
+	/**
+	 * Placed in every property row column to assign properties to clients.
+	 * 
+	 * Can be used to assign a property to at most one client: multiple clients cannot be assigned because it is an advanced workflow
+	 * which might confuse basic users.
+	 */
+	class SAssignPropertyComboBox : public SCompoundWidget
+	{
+	public:
+
+		SLATE_BEGIN_ARGS(SAssignPropertyComboBox)
+		{}
+			SLATE_ARGUMENT(FConcertPropertyChain, DisplayedProperty)
+			SLATE_ARGUMENT(TArray<FSoftObjectPath>, EditedObjects)
+			SLATE_ARGUMENT(TSharedPtr<FText>, HighlightText)
+		SLATE_END_ARGS()
+
+		void Construct(const FArguments& InArgs,
+           TSharedRef<ConcertClientSharedSlate::IMultiReplicationStreamEditor> InEditor,
+           TSharedRef<IConcertClient> InConcertClient,
+           FReplicationClientManager& InClientManager
+		);
+
+	private:
+
+		FReplicationClientManager* ClientManager = nullptr;
+		
+		/** Used to obtain info about the streams */
+		TSharedPtr<ConcertClientSharedSlate::IMultiReplicationStreamEditor> Editor;
+		TSharedPtr<IConcertClient> ConcertClient;
+
+		/** The objects for which the property is being displayed. */
+		TArray<FSoftObjectPath> EditedObjects;
+		/** The property assigned to this column. */
+		FConcertPropertyChain Property;
+
+		/** The static menu content (when there is no drop-down). */
+		TSharedPtr<ConcertClientSharedSlate::SHorizontalClientList> ClientListWidget;
+
+		/** Passed to MakeListWidgetDelegate */
+		TSharedPtr<FText> HighlightText;
+
+		/** Updates the content of the combo box */
+		void RefreshContentBoxContent();
+		
+		/** Builds the drop-down menu */
+		TSharedRef<SWidget> GetMenuContent();
+
+		/** Called when a client is clicked in the drop-down menu */
+		void OnClickOption(const FGuid EndpointId) const;
+		
+		/** @return Whether an item in the drop-down menu should be clickable (returns true if the stream is writable) */
+		bool CanClickOption(const FGuid EndpointId) const { return CanClickOptionWithReason(EndpointId, nullptr); }
+		bool CanClickOptionWithReason(const FGuid& EndpointId, FText* Reason = nullptr) const;
+
+		/** @return Checkbox state for the item. Handles cases of multiple objects being property edited. */
+		ECheckBoxState GetOptionCheckState(const FGuid EndpointId) const;
+
+		/** Resubscribes to all clients changing. */
+		void RebuildSubscriptions();
+		void RebuildSubscriptionsAndRefresh() { RebuildSubscriptions(); RefreshContentBoxContent(); }
+	};
+}
+
diff --git a/Engine/Plugins/Developer/Concert/ConcertApp/MultiUserClient/Source/MultiUserClient/Private/Widgets/ActiveSession/Replication/Client/Multi/SMultiClientView.cpp b/Engine/Plugins/Developer/Concert/ConcertApp/MultiUserClient/Source/MultiUserClient/Private/Widgets/ActiveSession/Replication/Client/Multi/SMultiClientView.cpp
--- a/Engine/Plugins/Developer/Concert/ConcertApp/MultiUserClient/Source/MultiUserClient/Private/Widgets/ActiveSession/Replication/Client/Multi/SMultiClientView.cpp
+++ b/Engine/Plugins/Developer/Concert/ConcertApp/MultiUserClient/Source/MultiUserClient/Private/Widgets/ActiveSession/Replication/Client/Multi/SMultiClientView.cpp
@@ -10,13 +10,10 @@
 #include "Replication/Editor/Model/Property/SelectPropertyFromUClassModel.h"
 #include "Replication/Editor/View/IMultiReplicationStreamEditor.h"
 #include "Replication/ReplicationWidgetFactories.h"
-#include "Widgets/ActiveSession/Replication/Client/Columns/MultiStreamColumns.h"
+#include "Widgets/ActiveSession/Replication/Client/Multi/Columns/MultiStreamColumns.h"
 #include "Widgets/ActiveSession/Replication/Client/SClientToolbar.h"
 
-#include "Algo/Transform.h"
-#include "UObject/Package.h"
 #include "Widgets/SBoxPanel.h"
-#include "Widgets/Text/STextBlock.h"
 
 #define LOCTEXT_NAMESPACE "SMultiClientView"
 
diff --git a/Engine/Plugins/Developer/Concert/ConcertApp/MultiUserClient/Source/MultiUserClient/Private/Widgets/ActiveSession/Replication/Client/SWarningIcon.cpp b/Engine/Plugins/Developer/Concert/ConcertApp/MultiUserClient/Source/MultiUserClient/Private/Widgets/ActiveSession/Replication/Client/SWarningIcon.cpp
new file mode 100644
--- /dev/null
+++ b/Engine/Plugins/Developer/Concert/ConcertApp/MultiUserClient/Source/MultiUserClient/Private/Widgets/ActiveSession/Replication/Client/SWarningIcon.cpp
@@ -0,0 +1,25 @@
+// Copyright Epic Games, Inc. All Rights Reserved.
+
+#include "SWarningIcon.h"
+
+#include "Styling/AppStyle.h"
+#include "Widgets/Images/SImage.h"
+#include "Widgets/Layout/SScaleBox.h"
+
+namespace UE::MultiUserClient
+{
+	void SWarningIcon::Construct(const FArguments& InArgs)
+	{
+		ChildSlot
+		[
+			SNew(SScaleBox)
+			.Stretch(EStretch::ScaleToFit)
+			[
+				SNew(SImage)
+				.Image(FAppStyle::Get().GetBrush("Icons.WarningWithColor"))
+				.DesiredSizeOverride(FVector2D(16,16))
+			]
+		];
+	}
+}
+
diff --git a/Engine/Plugins/Developer/Concert/ConcertApp/MultiUserClient/Source/MultiUserClient/Private/Widgets/ActiveSession/Replication/Client/SWarningIcon.h b/Engine/Plugins/Developer/Concert/ConcertApp/MultiUserClient/Source/MultiUserClient/Private/Widgets/ActiveSession/Replication/Client/SWarningIcon.h
new file mode 100644
--- /dev/null
+++ b/Engine/Plugins/Developer/Concert/ConcertApp/MultiUserClient/Source/MultiUserClient/Private/Widgets/ActiveSession/Replication/Client/SWarningIcon.h
@@ -0,0 +1,23 @@
+// Copyright Epic Games, Inc. All Rights Reserved.
+
+#pragma once
+
+#include "Widgets/DeclarativeSyntaxSupport.h"
+#include "Widgets/SCompoundWidget.h"
+
+namespace UE::MultiUserClient
+{
+	/** Displays a warning triangle */
+	class SWarningIcon : public SCompoundWidget
+	{
+	public:
+
+		SLATE_BEGIN_ARGS(SWarningIcon)
+		{}
+		SLATE_END_ARGS()
+
+		void Construct(const FArguments& InArgs);
+	};
+}
+
+
diff --git a/Engine/Plugins/Developer/Concert/ConcertApp/MultiUserClient/Source/MultiUserClient/Private/Widgets/ActiveSession/Replication/Client/Single/Columns/SOwnerClientList.cpp b/Engine/Plugins/Developer/Concert/ConcertApp/MultiUserClient/Source/MultiUserClient/Private/Widgets/ActiveSession/Replication/Client/Single/Columns/SOwnerClientList.cpp
new file mode 100644
--- /dev/null
+++ b/Engine/Plugins/Developer/Concert/ConcertApp/MultiUserClient/Source/MultiUserClient/Private/Widgets/ActiveSession/Replication/Client/Single/Columns/SOwnerClientList.cpp
@@ -0,0 +1,59 @@
+// Copyright Epic Games, Inc. All Rights Reserved.
+
+#include "SOwnerClientList.h"
+
+#include "Replication/Util/GlobalAuthorityCache.h"
+#include "Widgets/ClientName/SHorizontalClientList.h"
+
+#include "Widgets/Layout/SBox.h"
+#include "Widgets/Text/STextBlock.h"
+
+#define LOCTEXT_NAMESPACE "SOwnerClientList"
+
+namespace UE::MultiUserClient
+{
+	void SOwnerClientList::Construct(const FArguments& InArgs, TSharedRef<IConcertClient> InClient, FGlobalAuthorityCache& InAuthorityCache)
+	{
+		AuthorityCache = &InAuthorityCache;
+		GetClientListDelegate = InArgs._GetClientList;
+		check(GetClientListDelegate.IsBound());
+		
+		AuthorityCache->OnCacheChanged().AddSP(this, &SOwnerClientList::OnClientChanged);
+		
+		ChildSlot
+		[
+			SAssignNew(ClientList, ConcertClientSharedSlate::SHorizontalClientList, MoveTemp(InClient))
+			.HighlightText(InArgs._HighlightText)
+			.EmptyListSlot()
+			[
+				SNew(SBox)
+				.HAlign(HAlign_Left)
+				.Padding(5.f, 0.f, 0.f, 0.f)
+				[
+					SNew(STextBlock)
+					.Text(LOCTEXT("NotApplicable", "-"))
+				]
+			]
+		];
+
+		RefreshList();
+	}
+
+	SOwnerClientList::~SOwnerClientList()
+	{
+		AuthorityCache->OnCacheChanged().RemoveAll(this);
+	}
+
+	void SOwnerClientList::OnClientChanged(const FGuid& ChangedClientId)
+	{
+		RefreshList();
+	}
+	
+	void SOwnerClientList::RefreshList()
+	{
+		const TArray<FGuid> Clients = GetClientListDelegate.Execute(*AuthorityCache);
+		ClientList->RefreshList(Clients);
+	}
+}
+
+#undef LOCTEXT_NAMESPACE
\ No newline at end of file
diff --git a/Engine/Plugins/Developer/Concert/ConcertApp/MultiUserClient/Source/MultiUserClient/Private/Widgets/ActiveSession/Replication/Client/Single/Columns/SOwnerClientList.h b/Engine/Plugins/Developer/Concert/ConcertApp/MultiUserClient/Source/MultiUserClient/Private/Widgets/ActiveSession/Replication/Client/Single/Columns/SOwnerClientList.h
new file mode 100644
--- /dev/null
+++ b/Engine/Plugins/Developer/Concert/ConcertApp/MultiUserClient/Source/MultiUserClient/Private/Widgets/ActiveSession/Replication/Client/Single/Columns/SOwnerClientList.h
@@ -0,0 +1,56 @@
+// Copyright Epic Games, Inc. All Rights Reserved.
+
+#pragma once
+
+#include "Widgets/DeclarativeSyntaxSupport.h"
+#include "Widgets/SCompoundWidget.h"
+
+namespace UE::ConcertClientSharedSlate
+{
+	class SHorizontalClientList;
+}
+
+class IConcertClient;
+
+namespace UE::MultiUserClient
+{
+	class FGlobalAuthorityCache;
+	
+	/**
+	 * Displays a list of authoritive clients in the SReplicationClientView columns for top-level, subobject and properties;
+	 * the passed in FGetClientList determines which is displayed
+	 */
+	class SOwnerClientList : public SCompoundWidget
+	{
+	public:
+
+		DECLARE_DELEGATE_RetVal_OneParam(TArray<FGuid>, FGetClientList, const FGlobalAuthorityCache&);
+		
+		SLATE_BEGIN_ARGS(SOwnerClientList)
+		{}
+			/** Obtains the client list to display */
+			SLATE_EVENT(FGetClientList, GetClientList)
+			/** Used for highlighting in the text */
+			SLATE_ATTRIBUTE(FText, HighlightText)
+		SLATE_END_ARGS()
+
+		void Construct(const FArguments& InArgs, TSharedRef<IConcertClient> InClient, FGlobalAuthorityCache& InAuthorityCache);
+		virtual ~SOwnerClientList();
+
+	private:
+		
+		/** Tells us when the authority change */
+		FGlobalAuthorityCache* AuthorityCache;
+
+		/** Displays the clients */
+		TSharedPtr<ConcertClientSharedSlate::SHorizontalClientList> ClientList;
+
+		/** Used to get the client list from FGlobalAuthorityCache */
+		FGetClientList GetClientListDelegate;
+
+		void OnClientChanged(const FGuid& ChangedClientId);
+		void RefreshList();
+	};
+}
+
+
diff --git a/Engine/Plugins/Developer/Concert/ConcertApp/MultiUserClient/Source/MultiUserClient/Private/Widgets/ActiveSession/Replication/Client/Single/Columns/SingleClientColumns.h b/Engine/Plugins/Developer/Concert/ConcertApp/MultiUserClient/Source/MultiUserClient/Private/Widgets/ActiveSession/Replication/Client/Single/Columns/SingleClientColumns.h
new file mode 100644
--- /dev/null
+++ b/Engine/Plugins/Developer/Concert/ConcertApp/MultiUserClient/Source/MultiUserClient/Private/Widgets/ActiveSession/Replication/Client/Single/Columns/SingleClientColumns.h
@@ -0,0 +1,160 @@
+// Copyright Epic Games, Inc. All Rights Reserved.
+
+#pragma once
+
+#include "Replication/Editor/View/PredefinedReplicationColumns.h"
+#include "Widgets/DeclarativeSyntaxSupport.h"
+#include "Widgets/SCompoundWidget.h"
+
+class IConcertClient;
+
+namespace UE::ConcertClientSharedSlate
+{
+	class IReplicationStreamViewer;
+	class IReplicationStreamModel;
+}
+
+namespace UE::MultiUserClient
+{
+	class FAuthorityChangeTracker;
+	class FGlobalAuthorityCache;
+	class IClientAuthoritySynchronizer;
+	class ISubmissionWorkflow;
+}
+
+namespace UE::MultiUserClient::SingleClientColumns
+{
+	/** @see UE::ConcertClientSharedSlate::ReplicationColumns::Subobject::ETopLevelColumnOrder */
+	enum class ETopLevelObjectColumnOrder
+	{
+		ToggleAuthority = 0,
+		Owner = 40
+	};
+	/** @see UE::ConcertClientSharedSlate::ReplicationColumns::Subobject::ESubobjectColumnOrder */
+	enum class ESubobjectColumnOrder
+	{
+		ToggleAuthority = 0,
+		ConflictWarning = 5,
+		Owner = 40
+	};
+	/** @see UE::ConcertClientSharedSlate::ReplicationColumns::Property::EReplicationPropertyColumnOrder */
+	enum class EPropertyColumnOrder
+	{
+		ConflictWarning = 5,
+		Owner = 50
+	};
+
+
+	/********** Toggle authority **********/
+	extern const FName ToggleTopLevelAuthorityColumnId;
+	extern const FName ToggleSubobjectAuthorityColumnId;
+
+	/**
+	 * Checkbox placed in the top-level view.
+	 * It gives / removes authority for the top-level object and all of its subobjects.
+	 *
+	 * @param ClientStreamModel Used to discover the client's registered subobjects
+	 * @param ChangeTracker Used to determine the checkbox state and change authority.
+	 * @param SubmissionWorkflow Used to determine whether changing authority is at all enabled
+	 * @return Column that can be placed in the table
+	 */
+	ConcertClientSharedSlate::ReplicationColumns::FReplicationTopLevelObjectColumn ToggleTopLevelAuthority(
+		ConcertClientSharedSlate::IReplicationStreamModel& ClientStreamModel,
+		FAuthorityChangeTracker& ChangeTracker,
+		ISubmissionWorkflow& SubmissionWorkflow
+		);
+
+	/**
+	 * Checkbox placed in the subobject view
+	 * It gives / removes authority for the object it is placed next to.
+	 * 
+	 * @param AuthoritySynchronizer Determines whether the client has authority or not.
+	 * @param ChangeTracker Used to determine the checkbox state and change authority.
+	 * @param AuthoritySynchronizer Used to determine whether changing authority is at all enabled
+	 * @return Column that can be placed in the table
+	 */
+	ConcertClientSharedSlate::ReplicationColumns::FReplicationSubobjectObjectColumn ToggleSubobjectAuthority(
+		FAuthorityChangeTracker& ChangeTracker,
+		ISubmissionWorkflow& SubmissionWorkflow
+		);
+
+	
+	/********** Owner **********/
+	extern const FName OwnerOfSubobjectColumnId;
+	extern const FName OwnerOfPropertyColumnId;
+	
+	/**
+	 * Displays the owner of a top-level object.
+	 * @param InClient The local Concert client used to look up other client display info
+	 * @param InAuthorityCache Used to determine which client owns the property
+	 * @param InObjectModel Used to get subobjects for a top-level object
+	 * @return Column that can be placed in the table
+	 */
+	ConcertClientSharedSlate::ReplicationColumns::FReplicationTopLevelObjectColumn OwnerOfTopLevelObject(
+		const TSharedRef<IConcertClient>& InClient,
+		FGlobalAuthorityCache& InAuthorityCache,
+		ConcertClientSharedSlate::IReplicationStreamModel& InObjectModel
+		);
+
+	/**
+	 * Displays the owner of a subobject.
+	 * @param InClient The local Concert client used to look up other client display info
+	 * @param InAuthorityCache Used to determine which client owns the property
+	 * @return Column that can be placed in the table
+	 */
+	ConcertClientSharedSlate::ReplicationColumns::FReplicationSubobjectObjectColumn OwnerOfSubobject(
+		const TSharedRef<IConcertClient>& InClient,
+		FGlobalAuthorityCache& InAuthorityCache
+		);
+	
+	/**
+	 * Displays the owner of a property.
+	 * @param InClient The local Concert client used to look up other client display info
+	 * @param InAuthorityCache Used to determine which client owns the property
+	 * @param InViewerAttribute Used to determine which objects the property box is displaying
+	 * @return Column that can be placed in the table
+	 */
+	ConcertClientSharedSlate::ReplicationColumns::FReplicationPropertyColumn OwnerOfProperty(
+		const TSharedRef<IConcertClient>& InClient,
+		FGlobalAuthorityCache& InAuthorityCache,
+		const TAttribute<const ConcertClientSharedSlate::IReplicationStreamViewer*>& InViewerAttribute
+		);
+
+	
+	/********** Conflict warning **********/
+	extern const FName ConflictWarningTopLevelObjectColumnId;
+	extern const FName ConflictWarningSubobjectColumnId;
+	extern const FName ConflictWarningPropertyColumnId;
+	
+	/**
+	 * Displays a warning symbol next to the checkbox if checking the checkbox would cause an authority conflict when submitted.
+	 * 
+	 * @param InClient The local Concert client used to look up other client display info
+	 * @param InAuthorityCache Used to determine which client owns the object
+	 * @param ClientId ID of the client for which the icon is being created
+	 * 
+	 * @return Column that can be placed in the table
+	 */
+	ConcertClientSharedSlate::ReplicationColumns::FReplicationSubobjectObjectColumn ConflictWarningForSubobject(
+		TSharedRef<IConcertClient> InClient,
+		FGlobalAuthorityCache& InAuthorityCache,
+		const FGuid& ClientId
+		);
+
+	/**
+	 * Displays a warning symbol next to the checkbox if checking the checkbox would cause an authority conflict when submitted.
+	 * 
+	 * @param InClient The local Concert client used to look up other client display info
+	 * @param InViewer Used to determine which objects the property box is displaying
+	 * @param InAuthorityCache Used to determine which client owns the property
+	 * @param ClientId ID of the client for which to check whether the property can be added to the stream
+	 * 
+	 * @return Column that can be placed in the table
+	 */
+	ConcertClientSharedSlate::ReplicationColumns::FReplicationPropertyColumn ConflictWarningForProperty(
+		TSharedRef<IConcertClient> InClient,
+		TAttribute<const ConcertClientSharedSlate::IReplicationStreamViewer*> InViewer,
+		FGlobalAuthorityCache& InAuthorityCache,
+		const FGuid& ClientId
+		);
+}
diff --git a/Engine/Plugins/Developer/Concert/ConcertApp/MultiUserClient/Source/MultiUserClient/Private/Widgets/ActiveSession/Replication/Client/Single/Columns/SingleClientColumns_ConflictWarning.cpp b/Engine/Plugins/Developer/Concert/ConcertApp/MultiUserClient/Source/MultiUserClient/Private/Widgets/ActiveSession/Replication/Client/Single/Columns/SingleClientColumns_ConflictWarning.cpp
new file mode 100644
--- /dev/null
+++ b/Engine/Plugins/Developer/Concert/ConcertApp/MultiUserClient/Source/MultiUserClient/Private/Widgets/ActiveSession/Replication/Client/Single/Columns/SingleClientColumns_ConflictWarning.cpp
@@ -0,0 +1,129 @@
+// Copyright Epic Games, Inc. All Rights Reserved.
+
+#include "SingleClientColumns.h"
+
+#include "IConcertClient.h"
+#include "Replication/Editor/Model/ReplicatedObjectData.h"
+#include "Replication/Editor/View/IReplicationStreamViewer.h"
+#include "Replication/Util/GlobalAuthorityCache.h"
+#include "Widgets/ActiveSession/Replication/Client/ClientUtils.h"
+#include "Widgets/ActiveSession/Replication/Client/SWarningIcon.h"
+
+#include "Algo/AllOf.h"
+#include "Widgets/Images/SImage.h"
+#include "Widgets/Views/SHeaderRow.h"
+
+#define LOCTEXT_NAMESPACE "SingleClientColumns.ConflictWarning"
+
+namespace UE::MultiUserClient::SingleClientColumns
+{
+	const FName ConflictWarningTopLevelObjectColumnId = TEXT("ConflictWarningTopLevelObjectColumn");
+	const FName ConflictWarningSubobjectColumnId = TEXT("ConflictWarningSubobjectColumn");
+	const FName ConflictWarningPropertyColumnId = TEXT("ConflictWarningPropertyColumn");
+	
+	ConcertClientSharedSlate::ReplicationColumns::FReplicationSubobjectObjectColumn ConflictWarningForSubobject(
+		TSharedRef<IConcertClient> InClient,
+		FGlobalAuthorityCache& InAuthorityCache,
+		const FGuid& ClientId
+		)
+	{
+		using FColumnType = ConcertClientSharedSlate::ReplicationColumns::FReplicationSubobjectObjectColumn;
+		return FColumnType(
+			FColumnType::FArguments()
+				.GenerateWidgetColumn_Lambda([InClient = MoveTemp(InClient), &InAuthorityCache, ClientId](const FColumnType::FBuildArgs& InArgs) mutable
+				{
+					const auto GetVisibility = [&InAuthorityCache, ClientId, ObjectPath = InArgs.RowData.GetObjectPath()]()
+					{
+						const FGlobalAuthorityCache::ECanTakeAuthority TakeAuthorityResult = InAuthorityCache.CanClientTakeAuthorityAfterSubmission(ObjectPath, ClientId);
+						return TakeAuthorityResult == FGlobalAuthorityCache::ECanTakeAuthority::Conflict
+							? EVisibility::Visible
+							: EVisibility::Collapsed;
+					};
+					const auto GetToolTip = [InClient = MoveTemp(InClient), &InAuthorityCache, ClientId, ObjectPath = InArgs.RowData.GetObjectPath()]()
+					{
+						TSet<FString> ClientNames;
+						const FGlobalAuthorityCache::ECanTakeAuthority TakeAuthorityResult = InAuthorityCache.CanClientTakeAuthorityAfterSubmission(ObjectPath, ClientId,
+						[InClient, &ClientNames](const FGuid& ClientId, const FConcertPropertyChain& ConflictingProperty)
+						{
+							ClientNames.Add(ClientUtils::GetClientDisplayName(*InClient, ClientId));
+							return EBreakBehavior::Continue;
+						});
+						
+						return TakeAuthorityResult != FGlobalAuthorityCache::ECanTakeAuthority::Conflict
+							? FText::GetEmpty()
+							: FText::Format(
+								LOCTEXT("Subobject.WarningFmt", "{0} {1}|plural(one=is,other=are) replicating some of the assigned properties. This change will not be sent to the server."),
+								FText::FromString(FString::Join(ClientNames, TEXT(", "))),
+								ClientNames.Num()
+							);
+					};
+					
+					return SNew(SWarningIcon)
+						.Visibility_Lambda(GetVisibility)
+						.ToolTipText_Lambda(GetToolTip);
+				})
+				.PopulateSearchItems_Lambda([](const ConcertClientSharedSlate::FReplicatedObjectData&, TArray<FString>&){})
+				.ColumnSortOrder(static_cast<int32>(ESubobjectColumnOrder::ConflictWarning)),
+			SHeaderRow::Column(ConflictWarningSubobjectColumnId)
+				.DefaultLabel(FText::GetEmpty())
+				.FixedWidth(20.f)
+			);
+	}
+
+	ConcertClientSharedSlate::ReplicationColumns::FReplicationPropertyColumn ConflictWarningForProperty(
+		TSharedRef<IConcertClient> InClient,
+		TAttribute<const ConcertClientSharedSlate::IReplicationStreamViewer*> Viewer,
+		FGlobalAuthorityCache& InAuthorityCache,
+		const FGuid& ClientId
+		)
+	{
+		using FColumnType = ConcertClientSharedSlate::ReplicationColumns::FReplicationPropertyColumn;
+		return FColumnType(
+			FColumnType::FArguments()
+				.GenerateWidgetColumn_Lambda([InClient = MoveTemp(InClient), Viewer, &InAuthorityCache, ClientId](const FColumnType::FBuildArgs& InArgs) mutable
+				{
+					const auto GetVisibility = [&Viewer, &InAuthorityCache, ClientId, Property = InArgs.RowData.GetProperty()]()
+					{
+						const bool bCanAddProperty = Algo::AllOf(Viewer.Get()->GetObjectsBeingPropertyEdited(),
+							[&InAuthorityCache, ClientId, Property](const FSoftObjectPath& ObjectPath)
+							{
+								return InAuthorityCache.CanClientAddProperty(ObjectPath, ClientId, Property);
+							});
+						return bCanAddProperty
+							? EVisibility::Collapsed
+							: EVisibility::Visible;
+					};
+					const auto GetToolTip = [InClient = MoveTemp(InClient), Viewer, &InAuthorityCache, ClientId, Property = InArgs.RowData.GetProperty()]()
+					{
+						TSet<FString> Names;
+						for (const FSoftObjectPath& ObjectPath: Viewer.Get()->GetObjectsBeingPropertyEdited())
+						{
+							const TOptional<FGuid> Owner = InAuthorityCache.GetClientWithAuthorityOverProperty(ObjectPath, Property);
+							if (Owner || *Owner != ClientId)
+							{
+								Names.Add(ClientUtils::GetClientDisplayName(*InClient, *Owner));
+							}
+						}
+						
+						return Names.IsEmpty()
+							? FText::GetEmpty()
+							: FText::Format(
+								LOCTEXT("Property.WarningFmt", "Replicated by {0}. This property will not be submitted to the server."),
+								FText::FromString(FString::Join(Names, TEXT(", ")))
+								);
+					};
+					
+					return SNew(SWarningIcon)
+						.Visibility_Lambda(GetVisibility)
+						.ToolTipText_Lambda(GetToolTip);
+				})
+				.PopulateSearchItems_Lambda([](const ConcertClientSharedSlate::FReplicatedPropertyData&, TArray<FString>&){})
+				.ColumnSortOrder(static_cast<int32>(EPropertyColumnOrder::ConflictWarning)),
+			SHeaderRow::Column(ConflictWarningPropertyColumnId)
+				.DefaultLabel(FText::GetEmpty())
+				.FixedWidth(20.f)
+			);
+	}
+}
+
+#undef LOCTEXT_NAMESPACE
\ No newline at end of file
diff --git a/Engine/Plugins/Developer/Concert/ConcertApp/MultiUserClient/Source/MultiUserClient/Private/Widgets/ActiveSession/Replication/Client/Single/Columns/SingleClientColumns_Owner.cpp b/Engine/Plugins/Developer/Concert/ConcertApp/MultiUserClient/Source/MultiUserClient/Private/Widgets/ActiveSession/Replication/Client/Single/Columns/SingleClientColumns_Owner.cpp
new file mode 100644
--- /dev/null
+++ b/Engine/Plugins/Developer/Concert/ConcertApp/MultiUserClient/Source/MultiUserClient/Private/Widgets/ActiveSession/Replication/Client/Single/Columns/SingleClientColumns_Owner.cpp
@@ -0,0 +1,175 @@
+// Copyright Epic Games, Inc. All Rights Reserved.
+
+#include "SingleClientColumns.h"
+
+#include "IConcertClient.h"
+#include "SOwnerClientList.h"
+#include "Replication/Editor/View/IReplicationStreamViewer.h"
+#include "Replication/Editor/View/ReplicationColumnsUtils.h"
+#include "Replication/Util/GlobalAuthorityCache.h"
+#include "Widgets/ActiveSession/Replication/Client/ClientUtils.h"
+#include "Widgets/ClientName/SClientName.h"
+
+#define LOCTEXT_NAMESPACE "SingleClientColumns.Owner"
+
+namespace UE::MultiUserClient::SingleClientColumns
+{
+	const FName OwnerOfSubobjectColumnId = TEXT("OwnerOfSubobjectColumn");
+	const FName OwnerOfPropertyColumnId = TEXT("OwnerOfPropertyColumn");
+
+	namespace Private::Shared
+	{
+		constexpr float DefaultWidth = 200.f;
+		
+		static void PopulateSearchTerms(const TSharedRef<IConcertClient>& InClient, const TArray<FGuid>& Clients, TArray<FString>& InOutSearchStrings)
+		{
+			for (const FGuid& ClientId : Clients)
+			{
+				InOutSearchStrings.Add(ClientUtils::GetClientDisplayName(*InClient, ClientId));
+			}
+		}
+	}
+
+	namespace Private::TopLevel
+	{
+		static TArray<FGuid> GetOwnersOfTopLevelObject(
+			const FGlobalAuthorityCache& InAuthorityCache,
+			ConcertClientSharedSlate::IReplicationStreamModel& InObjectModel,
+			const FSoftObjectPath& InTopLevelObject
+			)
+		{
+			TArray<FGuid> Owners = InAuthorityCache.GetClientsWithAuthorityOverObject(InTopLevelObject);
+			InObjectModel.ForEachSubobject(InTopLevelObject, [&InAuthorityCache, &Owners](const FSoftObjectPath& Subobject)
+			{
+				for (const FGuid& Client : InAuthorityCache.GetClientsWithAuthorityOverObject(Subobject))
+				{
+					Owners.AddUnique(Client);
+				}
+				return EBreakBehavior::Continue;
+			});
+			return Owners;
+		}
+	}
+	
+	ConcertClientSharedSlate::ReplicationColumns::FReplicationTopLevelObjectColumn OwnerOfTopLevelObject(
+		const TSharedRef<IConcertClient>& InClient,
+		FGlobalAuthorityCache& InAuthorityCache,
+		ConcertClientSharedSlate::IReplicationStreamModel& InObjectModel
+		)
+	{
+		using FColumnType = ConcertClientSharedSlate::ReplicationColumns::FReplicationSubobjectObjectColumn;
+		return FColumnType(
+			FColumnType::FArguments()
+				.GenerateWidgetColumn_Lambda([InClient, &InObjectModel, &InAuthorityCache](const FColumnType::FBuildArgs& InArgs)
+				{
+					return SNew(SOwnerClientList, InClient, InAuthorityCache)
+						.GetClientList_Lambda([&InAuthorityCache, &InObjectModel, ObjectPath = InArgs.RowData.GetObjectPath()](const FGlobalAuthorityCache&)
+						{
+							return Private::TopLevel::GetOwnersOfTopLevelObject(InAuthorityCache, InObjectModel, ObjectPath);
+						})
+						.HighlightText_Lambda([HighlightText = InArgs.HighlightText](){ return *HighlightText.Get(); });
+				})
+				.PopulateSearchItems_Lambda([InClient, &InAuthorityCache, &InObjectModel](const ConcertClientSharedSlate::FReplicatedObjectData& InArgs, TArray<FString>& InOutSearchStrings)
+				{
+					Private::Shared::PopulateSearchTerms(
+						InClient,
+						Private::TopLevel::GetOwnersOfTopLevelObject(InAuthorityCache, InObjectModel, InArgs.GetObjectPath()),
+						InOutSearchStrings
+						);
+				})
+				.ColumnSortOrder(static_cast<int32>(ETopLevelObjectColumnOrder::Owner)),
+			SHeaderRow::Column(OwnerOfSubobjectColumnId)
+				.DefaultLabel(LOCTEXT("Subobject.Owner", "Owner"))
+				.FillSized(Private::Shared::DefaultWidth)
+			);
+	}
+	
+	ConcertClientSharedSlate::ReplicationColumns::FReplicationSubobjectObjectColumn OwnerOfSubobject(
+		const TSharedRef<IConcertClient>& InClient,
+		FGlobalAuthorityCache& InAuthorityCache
+		)
+	{
+		using FColumnType = ConcertClientSharedSlate::ReplicationColumns::FReplicationSubobjectObjectColumn;
+		return FColumnType(
+			FColumnType::FArguments()
+				.GenerateWidgetColumn_Lambda([InClient, &InAuthorityCache](const FColumnType::FBuildArgs& InArgs)
+				{
+					return SNew(SOwnerClientList, InClient, InAuthorityCache)
+						.GetClientList_Lambda([&InAuthorityCache, ObjectPath = InArgs.RowData.GetObjectPath()](const FGlobalAuthorityCache&)
+						{
+							return InAuthorityCache.GetClientsWithAuthorityOverObject(ObjectPath);
+						})
+						.HighlightText_Lambda([HighlightText = InArgs.HighlightText](){ return *HighlightText.Get(); });
+				})
+				.PopulateSearchItems_Lambda([InClient, &InAuthorityCache](const ConcertClientSharedSlate::FReplicatedObjectData& InArgs, TArray<FString>& InOutSearchStrings)
+				{
+					Private::Shared::PopulateSearchTerms(
+						InClient,
+						InAuthorityCache.GetClientsWithAuthorityOverObject(InArgs.GetObjectPath()),
+						InOutSearchStrings
+						);
+				})
+				.ColumnSortOrder(static_cast<int32>(ESubobjectColumnOrder::Owner)),
+			SHeaderRow::Column(OwnerOfSubobjectColumnId)
+				.DefaultLabel(LOCTEXT("Subobject.Owner", "Owner"))
+				.FillSized(Private::Shared::DefaultWidth)
+			);
+	}
+	
+	namespace Private::Property
+	{
+		static TArray<FGuid> GetPropertyOwners(
+			const FGlobalAuthorityCache& InAuthorityCache,
+			const ConcertClientSharedSlate::IReplicationStreamViewer& InViewer,
+			const FConcertPropertyChain& Property
+			)
+		{
+			TArray<FGuid> Result;
+			const TArray<FSoftObjectPath> SelectedObjects = InViewer.GetObjectsBeingPropertyEdited();
+			for (const FSoftObjectPath& SelectedObject : SelectedObjects)
+			{
+				const TOptional<FGuid> Owner = InAuthorityCache.GetClientWithAuthorityOverProperty(SelectedObject, Property);
+				if (Owner)
+				{
+					Result.AddUnique(*Owner);
+				}
+			}
+			return Result;
+		}
+	}
+
+	ConcertClientSharedSlate::ReplicationColumns::FReplicationPropertyColumn OwnerOfProperty(
+		const TSharedRef<IConcertClient>& InClient,
+		FGlobalAuthorityCache& InAuthorityCache,
+		const TAttribute<const ConcertClientSharedSlate::IReplicationStreamViewer*>& InViewer
+		)
+	{
+		using FColumnType = ConcertClientSharedSlate::ReplicationColumns::FReplicationPropertyColumn;
+		return FColumnType(
+			FColumnType::FArguments()
+				.GenerateWidgetColumn_Lambda([InClient, &InAuthorityCache, InViewer](const FColumnType::FBuildArgs& InArgs)
+				{
+					return SNew(SOwnerClientList, InClient, InAuthorityCache)
+						.GetClientList_Lambda([&InAuthorityCache, InViewer, Property = InArgs.RowData.GetProperty()](const FGlobalAuthorityCache&)
+						{
+							return Private::Property::GetPropertyOwners(InAuthorityCache, *InViewer.Get(), Property);
+						})
+						.HighlightText_Lambda([HighlightText = InArgs.HighlightText](){ return *HighlightText.Get(); });
+				})
+				.PopulateSearchItems_Lambda([InClient, &InAuthorityCache, InViewer](const ConcertClientSharedSlate::FReplicatedPropertyData& InArgs, TArray<FString>& InOutSearchStrings)
+				{
+					Private::Shared::PopulateSearchTerms(
+						InClient,
+						Private::Property::GetPropertyOwners(InAuthorityCache, *InViewer.Get(), InArgs.GetProperty()),
+						InOutSearchStrings
+						);
+				})
+				.ColumnSortOrder(static_cast<int32>(EPropertyColumnOrder::Owner)),
+			SHeaderRow::Column(OwnerOfPropertyColumnId)
+				.DefaultLabel(LOCTEXT("Property.Owner", "Owner"))
+				.FillSized(Private::Shared::DefaultWidth)
+			);
+	}
+}
+
+#undef LOCTEXT_NAMESPACE
\ No newline at end of file
diff --git a/Engine/Plugins/Developer/Concert/ConcertApp/MultiUserClient/Source/MultiUserClient/Private/Widgets/ActiveSession/Replication/Client/Single/Columns/SingleClientColumns_ToggleAuthority.cpp b/Engine/Plugins/Developer/Concert/ConcertApp/MultiUserClient/Source/MultiUserClient/Private/Widgets/ActiveSession/Replication/Client/Single/Columns/SingleClientColumns_ToggleAuthority.cpp
new file mode 100644
--- /dev/null
+++ b/Engine/Plugins/Developer/Concert/ConcertApp/MultiUserClient/Source/MultiUserClient/Private/Widgets/ActiveSession/Replication/Client/Single/Columns/SingleClientColumns_ToggleAuthority.cpp
@@ -0,0 +1,218 @@
+// Copyright Epic Games, Inc. All Rights Reserved.
+
+#include "SingleClientColumns.h"
+
+#include "IConcertClient.h"
+#include "SOwnerClientList.h"
+#include "Replication/Authority/AuthorityChangeTracker.h"
+#include "Replication/Authority/EAuthorityMutability.h"
+#include "Replication/Authority/IClientAuthoritySynchronizer.h"
+#include "Replication/Editor/Model/IReplicationStreamModel.h"
+#include "Replication/Editor/View/IReplicationStreamViewer.h"
+#include "Replication/Editor/View/ReplicationColumnsUtils.h"
+#include "Replication/Submission/ISubmissionWorkflow.h"
+#include "Replication/Util/GlobalAuthorityCache.h"
+#include "Widgets/ClientName/SClientName.h"
+
+#define LOCTEXT_NAMESPACE "SingleClientColumns.ToggleAuthority"
+
+namespace UE::MultiUserClient::SingleClientColumns
+{
+	const FName ToggleTopLevelAuthorityColumnId = TEXT("ToggleTopLevelAuthorityColumn");
+	const FName ToggleSubobjectAuthorityColumnId = TEXT("ToggleSubobjectAuthorityColumn");
+
+	namespace Private::ToggleTopLevelAuthority
+	{
+		static ECheckBoxState GetCheckboxState(
+			const ConcertClientSharedSlate::FReplicatedObjectData& ObjectData,
+			ConcertClientSharedSlate::IReplicationStreamModel* ClientStreamModel,
+			FAuthorityChangeTracker* ChangeTracker
+			)
+		{
+			const FSoftObjectPath& ObjectPath = ObjectData.GetObjectPath();
+
+			// Init the checkbox state to the top level object, if the top level object has properties associated with it...
+			const bool bCanSetTopLevelState = ChangeTracker->CanSetAuthorityFor(ObjectPath)
+				|| ChangeTracker->GetChangeAuthorityMutability(ObjectPath) == EAuthorityMutability::NotSupported;
+			TOptional<ECheckBoxState> CheckBoxState = bCanSetTopLevelState
+				? ChangeTracker->GetAuthorityStateAfterApplied(ObjectPath) ? ECheckBoxState::Checked : ECheckBoxState::Unchecked
+				: TOptional<ECheckBoxState>{};
+
+			// ... and then determine whether all the other subobjects share that state (Check / Unchecked) or not (Undetermined)
+			ClientStreamModel->ForEachSubobject(ObjectPath, [&CheckBoxState, &ChangeTracker](const FSoftObjectPath& Child)
+			{
+				if (!ChangeTracker->CanSetAuthorityFor(Child)
+					&& ChangeTracker->GetChangeAuthorityMutability(Child) != EAuthorityMutability::NotSupported)
+				{
+					return EBreakBehavior::Continue;
+				}
+
+				// Init checkbox state
+				const bool bCurrentAuthorityState = ChangeTracker->GetAuthorityStateAfterApplied(Child);
+				const bool bIsFirstObject = !CheckBoxState.IsSet(); 
+				if (bIsFirstObject)
+				{
+					CheckBoxState = bCurrentAuthorityState ? ECheckBoxState::Checked : ECheckBoxState::Unchecked;
+					return EBreakBehavior::Continue;
+				}
+
+				// Return the checkbox state that all subobjects share, or Undetermined if mixed
+				const bool bMustBeUndetermined = (*CheckBoxState == ECheckBoxState::Unchecked && bCurrentAuthorityState)
+					|| (*CheckBoxState == ECheckBoxState::Checked && !bCurrentAuthorityState);
+				if (bMustBeUndetermined)
+				{
+					CheckBoxState = ECheckBoxState::Undetermined;
+					return EBreakBehavior::Break;
+				}
+				return EBreakBehavior::Continue;
+			});
+						
+			return CheckBoxState
+				? *CheckBoxState
+				// If this case happens, there are no objects to check. In that case, it visually "looks better" to be unchecked
+				: ECheckBoxState::Unchecked;
+		}
+
+		static void OnCheckboxStateChanged(
+			bool bIsChecked,
+			const ConcertClientSharedSlate::FReplicatedObjectData& ObjectData,
+			ConcertClientSharedSlate::IReplicationStreamModel* ClientStreamModel,
+			FAuthorityChangeTracker* ChangeTracker
+			)
+		{
+			const FSoftObjectPath& ObjectPath = ObjectData.GetObjectPath();
+			TArray<FSoftObjectPath> Paths { ObjectPath };
+			ClientStreamModel->ForEachSubobject(ObjectPath, [&Paths](const FSoftObjectPath& Child)
+			{
+				Paths.Add(Child);
+				return EBreakBehavior::Continue;
+			});
+			
+			ChangeTracker->SetAuthorityIfAllowed(Paths, bIsChecked);
+		}
+
+		static bool IsEnabled(
+			const ConcertClientSharedSlate::FReplicatedObjectData& ObjectData,
+			ConcertClientSharedSlate::IReplicationStreamModel* ClientStreamModel,
+			FAuthorityChangeTracker* ChangeTracker
+			)
+		{
+			const FSoftObjectPath& TopLevelObjectPath = ObjectData.GetObjectPath();
+			const bool bCanChangeTopLevel = ChangeTracker->CanSetAuthorityFor(TopLevelObjectPath);
+			if (bCanChangeTopLevel)
+			{
+				return true;
+			}
+						
+			bool bCanChangeAnySubobject = false;
+			ClientStreamModel->ForEachSubobject(TopLevelObjectPath, [&ChangeTracker, &bCanChangeAnySubobject](const FSoftObjectPath& Child)
+			{
+				bCanChangeAnySubobject |= ChangeTracker->CanSetAuthorityFor(Child);
+				return bCanChangeAnySubobject ? EBreakBehavior::Break : EBreakBehavior::Continue;
+			});
+			return bCanChangeAnySubobject;
+		}
+	}
+	
+	ConcertClientSharedSlate::ReplicationColumns::FReplicationTopLevelObjectColumn ToggleTopLevelAuthority(
+		ConcertClientSharedSlate::IReplicationStreamModel& ClientStreamModel,
+		FAuthorityChangeTracker& ChangeTracker,
+		ISubmissionWorkflow& SubmissionWorkflow
+		)
+	{
+		using namespace ConcertClientSharedSlate;
+		using namespace Private::ToggleTopLevelAuthority;
+		using FTopLevelColumnDelegates = TReplicationColumnDelegates<FReplicatedObjectData>;
+		
+		return MakeCheckboxColumn<FReplicatedObjectData>(
+			ToggleTopLevelAuthorityColumnId,
+				FTopLevelColumnDelegates(
+					FTopLevelColumnDelegates::FGetColumnCheckboxState::CreateStatic(&GetCheckboxState, &ClientStreamModel, &ChangeTracker),
+					FTopLevelColumnDelegates::FOnColumnCheckboxChanged::CreateStatic(&OnCheckboxStateChanged, &ClientStreamModel, &ChangeTracker),
+					FTopLevelColumnDelegates::FGetToolTipText::CreateLambda([&ClientStreamModel, &ChangeTracker, &SubmissionWorkflow](const FReplicatedObjectData& ObjectData)
+					{
+						
+						if (SubmissionWorkflow.GetUploadability() == EChangeUploadability::NotImplemented)
+						{
+							const ECheckBoxState CheckBoxState = GetCheckboxState(ObjectData, &ClientStreamModel, &ChangeTracker);
+							const FText CheckBoxText = [CheckBoxState]()
+							{
+								switch (CheckBoxState)
+								{
+								case ECheckBoxState::Unchecked: return LOCTEXT("TopLevel.ChangeAuthority.ToolTip.NotSupported.Unchecked", "None of the subobjects are being replicated.");
+								case ECheckBoxState::Checked: return LOCTEXT("TopLevel.ChangeAuthority.ToolTip.NotSupported.Checked", "All of the subobjects are being replicated.");
+								case ECheckBoxState::Undetermined: return LOCTEXT("TopLevel.ChangeAuthority.ToolTip.NotSupported.Undetermined", "Some of the subobjects are being replicated.");
+								default: checkNoEntry(); return FText::GetEmpty();
+								}
+							}();
+							
+							return FText::Format(
+								LOCTEXT("TopLevel.ChangeAuthority.ToolTip.NotSupportedFmt", "{0}\nEditing remote clients is not implemented. You can only edit the local client."),
+								CheckBoxText
+								);
+						}
+						
+						return IsEnabled(ObjectData, &ClientStreamModel, &ChangeTracker)
+							? LOCTEXT("TopLevel.ChangeAuthority.Allowed", "Toggles whether this object and its subobjects will replicate the assigned properties.")
+							: LOCTEXT("TopLevel.ChangeAuthority.Disallowed", "Toggles whether this object and its subobjects will replicate the assigned properties.\nDisabled because none of the subobjects can be replicated with their current configuration.");
+					}),
+					FTopLevelColumnDelegates::FIsEnabled::CreateStatic(&IsEnabled, &ClientStreamModel, &ChangeTracker)
+				),
+			FText::GetEmpty(),
+			static_cast<int32>(ETopLevelObjectColumnOrder::ToggleAuthority)
+		);
+	}
+
+	ConcertClientSharedSlate::ReplicationColumns::FReplicationSubobjectObjectColumn ToggleSubobjectAuthority(
+		FAuthorityChangeTracker& ChangeTracker,
+		ISubmissionWorkflow& SubmissionWorkflow
+		)
+	{
+		using namespace ConcertClientSharedSlate;
+		using FSubobjectColumnDelegates = TReplicationColumnDelegates<FReplicatedObjectData>;
+		return MakeCheckboxColumn<FReplicatedObjectData>(
+			ToggleTopLevelAuthorityColumnId,
+				FSubobjectColumnDelegates(
+					FSubobjectColumnDelegates::FGetColumnCheckboxState::CreateLambda(
+					[&ChangeTracker](const FReplicatedObjectData& ObjectData)
+					{
+						const bool bHasAuthority = ChangeTracker.GetAuthorityStateAfterApplied(ObjectData.GetObjectPath());
+						return bHasAuthority ? ECheckBoxState::Checked : ECheckBoxState::Unchecked;
+					}),
+					FSubobjectColumnDelegates::FOnColumnCheckboxChanged::CreateLambda(
+					[&ChangeTracker](bool bIsChecked, const FReplicatedObjectData& ObjectData)
+					{
+						ChangeTracker.SetAuthorityIfAllowed({ ObjectData.GetObjectPath() }, bIsChecked);
+					}),
+					FSubobjectColumnDelegates::FGetToolTipText::CreateLambda([&ChangeTracker, &SubmissionWorkflow](const FReplicatedObjectData& ObjectData)
+					{
+						const bool bHasAuthority = ChangeTracker.GetAuthorityStateAfterApplied(ObjectData.GetObjectPath());
+						const FText NotSupportedText = FText::Format(
+							LOCTEXT("Subobject.ChangeAuthority.ToolTip.NotSupportedFmt", "{0}\nEditing remote clients is not implemented. You can only edit the local client."),
+							bHasAuthority ? LOCTEXT("Subobject.ChangeAuthority.Replicating.True", "This object is being replicated.") : LOCTEXT("Subobject.ChangeAuthority.Replicating.False", "This object is not currently being replicated.")
+							);
+						if (SubmissionWorkflow.GetUploadability() == EChangeUploadability::NotImplemented)
+						{
+							return NotSupportedText;
+						}
+						
+						switch (ChangeTracker.GetChangeAuthorityMutability(ObjectData.GetObjectPath()))
+						{
+						case EAuthorityMutability::Allowed: return LOCTEXT("Subobject.ChangeAuthority.ToolTip.Allowed", "Toggles whether this object should be replicated.");
+						case EAuthorityMutability::NoProperties: return LOCTEXT("Subobject.ChangeAuthority.ToolTip.NoProperties", "Toggles whether this object should be replicated.\nAssign properties to this object first.\n");
+						case EAuthorityMutability::NotSupported: return NotSupportedText;
+						default: checkNoEntry(); return FText::GetEmpty();
+						}
+					}),
+					FSubobjectColumnDelegates::FIsEnabled::CreateLambda([&ChangeTracker](const FReplicatedObjectData& ObjectData)
+					{
+						return ChangeTracker.CanSetAuthorityFor(ObjectData.GetObjectPath());
+					})
+				),
+			FText::GetEmpty(),
+			static_cast<int32>(ESubobjectColumnOrder::ToggleAuthority)
+		);
+	}
+}
+
+#undef LOCTEXT_NAMESPACE
\ No newline at end of file
diff --git a/Engine/Plugins/Developer/Concert/ConcertApp/MultiUserClient/Source/MultiUserClient/Private/Widgets/ActiveSession/Replication/Client/Single/SReplicationClientView.cpp b/Engine/Plugins/Developer/Concert/ConcertApp/MultiUserClient/Source/MultiUserClient/Private/Widgets/ActiveSession/Replication/Client/Single/SReplicationClientView.cpp
--- a/Engine/Plugins/Developer/Concert/ConcertApp/MultiUserClient/Source/MultiUserClient/Private/Widgets/ActiveSession/Replication/Client/Single/SReplicationClientView.cpp
+++ b/Engine/Plugins/Developer/Concert/ConcertApp/MultiUserClient/Source/MultiUserClient/Private/Widgets/ActiveSession/Replication/Client/Single/SReplicationClientView.cpp
@@ -11,7 +11,7 @@
 #include "Replication/ReplicationWidgetFactories.h"
 #include "Replication/Client/ReplicationClientManager.h"
 #include "Replication/Submission/ISubmissionWorkflow.h"
-#include "Widgets/ActiveSession/Replication/Client/Columns/SingleClientColumns.h"
+#include "Widgets/ActiveSession/Replication/Client/Single/Columns/SingleClientColumns.h"
 #include "Widgets/ActiveSession/Replication/Client/SClientToolbar.h"
 
 #include "HAL/IConsoleManager.h"
