Add missing inline specififer to util function in header file

diff --git a/Engine/Plugins/Developer/Concert/ConcertApp/MultiUserClient/Source/MultiUserClient/Private/Replication/Submission/Data/EChangeUploadability.h b/Engine/Plugins/Developer/Concert/ConcertApp/MultiUserClient/Source/MultiUserClient/Private/Replication/Submission/Data/EChangeUploadability.h
--- a/Engine/Plugins/Developer/Concert/ConcertApp/MultiUserClient/Source/MultiUserClient/Private/Replication/Submission/Data/EChangeUploadability.h
+++ b/Engine/Plugins/Developer/Concert/ConcertApp/MultiUserClient/Source/MultiUserClient/Private/Replication/Submission/Data/EChangeUploadability.h
@@ -19,7 +19,7 @@ namespace UE::MultiUserClient
 		NotImplemented
 	};
 
-	bool CanEverSubmit(EChangeUploadability Uploadability)
+	inline bool CanEverSubmit(EChangeUploadability Uploadability)
 	{
 		return Uploadability != EChangeUploadability::NotImplemented;
 	}
