Add replication toggle for all clients view

diff --git a/Engine/Plugins/Developer/Concert/ConcertApp/MultiUserClient/Source/MultiUserClient/Private/Replication/Util/GlobalAuthorityCache.cpp b/Engine/Plugins/Developer/Concert/ConcertApp/MultiUserClient/Source/MultiUserClient/Private/Replication/Util/GlobalAuthorityCache.cpp
--- a/Engine/Plugins/Developer/Concert/ConcertApp/MultiUserClient/Source/MultiUserClient/Private/Replication/Util/GlobalAuthorityCache.cpp
+++ b/Engine/Plugins/Developer/Concert/ConcertApp/MultiUserClient/Source/MultiUserClient/Private/Replication/Util/GlobalAuthorityCache.cpp
@@ -20,6 +20,23 @@ namespace UE::MultiUserClient
 		ClientManager.OnPreRemoteClientRemoved().AddRaw(this, &FGlobalAuthorityCache::OnPreRemoteClientRemoved);
 	}
 
+	void FGlobalAuthorityCache::ForEachClientWithObjectInStream(const FSoftObjectPath& Object, TFunctionRef<EBreakBehavior(const FGuid& ClientId)> Callback) const
+	{
+		const TSet<FGuid>* Clients = RegisteredObjectsToClients.Find(Object);
+		if (!Clients)
+		{
+			return;
+		}
+
+		for (const FGuid& ClientId : *Clients)
+		{
+			if (Callback(ClientId) == EBreakBehavior::Break)
+			{
+				break;
+			}
+		}
+	}
+
 	void FGlobalAuthorityCache::ForEachClientWithAuthorityOverObject(const FSoftObjectPath& Object, TFunctionRef<EBreakBehavior(const FGuid& ClientId)> Callback) const
 	{
 		const TSet<FGuid>* Clients = OwnedObjectsToClients.Find(Object);
@@ -156,6 +173,7 @@ namespace UE::MultiUserClient
 		for (const TPair<FSoftObjectPath, FReplicatedObjectInfo>& StreamContents : ClientObjectMap.ReplicatedObjects)
 		{
 			const FSoftObjectPath& Object = StreamContents.Key;
+			RegisteredObjectsToClients.FindOrAdd(Object).Add(ClientId);
 			if (ClientAuthority.HasAuthorityOver(Object))
 			{
 				OwnedObjectsToClients.FindOrAdd(Object).Add(ClientId);
@@ -173,6 +191,15 @@ namespace UE::MultiUserClient
 				It.RemoveCurrent();
 			}
 		}
+
+		for (auto It = RegisteredObjectsToClients.CreateIterator(); It; ++It)
+		{
+			It->Value.Remove(ClientId);
+			if (It->Value.IsEmpty())
+			{
+				It.RemoveCurrent();
+			}
+		}
 	}
 	
 	void FGlobalAuthorityCache::ForEachStream(const FGuid& ClientEndpointId, TFunctionRef<EBreakBehavior(const FGuid& StreamId, const FObjectReplicationMap& ReplicationMap)> Callback) const
diff --git a/Engine/Plugins/Developer/Concert/ConcertApp/MultiUserClient/Source/MultiUserClient/Private/Replication/Util/GlobalAuthorityCache.h b/Engine/Plugins/Developer/Concert/ConcertApp/MultiUserClient/Source/MultiUserClient/Private/Replication/Util/GlobalAuthorityCache.h
--- a/Engine/Plugins/Developer/Concert/ConcertApp/MultiUserClient/Source/MultiUserClient/Private/Replication/Util/GlobalAuthorityCache.h
+++ b/Engine/Plugins/Developer/Concert/ConcertApp/MultiUserClient/Source/MultiUserClient/Private/Replication/Util/GlobalAuthorityCache.h
@@ -35,6 +35,9 @@ namespace UE::MultiUserClient
 		FGlobalAuthorityCache(FReplicationClientManager& InClientManager);
 		/** Called when the local client has been created and it is safe to register client events with FReplicationClientManager. */
 		void RegisterEvents();
+		
+		/** Iterates every client that Object in its stream. */
+		void ForEachClientWithObjectInStream(const FSoftObjectPath& Object, TFunctionRef<EBreakBehavior(const FGuid& ClientId)> Callback) const;
 
 		/** Iterates every client that has authority over Object. */
 		void ForEachClientWithAuthorityOverObject(const FSoftObjectPath& Object, TFunctionRef<EBreakBehavior(const FGuid& ClientId)> Callback) const;
@@ -83,6 +86,8 @@ namespace UE::MultiUserClient
 		
 		/** Maps objects that are owned to the clients that own them */
 		TMap<FSoftObjectPath, TSet<FGuid>> OwnedObjectsToClients;
+		/** Maps objects that are owned to the clients that own them */
+		TMap<FSoftObjectPath, TSet<FGuid>> RegisteredObjectsToClients;
 		
 		/** Called when the cache changes for a specific client. */
 		FOnCacheChanged OnCacheChangedDelegate;
diff --git a/Engine/Plugins/Developer/Concert/ConcertApp/MultiUserClient/Source/MultiUserClient/Private/Widgets/ActiveSession/Replication/Client/Multi/Columns/MultiStreamColumns.h b/Engine/Plugins/Developer/Concert/ConcertApp/MultiUserClient/Source/MultiUserClient/Private/Widgets/ActiveSession/Replication/Client/Multi/Columns/MultiStreamColumns.h
--- a/Engine/Plugins/Developer/Concert/ConcertApp/MultiUserClient/Source/MultiUserClient/Private/Widgets/ActiveSession/Replication/Client/Multi/Columns/MultiStreamColumns.h
+++ b/Engine/Plugins/Developer/Concert/ConcertApp/MultiUserClient/Source/MultiUserClient/Private/Widgets/ActiveSession/Replication/Client/Multi/Columns/MultiStreamColumns.h
@@ -23,8 +23,33 @@ namespace UE::ConcertClientSharedSlate
 
 namespace UE::MultiUserClient::MultiStreamColumns
 {
+	const extern FName ReplicationToggleColumnId;
 	const extern FName AssignPropertyColumnId;
+	
+	/* @see EReplicationPropertyColumnOrder */
+	enum class EColumnSortOrder
+	{
+		ReplicationToggle = 0,
+		AssignPropertyColumn = 30
+	};
 
+	/**
+	 * Toggles replication for all clients assigned to the object (and optionally all children).
+	 * 
+	 * @param ConcertClient Used to look up client names
+	 * @param ClientManager Used to access all clients for toggling authority
+	 * @param ConsolidatedStreamModelAttribute Used to get child objects
+	 * @param ColumnsSortPriority The order relative to the other columns
+	 * 
+	 * @return A checkbox for controlling the authority of the object in the row
+	 */
+	ConcertClientSharedSlate::ReplicationColumns::FReplicationTopLevelObjectColumn ReplicationToggle(
+		TSharedRef<IConcertClient> ConcertClient,
+		TAttribute<ConcertClientSharedSlate::IReplicationStreamModel*> ConsolidatedStreamModelAttribute,
+		FReplicationClientManager& ClientManager,
+		const int32 ColumnsSortPriority = static_cast<int32>(EColumnSortOrder::ReplicationToggle)
+		);
+	
 	/**
 	 * Creates a property column which assigns the property to the stream selected in the combo box.
 	 * 
@@ -39,6 +64,6 @@ namespace UE::MultiUserClient::MultiStreamColumns
 		TAttribute<TSharedPtr<ConcertClientSharedSlate::IMultiReplicationStreamEditor>> MultiStreamEditor,
 		TSharedRef<IConcertClient> ConcertClient,
 		FReplicationClientManager& ClientManager,
-		const int32 ColumnsSortPriority = 30 /* @see EReplicationPropertyColumnOrder */
+		const int32 ColumnsSortPriority = static_cast<int32>(EColumnSortOrder::AssignPropertyColumn)
 		);
 }
diff --git a/Engine/Plugins/Developer/Concert/ConcertApp/MultiUserClient/Source/MultiUserClient/Private/Widgets/ActiveSession/Replication/Client/Multi/Columns/MultiStreamColumns_AssignProperties.cpp b/Engine/Plugins/Developer/Concert/ConcertApp/MultiUserClient/Source/MultiUserClient/Private/Widgets/ActiveSession/Replication/Client/Multi/Columns/MultiStreamColumns_AssignProperties.cpp
--- a/Engine/Plugins/Developer/Concert/ConcertApp/MultiUserClient/Source/MultiUserClient/Private/Widgets/ActiveSession/Replication/Client/Multi/Columns/MultiStreamColumns_AssignProperties.cpp
+++ b/Engine/Plugins/Developer/Concert/ConcertApp/MultiUserClient/Source/MultiUserClient/Private/Widgets/ActiveSession/Replication/Client/Multi/Columns/MultiStreamColumns_AssignProperties.cpp
@@ -26,7 +26,7 @@ namespace UE::MultiUserClient::MultiStreamColumns
 			TFunctionRef<void(const TSharedRef<ConcertClientSharedSlate::IEditableReplicationStreamModel>& Stream)> Consume
 		)
 		{
-			for (const TSharedRef<ConcertClientSharedSlate::IEditableReplicationStreamModel>& Stream : MultiEditor.GetModel().GetEditableStreams())
+			for (const TSharedRef<ConcertClientSharedSlate::IEditableReplicationStreamModel>& Stream : MultiEditor.GetMultiStreamModel().GetEditableStreams())
 			{
 				for (const FSoftObjectPath& SelectedObject : MultiEditor.GetEditorBase().GetObjectsBeingPropertyEdited())
 				{
diff --git a/Engine/Plugins/Developer/Concert/ConcertApp/MultiUserClient/Source/MultiUserClient/Private/Widgets/ActiveSession/Replication/Client/Multi/Columns/MultiStreamColumns_ReplicationToggle.cpp b/Engine/Plugins/Developer/Concert/ConcertApp/MultiUserClient/Source/MultiUserClient/Private/Widgets/ActiveSession/Replication/Client/Multi/Columns/MultiStreamColumns_ReplicationToggle.cpp
new file mode 100644
--- /dev/null
+++ b/Engine/Plugins/Developer/Concert/ConcertApp/MultiUserClient/Source/MultiUserClient/Private/Widgets/ActiveSession/Replication/Client/Multi/Columns/MultiStreamColumns_ReplicationToggle.cpp
@@ -0,0 +1,44 @@
+// Copyright Epic Games, Inc. All Rights Reserved.
+
+#include "MultiStreamColumns.h"
+
+#include "SReplicationMultiToggleCheckbox.h"
+
+#include "Widgets/Layout/SBox.h"
+
+#define LOCTEXT_NAMESPACE "ReplicationToggle"
+
+namespace UE::MultiUserClient::MultiStreamColumns
+{
+	const FName ReplicationToggleColumnId(TEXT("ReplicationToggleColumn"));
+	
+	ConcertClientSharedSlate::ReplicationColumns::FReplicationTopLevelObjectColumn ReplicationToggle(
+		TSharedRef<IConcertClient> ConcertClient,
+		TAttribute<ConcertClientSharedSlate::IReplicationStreamModel*> ConsolidatedStreamModelAttribute,
+		FReplicationClientManager& ClientManager,
+		const int32 ColumnsSortPriority
+		)
+	{
+		using namespace ConcertClientSharedSlate::ReplicationColumns;
+		
+		return FReplicationTopLevelObjectColumn(
+			FReplicationTopLevelObjectColumn::FArguments()
+				.GenerateWidgetColumn_Lambda([ConcertClient = MoveTemp(ConcertClient), ConsolidatedStreamModelAttribute = MoveTemp(ConsolidatedStreamModelAttribute), &ClientManager](const FReplicationTopLevelObjectColumn::FBuildArgs& InArgs)
+				{
+					return SNew(SBox)
+					.HAlign(HAlign_Left) // Warning icon is sometimes collapsed - we don't want the widget to be centered
+					[
+						SNew(SReplicationMultiToggleCheckbox, ClientManager, ConcertClient)
+						.Object(InArgs.RowData.GetObjectPath())
+						.ConsolidatedStreamModelAttribute(ConsolidatedStreamModelAttribute)
+					];
+				})
+				.ColumnSortOrder(ColumnsSortPriority),
+			SHeaderRow::Column(ReplicationToggleColumnId)
+				.DefaultLabel(FText::GetEmpty())
+				.FixedWidth(70.f)
+			);
+	}
+}
+
+#undef LOCTEXT_NAMESPACE
\ No newline at end of file
diff --git a/Engine/Plugins/Developer/Concert/ConcertApp/MultiUserClient/Source/MultiUserClient/Private/Widgets/ActiveSession/Replication/Client/Multi/Columns/SReplicationMultiToggleCheckbox.cpp b/Engine/Plugins/Developer/Concert/ConcertApp/MultiUserClient/Source/MultiUserClient/Private/Widgets/ActiveSession/Replication/Client/Multi/Columns/SReplicationMultiToggleCheckbox.cpp
new file mode 100644
--- /dev/null
+++ b/Engine/Plugins/Developer/Concert/ConcertApp/MultiUserClient/Source/MultiUserClient/Private/Widgets/ActiveSession/Replication/Client/Multi/Columns/SReplicationMultiToggleCheckbox.cpp
@@ -0,0 +1,272 @@
+// Copyright Epic Games, Inc. All Rights Reserved.
+
+#include "SReplicationMultiToggleCheckbox.h"
+
+#include "Replication/Client/ReplicationClientManager.h"
+#include "Replication/Util/GlobalAuthorityCache.h"
+#include "Widgets/ActiveSession/Replication/Client/ClientUtils.h"
+
+#include "Algo/AnyOf.h"
+#include "Framework/Commands/UIAction.h"
+#include "Framework/MultiBox/MultiBoxBuilder.h"
+#include "Styling/AppStyle.h"
+#include "Styling/SlateTypes.h"
+#include "Textures/SlateIcon.h"
+#include "Widgets/Images/SImage.h"
+#include "Widgets/Input/SCheckBox.h"
+#include "Widgets/Input/SComboBox.h"
+#include "Widgets/SBoxPanel.h"
+#include "Widgets/Layout/SScaleBox.h"
+
+#define LOCTEXT_NAMESPACE "SReplicationMultiToggleCheckbox"
+
+namespace UE::MultiUserClient
+{
+	void SReplicationMultiToggleCheckbox::Construct(
+		const FArguments& InArgs,
+		FReplicationClientManager& InClientManager,
+		TSharedRef<IConcertClient> InConcertClient
+		)
+	{
+		Object = InArgs._Object;
+		ConsolidatedStreamModelAttribute = InArgs._ConsolidatedStreamModelAttribute;
+		check(ConsolidatedStreamModelAttribute.IsSet() || ConsolidatedStreamModelAttribute.IsBound());
+		
+		ClientManager = &InClientManager;
+		ConcertClient = MoveTemp(InConcertClient);
+		
+		ChildSlot
+		[
+			SNew(SComboButton)
+			.HasDownArrow(true)
+			.ComboButtonStyle(FAppStyle::Get(), "SimpleComboButton")
+			.ForegroundColor(FSlateColor::UseStyle())
+			.ContentPadding(FMargin(2.0f, 2.0f))
+			.ButtonContent()
+			[
+				SNew(SHorizontalBox)
+				.ToolTipText(this, &SReplicationMultiToggleCheckbox::GetRootToolTipText)
+
+				+SHorizontalBox::Slot()
+				.AutoWidth()
+				[
+					SNew(SCheckBox)
+					.IsChecked(this, &SReplicationMultiToggleCheckbox::GetCheckboxStateForObject, Object)
+					.IsEnabled(this, &SReplicationMultiToggleCheckbox::IsCheckboxEnabledForObject, Object)
+					.OnCheckStateChanged(this, &SReplicationMultiToggleCheckbox::OnCheckboxStateChangedForObject, Object)
+				]
+			
+				+SHorizontalBox::Slot()
+				.AutoWidth()
+				[
+					SNew(SScaleBox)
+					.Stretch(EStretch::ScaleToFit)
+					.Visibility(this, &SReplicationMultiToggleCheckbox::GetWarningVisibility)
+					.ToolTipText(this, &SReplicationMultiToggleCheckbox::GetWarningToolTipText)
+					[
+						SNew(SImage)
+						.Image(FAppStyle::GetBrush("Icons.WarningWithColor"))
+					]
+				]
+			]
+			.OnGetMenuContent(this, &SReplicationMultiToggleCheckbox::GetDropDownMenuContent)
+		];
+	}
+
+	FText SReplicationMultiToggleCheckbox::GetRootToolTipText() const
+	{
+		return LOCTEXT("RootTooltip", "Controls whether the assigned clients should replicate this object.");
+	}
+
+	ECheckBoxState SReplicationMultiToggleCheckbox::GetCheckboxStateForObject(FSoftObjectPath InObject) const
+	{
+		// The checkbox shows
+		// - checked if all editable clients with the object registered have authority
+		// - unchecked if all editable clients with the object registered do not have authority
+		// - undetermined otherwise
+		const TArray<FGuid> ClientsWithAuthority = ClientManager->GetAuthorityCache().GetClientsWithAuthorityOverObject(InObject);
+
+		// TODO UE-200496 Predict conflicts
+		
+		bool bHadAnyEditableClient = false;
+		ECheckBoxState ConsolidatedState = ECheckBoxState::Undetermined;
+		ClientManager->GetAuthorityCache().ForEachClientWithObjectInStream(InObject, [this, &ClientsWithAuthority, &bHadAnyEditableClient, &ConsolidatedState](const FGuid& ClientId)
+		{
+			// The state of the checkbox always skips non-editable clients
+			const FReplicationClient* Client = ClientManager->FindClient(ClientId);
+			if (!ensure(Client) || !Client->AllowsEditing())
+			{
+				return EBreakBehavior::Continue;
+			}
+			
+			bHadAnyEditableClient = true;
+			const bool bHasAuthority = ClientsWithAuthority.Contains(Client->GetEndpointId());
+			const ECheckBoxState ClientState = bHasAuthority ? ECheckBoxState::Checked : ECheckBoxState::Unchecked;
+			
+			if (ConsolidatedState == ECheckBoxState::Undetermined)
+			{
+				ConsolidatedState = ClientState;
+			}
+			else if (ConsolidatedState != ClientState)
+			{
+				ConsolidatedState = ECheckBoxState::Undetermined;
+				return EBreakBehavior::Break;
+			}
+			
+			return EBreakBehavior::Continue;
+		});
+		
+		return bHadAnyEditableClient ? ConsolidatedState : ECheckBoxState::Unchecked;
+	}
+
+	bool SReplicationMultiToggleCheckbox::IsCheckboxEnabledForObject(FSoftObjectPath InObject) const
+	{
+		// TODO UE-200496 Predict conflicts
+		
+		bool bHasAtLeastOneEditableClient = false;
+		ClientManager->GetAuthorityCache().ForEachClientWithObjectInStream(InObject, [this, &bHasAtLeastOneEditableClient](const FGuid& ClientId)
+		{
+			const FReplicationClient* Client = ClientManager->FindClient(ClientId);
+			bHasAtLeastOneEditableClient = ensure(Client) && Client->AllowsEditing();
+			return bHasAtLeastOneEditableClient ? EBreakBehavior::Break : EBreakBehavior::Continue;
+		});
+		return bHasAtLeastOneEditableClient;
+	}
+
+	void SReplicationMultiToggleCheckbox::OnCheckboxStateChangedForObject(ECheckBoxState NewState, FSoftObjectPath InObject) const
+	{
+		const bool bShouldHaveAuthority = NewState == ECheckBoxState::Checked;
+		ClientManager->GetAuthorityCache().ForEachClientWithObjectInStream(InObject, [this, InObject, bShouldHaveAuthority](const FGuid& ClientId)
+		{
+			FReplicationClient* Client = ClientManager->FindClient(ClientId);
+			if (ensure(Client) && Client->AllowsEditing())
+			{
+				Client->GetAuthorityDiffer().SetAuthorityIfAllowed({ InObject }, bShouldHaveAuthority);
+			}
+			
+			return EBreakBehavior::Continue;
+		});
+	}
+
+	TSharedRef<SWidget> SReplicationMultiToggleCheckbox::GetDropDownMenuContent()
+	{
+		FMenuBuilder MenuBuilder(true, nullptr);
+
+		MenuBuilder.AddMenuEntry(
+			LOCTEXT("OnlyThis", "Toggle object only"),
+			FText::GetEmpty(),
+			FSlateIcon(),
+			FUIAction(
+				FExecuteAction::CreateLambda([this]()
+				{
+					const ECheckBoxState OldCheckboxState = GetCheckboxStateForObject(Object);
+					const bool bGiveAuthority = OldCheckboxState == ECheckBoxState::Unchecked;
+					OnCheckboxStateChangedForObject(bGiveAuthority ? ECheckBoxState::Checked : ECheckBoxState::Unchecked, Object);
+				}),
+				FCanExecuteAction::CreateLambda([this](){ return IsCheckboxEnabledForObject(Object); })
+				)
+			);
+		MenuBuilder.AddMenuEntry(
+			LOCTEXT("ThisAndChildren", "Toggle children only"),
+			FText::GetEmpty(),
+			FSlateIcon(),
+			FUIAction(
+				FExecuteAction::CreateSP(this, &SReplicationMultiToggleCheckbox::ToggleChildren),
+				FCanExecuteAction::CreateSP(this, &SReplicationMultiToggleCheckbox::CanToggleChildren)
+				)
+			);
+		
+		return MenuBuilder.MakeWidget();
+	}
+
+	void SReplicationMultiToggleCheckbox::ToggleChildren() const
+	{
+		const TArray<FSoftObjectPath> Subobjects = ConsolidatedStreamModelAttribute.Get()->GetSubobjects(Object);
+		ToggleObjects(Subobjects);
+	}
+
+	bool SReplicationMultiToggleCheckbox::CanToggleChildren() const
+	{
+		const TArray<FSoftObjectPath> Subobjects = ConsolidatedStreamModelAttribute.Get()->GetSubobjects(Object);
+		return CanToggleObjects(Subobjects);
+	}
+
+	void SReplicationMultiToggleCheckbox::ToggleObjects(TConstArrayView<FSoftObjectPath> Objects) const
+	{
+		if (Objects.IsEmpty())
+		{
+			return;
+		}
+		
+		ECheckBoxState ConsolidatedCheckboxState = GetCheckboxStateForObject(Objects[0]);
+		for (int32 i = 1; i < Objects.Num(); ++i)
+		{
+			const ECheckBoxState ObjectState = GetCheckboxStateForObject(Objects[i]);
+			if (ConsolidatedCheckboxState != ObjectState)
+			{
+				// If the states are different, pretend they are all Checked so below we default to toggling off
+				ConsolidatedCheckboxState = ECheckBoxState::Checked;
+				break;
+			}
+		}
+
+		const ECheckBoxState StateToSet = ConsolidatedCheckboxState == ECheckBoxState::Unchecked ? ECheckBoxState::Checked : ECheckBoxState::Unchecked;
+		for (const FSoftObjectPath& ObjectPath : Objects)
+		{
+			OnCheckboxStateChangedForObject(StateToSet, ObjectPath);
+		}
+	}
+
+	bool SReplicationMultiToggleCheckbox::CanToggleObjects(TConstArrayView<FSoftObjectPath> Objects) const
+	{
+		return Algo::AnyOf(Objects, [this](const FSoftObjectPath& ObjectPath)
+		{
+			return IsCheckboxEnabledForObject(ObjectPath);
+		});
+	}
+
+	EVisibility SReplicationMultiToggleCheckbox::GetWarningVisibility() const
+	{
+		const bool bContainsNonEditableClients = !GetNonEditableClients().IsEmpty();
+		return IsCheckboxEnabledForObject(Object) && bContainsNonEditableClients
+			? EVisibility::Visible
+			: EVisibility::Collapsed;
+	}
+
+	FText SReplicationMultiToggleCheckbox::GetWarningToolTipText() const
+	{
+		const TArray<FGuid> NonEditableClients = GetNonEditableClients();
+		if (NonEditableClients.IsEmpty())
+		{
+			// Avoid text substitution below
+			return FText::GetEmpty();
+		}
+		
+		FString Clients = FString::JoinBy(NonEditableClients, TEXT(", "), [this](const FGuid& ClientEndpointId)
+		{
+			return ClientUtils::GetClientDisplayName(*ConcertClient, ClientEndpointId);
+		});
+		return FText::Format(
+			LOCTEXT("WarningIcon.ToolTip", "Has clients that cannot be edited remotely.\nClients: {0}"),
+			FText::FromString(MoveTemp(Clients))
+			);
+	}
+
+	TArray<FGuid> SReplicationMultiToggleCheckbox::GetNonEditableClients() const
+	{
+		TArray<FGuid> Result;
+		ClientManager->GetAuthorityCache().ForEachClientWithObjectInStream(Object, [this, &Result](const FGuid& ClientId)
+		{
+			// The warning icon cares only about editable clients
+			const FReplicationClient* Client = ClientManager->FindClient(ClientId);
+			if (ensure(Client) && !Client->AllowsEditing())
+			{
+				Result.Add(ClientId);
+			}
+			return EBreakBehavior::Continue;
+		});
+		return Result;
+	}
+}
+
+#undef LOCTEXT_NAMESPACE
\ No newline at end of file
diff --git a/Engine/Plugins/Developer/Concert/ConcertApp/MultiUserClient/Source/MultiUserClient/Private/Widgets/ActiveSession/Replication/Client/Multi/Columns/SReplicationMultiToggleCheckbox.h b/Engine/Plugins/Developer/Concert/ConcertApp/MultiUserClient/Source/MultiUserClient/Private/Widgets/ActiveSession/Replication/Client/Multi/Columns/SReplicationMultiToggleCheckbox.h
new file mode 100644
--- /dev/null
+++ b/Engine/Plugins/Developer/Concert/ConcertApp/MultiUserClient/Source/MultiUserClient/Private/Widgets/ActiveSession/Replication/Client/Multi/Columns/SReplicationMultiToggleCheckbox.h
@@ -0,0 +1,85 @@
+// Copyright Epic Games, Inc. All Rights Reserved.
+
+#pragma once
+
+#include "UObject/SoftObjectPath.h"
+#include "Widgets/DeclarativeSyntaxSupport.h"
+#include "Widgets/SCompoundWidget.h"
+
+class IConcertClient;
+enum class ECheckBoxState : uint8;
+
+namespace UE::ConcertClientSharedSlate
+{
+	class IReplicationStreamModel;
+}
+
+namespace UE::MultiUserClient
+{
+	class FReplicationClientManager;
+}
+
+namespace UE::MultiUserClient
+{
+	/**
+	 * Checkbox in a combo button.
+	 * 
+	 * The checkbox can either the object it was created for or all of its child objects.
+	 * If the checkbox in the combo button is pressed it defaults to toggling just the object.
+	 *
+	 * TODO UE-200487: Make sure this widget updates when a remote client changes their setting for being editable
+	 */
+	class SReplicationMultiToggleCheckbox : public SCompoundWidget
+	{
+	public:
+
+		SLATE_BEGIN_ARGS(SReplicationMultiToggleCheckbox)
+		{}
+		    SLATE_ARGUMENT(FSoftObjectPath, Object)
+			SLATE_ATTRIBUTE(ConcertClientSharedSlate::IReplicationStreamModel*, ConsolidatedStreamModelAttribute)
+		SLATE_END_ARGS()
+
+		void Construct(
+			const FArguments& InArgs,
+		    FReplicationClientManager& InClientManager,
+		    TSharedRef<IConcertClient> InConcertClient
+		);
+
+	private:
+
+		/** The object path this checkbox was created for */
+		FSoftObjectPath Object;
+
+		/** Used to access all clients for toggling authority. */
+		FReplicationClientManager* ClientManager = nullptr;
+		/** Used to get children of Object */
+		TAttribute<ConcertClientSharedSlate::IReplicationStreamModel*> ConsolidatedStreamModelAttribute;
+		
+		/** Used to look up client display names in case of conflicts. */
+		TSharedPtr<IConcertClient> ConcertClient;
+
+		/** @return Tooltip for the combo box. Informs user what the end result is if enabled or why it's disabled. */
+		FText GetRootToolTipText() const;
+		
+		/** @return Goes through all clients that can be edited and returns whether all of them have authority. */
+		ECheckBoxState GetCheckboxStateForObject(FSoftObjectPath InObject) const;
+		/** @return Whether there are any clients for which the authority can be changed. */
+		bool IsCheckboxEnabledForObject(FSoftObjectPath InObject) const;
+		/** Gives all editable clients which have something registered to Object authority or takes it away. */
+		void OnCheckboxStateChangedForObject(ECheckBoxState NewState, FSoftObjectPath InObject) const;
+
+		/** @return Menu widget that with the options to remove or give authority to this and other objects. */
+		TSharedRef<SWidget> GetDropDownMenuContent();
+		
+		void ToggleChildren() const;
+		bool CanToggleChildren() const;
+
+		void ToggleObjects(TConstArrayView<FSoftObjectPath> Objects) const;
+		bool CanToggleObjects(TConstArrayView<FSoftObjectPath> Objects) const;
+		
+		EVisibility GetWarningVisibility() const;
+		FText GetWarningToolTipText() const;
+		TArray<FGuid> GetNonEditableClients() const;
+	};
+}
+
diff --git a/Engine/Plugins/Developer/Concert/ConcertApp/MultiUserClient/Source/MultiUserClient/Private/Widgets/ActiveSession/Replication/Client/Multi/SMultiClientView.cpp b/Engine/Plugins/Developer/Concert/ConcertApp/MultiUserClient/Source/MultiUserClient/Private/Widgets/ActiveSession/Replication/Client/Multi/SMultiClientView.cpp
--- a/Engine/Plugins/Developer/Concert/ConcertApp/MultiUserClient/Source/MultiUserClient/Private/Widgets/ActiveSession/Replication/Client/Multi/SMultiClientView.cpp
+++ b/Engine/Plugins/Developer/Concert/ConcertApp/MultiUserClient/Source/MultiUserClient/Private/Widgets/ActiveSession/Replication/Client/Multi/SMultiClientView.cpp
@@ -61,6 +61,11 @@ namespace UE::MultiUserClient
 		   {
 			   return StreamEditor;
 		   });
+		TAttribute<IReplicationStreamModel*> ConsolidatedStreamModelAttribute =
+		   TAttribute<IReplicationStreamModel*>::CreateLambda([this]()
+		   {
+			   return &StreamEditor->GetConsolidatedModel();
+		   });
 		
 		FCreateMultiStreamEditorParams Params
 		{
@@ -71,7 +76,8 @@ namespace UE::MultiUserClient
 			.ViewerParams =
 			{
 				.SubobjectModel = CreateDefaultComponentHierarchySubobjectModel(), // This makes actors have children in the top view
-				.AdditionalPropertyColumns = { MultiStreamColumns::AssignPropertyColumn(MultiStreamEditorAttribute, InConcertClient, InClientManager) }
+				.AdditionalObjectColumns = { MultiStreamColumns::ReplicationToggle(InConcertClient, MoveTemp(ConsolidatedStreamModelAttribute), InClientManager) },
+				.AdditionalPropertyColumns = { MultiStreamColumns::AssignPropertyColumn(MoveTemp(MultiStreamEditorAttribute), InConcertClient, InClientManager) }
 			}
 		};
 		StreamEditor = CreateBaseMultiStreamEditor(MoveTemp(Params));
diff --git a/Engine/Plugins/Developer/Concert/ConcertUI/ConcertClientSharedSlate/Source/ConcertClientSharedSlate/Private/Replication/Editor/View/MultiEditor/SMultiReplicationStreamEditor.h b/Engine/Plugins/Developer/Concert/ConcertUI/ConcertClientSharedSlate/Source/ConcertClientSharedSlate/Private/Replication/Editor/View/MultiEditor/SMultiReplicationStreamEditor.h
--- a/Engine/Plugins/Developer/Concert/ConcertUI/ConcertClientSharedSlate/Source/ConcertClientSharedSlate/Private/Replication/Editor/View/MultiEditor/SMultiReplicationStreamEditor.h
+++ b/Engine/Plugins/Developer/Concert/ConcertUI/ConcertClientSharedSlate/Source/ConcertClientSharedSlate/Private/Replication/Editor/View/MultiEditor/SMultiReplicationStreamEditor.h
@@ -2,8 +2,10 @@
 
 #pragma once
 
+#include "ConsolidatedMultiStreamModel.h"
 #include "Replication/ReplicationWidgetFactories.h"
 #include "Replication/Editor/View/IMultiReplicationStreamEditor.h"
+#include "Replication/Editor/View/IReplicationStreamEditor.h"
 
 #include "Widgets/DeclarativeSyntaxSupport.h"
 #include "Widgets/SCompoundWidget.h"
@@ -17,7 +19,6 @@ namespace UE::ConcertClientSharedSlate
 
 namespace UE::ConcertClientSharedSlate
 {
-	class FAggregateReadOnlyStreamModel;
 	struct FCreateMultiStreamEditorParams;
 
 	/**
@@ -39,7 +40,8 @@ namespace UE::ConcertClientSharedSlate
 		void Construct(const FArguments& InArgs, FCreateMultiStreamEditorParams Params);
 
 		virtual IReplicationStreamEditor& GetEditorBase() const override { return *EditorView; }
-		virtual IEditableMultiReplicationStreamModel& GetModel() const override { return *MultiStreamModel; }
+		virtual IEditableMultiReplicationStreamModel& GetMultiStreamModel() const override { return *MultiStreamModel; }
+		virtual IReplicationStreamModel& GetConsolidatedModel() const override { return *ConsolidatedModel; }
 
 	private:
 
@@ -51,6 +53,5 @@ namespace UE::ConcertClientSharedSlate
 
 		/** The main view, which is displaying ConsolidatedModel. */
 		TSharedPtr<IReplicationStreamEditor> EditorView;
-		
 	};
 }
diff --git a/Engine/Plugins/Developer/Concert/ConcertUI/ConcertClientSharedSlate/Source/ConcertClientSharedSlate/Public/Replication/Editor/Model/IReplicationStreamModel.h b/Engine/Plugins/Developer/Concert/ConcertUI/ConcertClientSharedSlate/Source/ConcertClientSharedSlate/Public/Replication/Editor/Model/IReplicationStreamModel.h
--- a/Engine/Plugins/Developer/Concert/ConcertUI/ConcertClientSharedSlate/Source/ConcertClientSharedSlate/Public/Replication/Editor/Model/IReplicationStreamModel.h
+++ b/Engine/Plugins/Developer/Concert/ConcertUI/ConcertClientSharedSlate/Source/ConcertClientSharedSlate/Public/Replication/Editor/Model/IReplicationStreamModel.h
@@ -132,6 +132,18 @@ namespace UE::ConcertClientSharedSlate
 					: EBreakBehavior::Continue;
 			});
 		}
+
+		/** @return Gets the subobjects as an array */
+		TArray<FSoftObjectPath> GetSubobjects(const FSoftObjectPath& Parent) const
+		{
+			TArray<FSoftObjectPath> Subobjects;
+			ForEachSubobject(Parent, [&Subobjects](const FSoftObjectPath& Child)
+			{
+				Subobjects.Add(Child);
+				return EBreakBehavior::Continue;
+			});
+			return Subobjects;
+		}
 		
 		virtual ~IReplicationStreamModel() = default;
 	};
diff --git a/Engine/Plugins/Developer/Concert/ConcertUI/ConcertClientSharedSlate/Source/ConcertClientSharedSlate/Public/Replication/Editor/View/IMultiReplicationStreamEditor.h b/Engine/Plugins/Developer/Concert/ConcertUI/ConcertClientSharedSlate/Source/ConcertClientSharedSlate/Public/Replication/Editor/View/IMultiReplicationStreamEditor.h
--- a/Engine/Plugins/Developer/Concert/ConcertUI/ConcertClientSharedSlate/Source/ConcertClientSharedSlate/Public/Replication/Editor/View/IMultiReplicationStreamEditor.h
+++ b/Engine/Plugins/Developer/Concert/ConcertUI/ConcertClientSharedSlate/Source/ConcertClientSharedSlate/Public/Replication/Editor/View/IMultiReplicationStreamEditor.h
@@ -17,10 +17,13 @@ namespace UE::ConcertClientSharedSlate
 	{
 	public:
 
-		/** @return The widget drawing the consolidate model. */
+		/** @return The widget drawing the consolidated model. */
 		virtual IReplicationStreamEditor& GetEditorBase() const = 0;
 
-		/** @return The underlying model */
-		virtual IEditableMultiReplicationStreamModel& GetModel() const = 0;
+		/** @return The source of the sub-streams */
+		virtual IEditableMultiReplicationStreamModel& GetMultiStreamModel() const = 0;
+
+		/** @return Gets a model that combines all streams into one. */
+		virtual IReplicationStreamModel& GetConsolidatedModel() const = 0;
 	};
 }
diff --git a/Engine/Plugins/Developer/Concert/ConcertUI/ConcertClientSharedSlate/Source/ConcertClientSharedSlate/Public/Replication/Editor/View/ReplicationColumn.h b/Engine/Plugins/Developer/Concert/ConcertUI/ConcertClientSharedSlate/Source/ConcertClientSharedSlate/Public/Replication/Editor/View/ReplicationColumn.h
--- a/Engine/Plugins/Developer/Concert/ConcertUI/ConcertClientSharedSlate/Source/ConcertClientSharedSlate/Public/Replication/Editor/View/ReplicationColumn.h
+++ b/Engine/Plugins/Developer/Concert/ConcertUI/ConcertClientSharedSlate/Source/ConcertClientSharedSlate/Public/Replication/Editor/View/ReplicationColumn.h
@@ -90,7 +90,7 @@ namespace UE::ConcertClientSharedSlate
 					})
 					.PopulateSearchItems_Lambda([TransformOperation, PopulateSearchString](const TOtherColumnType& InOtherRowData, TArray<FString>& InOutSearchStrings)
 					{
-						if (ensure(PopulateSearchString.IsBound()))
+						if (PopulateSearchString.IsBound())
 						{
 							PopulateSearchString.Execute(TransformOperation(InOtherRowData), InOutSearchStrings);
 						}
