Add support for sending BP struct data over Multi-user.

diff --git a/Engine/Plugins/Developer/Concert/ConcertApp/MultiUserClient/Source/MultiUserClientLibrary/MultiUserClientLibrary.Build.cs b/Engine/Plugins/Developer/Concert/ConcertApp/MultiUserClient/Source/MultiUserClientLibrary/MultiUserClientLibrary.Build.cs
--- a/Engine/Plugins/Developer/Concert/ConcertApp/MultiUserClient/Source/MultiUserClientLibrary/MultiUserClientLibrary.Build.cs
+++ b/Engine/Plugins/Developer/Concert/ConcertApp/MultiUserClient/Source/MultiUserClientLibrary/MultiUserClientLibrary.Build.cs
@@ -11,7 +11,7 @@ namespace UnrealBuildTool.Rules
 				{
 					"Core",
 					"CoreUObject",
-					"Engine",
+					"Engine"
 				}
 			);
 
@@ -24,13 +24,15 @@ namespace UnrealBuildTool.Rules
 					{
 						"Concert",
 						"ConcertClient",
-						"ConcertSyncCore"
+						"ConcertSyncCore",
+						"ConcertTransport"
 					}
 				);
 
 				PrivateIncludePathModuleNames.AddRange(
 					new string[]
 					{
+						"Concert",
 						"ConcertSyncCore",
 						"ConcertSyncClient",
 						"MultiUserClient",
@@ -47,6 +49,13 @@ namespace UnrealBuildTool.Rules
 			else
 			{
 				PrivateDefinitions.Add("WITH_CONCERT=0");
+				PrivateDependencyModuleNames.AddRange(
+					new string[]
+					{
+						"Concert",
+						"ConcertTransport"
+					}
+				);
 			}
 		}
 	}
diff --git a/Engine/Plugins/Developer/Concert/ConcertApp/MultiUserClient/Source/MultiUserClientLibrary/Private/MultiUserSubsystem.cpp b/Engine/Plugins/Developer/Concert/ConcertApp/MultiUserClient/Source/MultiUserClientLibrary/Private/MultiUserSubsystem.cpp
new file mode 100644
--- /dev/null
+++ b/Engine/Plugins/Developer/Concert/ConcertApp/MultiUserClient/Source/MultiUserClientLibrary/Private/MultiUserSubsystem.cpp
@@ -0,0 +1,306 @@
+// Copyright Epic Games, Inc. All Rights Reserved.
+
+#include "MultiUserSubsystem.h"
+
+#include "Engine/Engine.h"
+#include "Engine/GameEngine.h"
+#include "Delegates/IDelegateInstance.h"
+#include "Templates/UniquePtr.h"
+
+#if WITH_CONCERT
+#include "ConcertTransportMessages.h"
+#include "IConcertClient.h"
+#include "IConcertSyncClient.h"
+#include "IConcertSession.h"
+#include "IConcertSessionHandler.h"
+#include "IConcertSyncClientModule.h"
+#include "ConcertLogGlobal.h"
+#include "MultiUserSubsystemMessages.h"
+#endif
+
+
+#if WITH_CONCERT
+namespace UE::MultiUserSubsystem::Private
+{
+struct FConcertManager
+{
+	FConcertManager()
+	{
+		if (TSharedPtr<IConcertSyncClient> ConcertSyncClient = IConcertSyncClientModule::Get().GetClient(TEXT("MultiUser")))
+		{
+			IConcertClientRef ConcertClient = ConcertSyncClient->GetConcertClient();
+			ConcertClient->OnSessionStartup().AddRaw(this, &FConcertManager::Register);
+			ConcertClient->OnSessionShutdown().AddRaw(this, &FConcertManager::Unregister);
+
+			if (TSharedPtr<IConcertClientSession> ConcertClientSession = ConcertClient->GetCurrentSession())
+			{
+				Register(ConcertClientSession.ToSharedRef());
+			}
+
+			ConcertClient->OnSessionConnectionChanged().AddRaw(this, &FConcertManager::OnSessionConnectionChanged);
+		}
+	}
+
+	~FConcertManager()
+	{
+		if (IConcertSyncClientModule::IsAvailable())
+		{
+			Unregister();
+			if (const TSharedPtr<IConcertSyncClient> ConcertSyncClient = IConcertSyncClientModule::Get().GetClient(TEXT("MultiUser")))
+			{
+				IConcertClientRef ConcertClient = ConcertSyncClient->GetConcertClient();
+				ConcertClient->OnSessionStartup().RemoveAll(this);
+				ConcertClient->OnSessionShutdown().RemoveAll(this);
+			}
+		}
+	}
+
+	void ConnectToSession()
+	{
+		if (GEngine)
+		{
+			UMultiUserSubsystem* SubSystem = GEngine->GetEngineSubsystem<UMultiUserSubsystem>();
+			if (SubSystem)
+			{
+				SubSystem->OnSessionConnected.Broadcast();
+			}
+		}
+	}
+
+
+	void DisconnectFromSession()
+	{
+		if (GEngine)
+		{
+			UMultiUserSubsystem* SubSystem = GEngine->GetEngineSubsystem<UMultiUserSubsystem>();
+			if (SubSystem)
+			{
+				SubSystem->OnSessionDisconnected.Broadcast();
+			}
+		}
+	}
+
+	void OnSessionConnectionChanged(IConcertClientSession& InSession, EConcertConnectionStatus ConnectionStatus)
+	{
+		if (ConnectionStatus == EConcertConnectionStatus::Connected)
+		{
+			DisconnectFromSession();
+			ConnectToSession();
+		}
+		else if (ConnectionStatus == EConcertConnectionStatus::Disconnecting)
+		{
+			DisconnectFromSession();
+		}
+	}
+
+	/**
+	 * Register a new multi-user session.
+	 */
+	void Register(TSharedRef<IConcertClientSession> InSession)
+	{
+		WeakSession = InSession;
+
+		InSession->RegisterCustomEventHandler<FConcertBlueprintEvent>(this, &FConcertManager::HandleBlueprintEvent);
+
+	}
+
+	/**
+	 * Unregister a multi-user session (if one exists)
+	 */
+	void Unregister()
+	{
+		if (TSharedPtr<IConcertClientSession> Session = WeakSession.Pin())
+		{
+			Session->UnregisterCustomEventHandler<FConcertBlueprintEvent>(this);
+		}
+
+		WeakSession.Reset();
+	}
+
+	/**
+	 * Unregister call that is invoked by Multi-user when it is shutdown.
+	 */
+	void Unregister(TSharedRef<IConcertClientSession> InSession)
+	{
+		Unregister();
+	}
+
+	void HandleBlueprintEvent(const FConcertSessionContext& Context, const FConcertBlueprintEvent& InEvent)
+	{
+		check(GEngine);
+		UMultiUserSubsystem* SubSystem = GEngine->GetEngineSubsystem<UMultiUserSubsystem>();
+		check(SubSystem);
+
+		FStructOnScope StructOnScope;
+		if (InEvent.Data.GetPayload(StructOnScope))
+		{
+			SubSystem->DispatchEvent(InEvent);
+		}
+	}
+
+	/** Weak pointer to the client session with which to send events. May be null or stale. */
+	TWeakPtr<IConcertClientSession> WeakSession;
+};
+
+TUniquePtr<FConcertManager> MakeConcertManager()
+{
+	return MakeUnique<FConcertManager>();
+}
+
+static TUniquePtr<FConcertManager> ConcertManager;
+
+} // namespace UE::MultiUserSubsystem::Private
+#endif
+
+void UMultiUserSubsystem::Initialize(FSubsystemCollectionBase& Collection)
+{
+#if WITH_CONCERT
+	using namespace UE::MultiUserSubsystem::Private;
+	ConcertManager = MakeConcertManager();
+#endif
+}
+
+void UMultiUserSubsystem::Deinitialize()
+{
+#if WITH_CONCERT
+	UE::MultiUserSubsystem::Private::ConcertManager.Reset();
+#endif
+}
+
+void UMultiUserSubsystem::K2_SendCustomEvent(const int32& Message)
+{
+	// This will never be called, the exec version below will be hit instead
+	check(0);
+}
+
+
+DEFINE_FUNCTION(UMultiUserSubsystem::execK2_SendCustomEvent)
+{
+#if WITH_CONCERT
+	Stack.StepCompiledIn<FStructProperty>(nullptr);
+	void* MessagePtr = Stack.MostRecentPropertyAddress;
+	FStructProperty* StructProp = CastField<FStructProperty>(Stack.MostRecentProperty);
+
+	P_FINISH;
+
+	if (ensure((StructProp != nullptr) && (StructProp->Struct != nullptr) && (MessagePtr != nullptr)))
+	{
+		TSharedPtr<IConcertClientSession> Session = UE::MultiUserSubsystem::Private::ConcertManager->WeakSession.Pin();
+		FConcertBlueprintEvent OutEvent;
+		OutEvent.Data.SetPayload(StructProp->Struct, MessagePtr);
+
+		P_THIS->SendCustomEvent(OutEvent);
+	}
+#endif
+}
+
+void UMultiUserSubsystem::K2_ExtractEventData(FMultiUserBlueprintEventData& EventData, int32& StructOut)
+{
+	// This will never be called, the exec version below will be hit instead
+	check(0);
+}
+
+DEFINE_FUNCTION(UMultiUserSubsystem::execK2_ExtractEventData)
+{
+#if WITH_CONCERT
+	P_GET_STRUCT_REF(FMultiUserBlueprintEventData, EventData);
+
+	Stack.MostRecentPropertyAddress = nullptr;
+	Stack.MostRecentPropertyContainer = nullptr;
+	Stack.StepCompiledIn<FStructProperty>(nullptr);
+	void* StructOutPtr = Stack.MostRecentPropertyAddress;
+	FStructProperty* StructProp = CastField<FStructProperty>(Stack.MostRecentProperty);
+
+	P_FINISH;
+
+	if (StructOutPtr && StructProp && StructProp->Struct)
+	{
+		if (EventData.SharedEventDataPtr.IsValid())
+		{
+			EventData.SharedEventDataPtr->Data.GetPayload(StructProp->Struct, StructOutPtr);
+		}
+	}
+#endif
+}
+
+bool UMultiUserSubsystem::IsConnectedToSession() const
+{
+#if WITH_CONCERT
+	using namespace UE::MultiUserSubsystem::Private;
+	if (ConcertManager.IsValid())
+	{
+		return ConcertManager->WeakSession.IsValid();
+	}
+#endif
+	return false;
+}
+
+void UMultiUserSubsystem::SendCustomEvent(const FConcertBlueprintEvent& EventData)
+{
+#if WITH_CONCERT
+	using namespace UE::MultiUserSubsystem::Private;
+	if (!ConcertManager.IsValid())
+	{
+		UE_LOG(LogConcert, Error, TEXT("No concert manager initialized. Unable to send Blueprint RPC over Multi-user."));
+		return;
+	}
+
+	if (!ConcertManager->WeakSession.IsValid())
+	{
+		UE_LOG(LogConcert, Warning, TEXT("Attempting to send a Multi-user custom event when not in a multi-user session. "));
+		return;
+	}
+
+	TSharedPtr<IConcertClientSession> Session = ConcertManager->WeakSession.Pin();
+	Session->SendCustomEvent(EventData, Session->GetSessionClientEndpointIds(), EConcertMessageFlags::ReliableOrdered);
+#endif
+}
+
+bool UMultiUserSubsystem::RegisterCustomEventHandler(const UStruct* EventType, FCustomEventHandler InEventHandler)
+{
+#if WITH_CONCERT
+	FName EventTypeName = *EventType->GetPathName();
+	FCustomEventHandler* Handler = CustomEventHandlerTable.Find(EventTypeName);
+	if (Handler != nullptr)
+	{
+		UE_LOG(LogConcert, Error, TEXT("Multi-user Subsystem: A custom event handler for %s already exists and a new handler cannot be registered."), *EventTypeName.ToString());
+		return false;
+	}
+	UE_LOG(LogConcert, Display, TEXT("Multi-user Subsystem: Registered custom event handler for %s messages."), *EventTypeName.ToString());
+	CustomEventHandlerTable.Emplace(EventTypeName, MoveTemp(InEventHandler));
+	return true;
+#else
+	return false;
+#endif
+}
+
+bool UMultiUserSubsystem::UnregisterCustomEventHandler(const UStruct* EventType)
+{
+#if WITH_CONCERT
+	FName EventTypeName = *EventType->GetPathName();
+	FCustomEventHandler* Handler = CustomEventHandlerTable.Find(EventTypeName);
+	if (Handler == nullptr)
+	{
+		return false;
+	}
+	UE_LOG(LogConcert, Display, TEXT("Multi-user Subsystem: Unregistered custom event handler for %s messages."), *EventTypeName.ToString());
+	CustomEventHandlerTable.Remove(EventTypeName);
+
+	return true;
+#else
+	return false;
+#endif
+}
+
+void UMultiUserSubsystem::DispatchEvent(const FConcertBlueprintEvent& InEvent)
+{
+#if WITH_CONCERT
+	FCustomEventHandler* Handler = CustomEventHandlerTable.Find(InEvent.Data.PayloadTypeName);
+	if (Handler)
+	{
+		FMultiUserBlueprintEventData BPEvent;
+		BPEvent.SharedEventDataPtr = MakeShared<FConcertBlueprintEvent>(InEvent);
+		Handler->ExecuteIfBound(BPEvent);
+	}
+#endif
+}
diff --git a/Engine/Plugins/Developer/Concert/ConcertApp/MultiUserClient/Source/MultiUserClientLibrary/Private/MultiUserSubsystemMessages.h b/Engine/Plugins/Developer/Concert/ConcertApp/MultiUserClient/Source/MultiUserClientLibrary/Private/MultiUserSubsystemMessages.h
new file mode 100644
--- /dev/null
+++ b/Engine/Plugins/Developer/Concert/ConcertApp/MultiUserClient/Source/MultiUserClientLibrary/Private/MultiUserSubsystemMessages.h
@@ -0,0 +1,15 @@
+// Copyright Epic Games, Inc. All Rights Reserved.
+#pragma once
+
+#include "ConcertMessageData.h"
+
+#include "MultiUserSubsystemMessages.generated.h"
+
+USTRUCT()
+struct FConcertBlueprintEvent
+{
+	GENERATED_BODY()
+
+	UPROPERTY()
+	FConcertSessionSerializedPayload Data;
+};
diff --git a/Engine/Plugins/Developer/Concert/ConcertApp/MultiUserClient/Source/MultiUserClientLibrary/Public/MultiUserSubsystem.h b/Engine/Plugins/Developer/Concert/ConcertApp/MultiUserClient/Source/MultiUserClientLibrary/Public/MultiUserSubsystem.h
new file mode 100644
--- /dev/null
+++ b/Engine/Plugins/Developer/Concert/ConcertApp/MultiUserClient/Source/MultiUserClientLibrary/Public/MultiUserSubsystem.h
@@ -0,0 +1,106 @@
+// Copyright Epic Games, Inc. All Rights Reserved.
+
+#pragma once
+
+#include "Delegates/DelegateCombinations.h"
+#include "Engine/UserDefinedStruct.h"
+#include "Subsystems/EngineSubsystem.h"
+#include "Templates/SharedPointer.h"
+
+#include "MultiUserSubsystem.generated.h"
+
+DECLARE_DYNAMIC_MULTICAST_DELEGATE(FOnSessionConnected);
+DECLARE_DYNAMIC_MULTICAST_DELEGATE(FOnSessionDisconnected);
+
+struct FConcertSessionContext;
+struct FConcertBlueprintEvent;
+
+USTRUCT(BlueprintType)
+struct FMultiUserBlueprintEventData
+{
+	GENERATED_BODY();
+
+	TSharedPtr<FConcertBlueprintEvent> SharedEventDataPtr;
+};
+
+DECLARE_DYNAMIC_DELEGATE_OneParam(FCustomEventHandler, FMultiUserBlueprintEventData, Data);
+
+UCLASS(BlueprintType, Category = "Multi-user", DisplayName = "MultiUserSubsystem")
+class MULTIUSERCLIENTLIBRARY_API UMultiUserSubsystem : public UEngineSubsystem
+{
+	GENERATED_BODY()
+
+	virtual void Initialize(FSubsystemCollectionBase& Collection) override;
+	virtual void Deinitialize() override;
+
+public:
+	/**
+	 * Delegate invoked when a session connection event has occurred.
+	 */
+	UPROPERTY(BlueprintAssignable, Category = "Multi-user")
+	FOnSessionConnected OnSessionConnected;
+
+	/**
+	 * Delegate invoked when a session disconnection event has occurred.
+	 */
+	UPROPERTY(BlueprintAssignable, Category = "Multi-user")
+	FOnSessionDisconnected OnSessionDisconnected;
+
+	/**
+	 * Returns true if we are currently in a session.
+	 */
+	UFUNCTION(BlueprintCallable, Category = "Multi-user")
+	bool IsConnectedToSession() const;
+
+	/**
+	 * Send a custom event message over multi-user.  If you a not in an active session, this function does nothing and only
+	 * reports and warning to the output log.
+	 *
+	 * @param EventData			The event data to send. This must be the same type of UScriptStruct expected by the receivers on the registered handler.
+	 */
+	UFUNCTION(BlueprintCallable, CustomThunk, Category="Multi-user", meta=(CustomStructureParam="EventData", AllowAbstract="false", DisplayName="Send Custom Event"))
+	void K2_SendCustomEvent(const int32& EventData);
+
+	DECLARE_FUNCTION(execK2_SendCustomEvent);
+
+	/**
+	 * Extract event message data. Given a MultiUserBlueprintEventData extract the message contents and put it into the desired
+	 * struct.  The struct must match the target type.
+	 *
+	 * @param StructOut	   The a reference to the structure that will receive the data.
+	 */
+	UFUNCTION(BlueprintCallable, CustomThunk, Category="Multi-user", meta=(CustomStructureParam="StructOut", AllowAbstract="false", DisplayName="Extract Event Data"))
+	void K2_ExtractEventData(UPARAM(ref) FMultiUserBlueprintEventData& EventData, UPARAM(ref) int32& StructOut);
+
+	DECLARE_FUNCTION(execK2_ExtractEventData);
+
+	/**
+	 * Register an event handler for the given type. Only one handler per type is allowed. You should register your handler
+	 * on session connected and remove the handler when the session disconnects.
+	 *
+	 * @param EventType   The struct type that we support handling.
+	 *
+	 * @return true if the type handler was accepted.
+	 */
+	UFUNCTION(BlueprintCallable, Category = "Multi-user")
+	bool RegisterCustomEventHandler(const UStruct* EventType, FCustomEventHandler InEventHandler);
+
+	/**
+	 * Unregister an event handler for the given type.
+	 *
+	 * @param EventType   The struct type that we support handling.
+	 *
+	 * @return true if the type handler was successfully unregistered.
+	 */
+	UFUNCTION(BlueprintCallable, Category = "Multi-user")
+	bool UnregisterCustomEventHandler(const UStruct* EventType);
+
+	// Public interface to send the event to the registered blueprint handler.
+	void DispatchEvent(const FConcertBlueprintEvent& InEvent);
+private:
+	// Table of registered handlers.
+	TMap<FName, FCustomEventHandler> CustomEventHandlerTable;
+
+	// Internal send custom event message.
+	void SendCustomEvent(const FConcertBlueprintEvent& EventData);
+};
