Refactor ISubmissionWorkflow to be given the changes to submit instead of injecting them itself

diff --git a/Engine/Plugins/Developer/Concert/ConcertApp/MultiUserClient/Source/MultiUserClient/Private/Replication/Authority/AuthorityChangeTracker.cpp b/Engine/Plugins/Developer/Concert/ConcertApp/MultiUserClient/Source/MultiUserClient/Private/Replication/Authority/AuthorityChangeTracker.cpp
--- a/Engine/Plugins/Developer/Concert/ConcertApp/MultiUserClient/Source/MultiUserClient/Private/Replication/Authority/AuthorityChangeTracker.cpp
+++ b/Engine/Plugins/Developer/Concert/ConcertApp/MultiUserClient/Source/MultiUserClient/Private/Replication/Authority/AuthorityChangeTracker.cpp
@@ -88,10 +88,14 @@ namespace UE::MultiUserClient
 		return AuthoritySynchronizer.GetChangeAuthorityMutability(ObjectPath);
 	}
 
-	ConcertSyncClient::Replication::FAuthorityChangeRequest FAuthorityChangeTracker::BuildChangeRequest(const FGuid& StreamId) const
+	TOptional<ConcertSyncClient::Replication::FAuthorityChangeRequest> FAuthorityChangeTracker::BuildChangeRequest(const FGuid& StreamId) const
 	{
+		if (NewAuthorityStates.IsEmpty())
+		{
+			return {};
+		}
+		
 		ConcertSyncClient::Replication::FAuthorityChangeRequest ChangeRequest;
-
 		for (const TPair<FSoftObjectPath, bool>& NewAuthorityState : NewAuthorityStates)
 		{
 			if (NewAuthorityState.Value)
diff --git a/Engine/Plugins/Developer/Concert/ConcertApp/MultiUserClient/Source/MultiUserClient/Private/Replication/Authority/AuthorityChangeTracker.h b/Engine/Plugins/Developer/Concert/ConcertApp/MultiUserClient/Source/MultiUserClient/Private/Replication/Authority/AuthorityChangeTracker.h
--- a/Engine/Plugins/Developer/Concert/ConcertApp/MultiUserClient/Source/MultiUserClient/Private/Replication/Authority/AuthorityChangeTracker.h
+++ b/Engine/Plugins/Developer/Concert/ConcertApp/MultiUserClient/Source/MultiUserClient/Private/Replication/Authority/AuthorityChangeTracker.h
@@ -45,8 +45,8 @@ namespace UE::MultiUserClient
 		/** @return Detailed information why the object's authority can or cannot be changed.  */
 		EAuthorityMutability GetChangeAuthorityMutability(const FSoftObjectPath& ObjectPath) const;
 
-		/** Builds a change request from the local changes. */
-		ConcertSyncClient::Replication::FAuthorityChangeRequest BuildChangeRequest(const FGuid& StreamId) const;
+		/** Builds a change request from the local changes, if there are changes. */
+		TOptional<ConcertSyncClient::Replication::FAuthorityChangeRequest> BuildChangeRequest(const FGuid& StreamId) const;
 
 		/** Called entries are added to NewAuthorityStates. NOT called when objects are removed. */
 		DECLARE_MULTICAST_DELEGATE(FOnAuthorityChangeMade);
diff --git a/Engine/Plugins/Developer/Concert/ConcertApp/MultiUserClient/Source/MultiUserClient/Private/Replication/Client/LocalReplicationClient.cpp b/Engine/Plugins/Developer/Concert/ConcertApp/MultiUserClient/Source/MultiUserClient/Private/Replication/Client/LocalReplicationClient.cpp
--- a/Engine/Plugins/Developer/Concert/ConcertApp/MultiUserClient/Source/MultiUserClient/Private/Replication/Client/LocalReplicationClient.cpp
+++ b/Engine/Plugins/Developer/Concert/ConcertApp/MultiUserClient/Source/MultiUserClient/Private/Replication/Client/LocalReplicationClient.cpp
@@ -25,9 +25,6 @@ namespace UE::MultiUserClient
 			{
 				return GetStreamDiffer().DoesObjectHavePropertiesAfterSubmit(ObjectPath);
 			})),
-			[this, InClient, &InAuthorityCache](FStreamChangeTracker& InStreamChangeTracker, FAuthorityChangeTracker& InAuthorityChangeTracker, IClientStreamSynchronizer& InStreamSynchronizer) mutable
-			{
-				return MakeUnique<FSubmissionWorkflow_LocalClient>(MoveTemp(InClient), InStreamChangeTracker, InAuthorityChangeTracker, InStreamSynchronizer, InAuthorityCache);
-			})
+			MakeUnique<FSubmissionWorkflow_LocalClient>(MoveTemp(InClient)))
 	{}
 }
diff --git a/Engine/Plugins/Developer/Concert/ConcertApp/MultiUserClient/Source/MultiUserClient/Private/Replication/Client/RemoteReplicationClient.cpp b/Engine/Plugins/Developer/Concert/ConcertApp/MultiUserClient/Source/MultiUserClient/Private/Replication/Client/RemoteReplicationClient.cpp
--- a/Engine/Plugins/Developer/Concert/ConcertApp/MultiUserClient/Source/MultiUserClient/Private/Replication/Client/RemoteReplicationClient.cpp
+++ b/Engine/Plugins/Developer/Concert/ConcertApp/MultiUserClient/Source/MultiUserClient/Private/Replication/Client/RemoteReplicationClient.cpp
@@ -27,10 +27,7 @@ namespace UE::MultiUserClient
 				{
 					return GetStreamDiffer().DoesObjectHavePropertiesAfterSubmit(ObjectPath);
 				})),
-			[](FStreamChangeTracker& InStreamChangeTracker, FAuthorityChangeTracker& InAuthorityChangeTracker, IClientStreamSynchronizer& InStreamSynchronizer)
-			{
-				return MakeUnique<FSubmissionWorkflow_RemoteClient>();
-			})
+			MakeUnique<FSubmissionWorkflow_RemoteClient>())
 	{
 		// When the remote client's state has changed, refresh the UI.
 		GetStreamSynchronizer().OnServerStateChanged().AddLambda([this]()
diff --git a/Engine/Plugins/Developer/Concert/ConcertApp/MultiUserClient/Source/MultiUserClient/Private/Replication/Client/ReplicationClient.cpp b/Engine/Plugins/Developer/Concert/ConcertApp/MultiUserClient/Source/MultiUserClient/Private/Replication/Client/ReplicationClient.cpp
--- a/Engine/Plugins/Developer/Concert/ConcertApp/MultiUserClient/Source/MultiUserClient/Private/Replication/Client/ReplicationClient.cpp
+++ b/Engine/Plugins/Developer/Concert/ConcertApp/MultiUserClient/Source/MultiUserClient/Private/Replication/Client/ReplicationClient.cpp
@@ -18,12 +18,13 @@ namespace UE::MultiUserClient
 		UMultiUserReplicationClientPreset& InSessionContent,
 		TUniquePtr<IClientStreamSynchronizer> InStreamSynchronizer,
 		TUniquePtr<IClientAuthoritySynchronizer> InAuthoritySynchronizer,
-		TFunctionRef<FMakeSubmissionWorkflow> MakeSubmissionWorkflowFunc
+		TUniquePtr<ISubmissionWorkflow> InSubmissionWorkflow
 		)
 		: EndpointId(EndpointId)
 		, ClientContentStorage(&InSessionContent)
 		, StreamSynchronizer(MoveTemp(InStreamSynchronizer))
 		, AuthoritySynchronizer(MoveTemp(InAuthoritySynchronizer))
+		, SubmissionWorkflow(MoveTemp(InSubmissionWorkflow))
 		, LocalClientEditModel(ConcertClientSharedSlate::CreatePropertySelectionModel(
 			*ClientContentStorage->Stream,
 			ClientContentStorage->Stream->MakeReplicationMapGetterAttribute()
@@ -34,8 +35,8 @@ namespace UE::MultiUserClient
 			FStreamChangeTracker::FOnModifyReplicationMap::CreateLambda([this](){ ClientContentStorage->Stream->Modify(); })
 			)
 		, LocalAuthorityDiffer(EndpointId, *AuthoritySynchronizer, InAuthorityCache)
-		, SubmissionWorkflow(MakeSubmissionWorkflowFunc(LocalClientStreamDiffer, LocalAuthorityDiffer, *StreamSynchronizer.Get()))
-		, AutoSubmissionPolicy(*SubmissionWorkflow.Get(), LocalClientEditModel.Get(), LocalAuthorityDiffer)
+		, ChangeRequestBuilder(EndpointId, InAuthorityCache, *StreamSynchronizer, LocalClientStreamDiffer, LocalAuthorityDiffer)
+		, AutoSubmissionPolicy(*SubmissionWorkflow.Get(), ChangeRequestBuilder, LocalClientEditModel.Get(), LocalAuthorityDiffer)
 	{
 		LocalClientEditModel->OnObjectsChanged().AddRaw(this, &FReplicationClient::OnObjectsChanged);
 		LocalClientEditModel->OnPropertiesChanged().AddRaw(this, &FReplicationClient::OnPropertiesChanged);
diff --git a/Engine/Plugins/Developer/Concert/ConcertApp/MultiUserClient/Source/MultiUserClient/Private/Replication/Client/ReplicationClient.h b/Engine/Plugins/Developer/Concert/ConcertApp/MultiUserClient/Source/MultiUserClient/Private/Replication/Client/ReplicationClient.h
--- a/Engine/Plugins/Developer/Concert/ConcertApp/MultiUserClient/Source/MultiUserClient/Private/Replication/Client/ReplicationClient.h
+++ b/Engine/Plugins/Developer/Concert/ConcertApp/MultiUserClient/Source/MultiUserClient/Private/Replication/Client/ReplicationClient.h
@@ -7,6 +7,7 @@
 #include "Replication/Stream/IClientStreamSynchronizer.h"
 #include "Replication/Stream/StreamChangeTracker.h"
 #include "Replication/Submission/AutoSubmissionPolicy.h"
+#include "Replication/Submission/ChangeRequestBuilder.h"
 #include "Replication/Submission/ISubmissionWorkflow.h"
 
 #include "Templates/SharedPointer.h"
@@ -32,16 +33,13 @@ namespace UE::MultiUserClient
 	{
 	public:
 		
-		/** Indirection for creating ISubmissionWorkflow because some ISubmissionWorkflow implementations needs members constructed in FReplicationClient. */
-		using FMakeSubmissionWorkflow = TUniquePtr<ISubmissionWorkflow>(FStreamChangeTracker&, FAuthorityChangeTracker&, IClientStreamSynchronizer&);
-
 		FReplicationClient(
 			const FGuid& EndpointId,
 			FGlobalAuthorityCache& InAuthorityCache,
 			UMultiUserReplicationClientPreset& InSessionContent,
 			TUniquePtr<IClientStreamSynchronizer> InStreamSynchronizer,
 			TUniquePtr<IClientAuthoritySynchronizer> InAuthoritySynchronizer,
-			TFunctionRef<FMakeSubmissionWorkflow> MakeSubmissionWorkflowFunc
+			TUniquePtr<ISubmissionWorkflow> InSubmissionWorkflow
 			);
 		~FReplicationClient();
 
@@ -87,6 +85,8 @@ namespace UE::MultiUserClient
 		TUniquePtr<IClientStreamSynchronizer> StreamSynchronizer;
 		/** Keeps the client's authority state on the server in sync.*/
 		TUniquePtr<IClientAuthoritySynchronizer> AuthoritySynchronizer;
+		/** Handles the logic of submitting and reverting for this client. */
+		TUniquePtr<ISubmissionWorkflow> SubmissionWorkflow;
 		
 		/**
 		 * Used to detect changes made to the client's config by the local editor.
@@ -104,8 +104,8 @@ namespace UE::MultiUserClient
 		/** Tracks changes made to the client's authority state. */
 		FAuthorityChangeTracker LocalAuthorityDiffer;
 		
-		/** Handles the logic of submitting and reverting for this client. */
-		TUniquePtr<ISubmissionWorkflow> SubmissionWorkflow;
+		/** Shared logic for building stream and authority change requests based on local change made. */
+		FChangeRequestBuilder ChangeRequestBuilder;
 		/** Automatically submits changes as they are made by the user. */
 		FAutoSubmissionPolicy AutoSubmissionPolicy;
 
diff --git a/Engine/Plugins/Developer/Concert/ConcertApp/MultiUserClient/Source/MultiUserClient/Private/Replication/Submission/AutoSubmissionPolicy.cpp b/Engine/Plugins/Developer/Concert/ConcertApp/MultiUserClient/Source/MultiUserClient/Private/Replication/Submission/AutoSubmissionPolicy.cpp
--- a/Engine/Plugins/Developer/Concert/ConcertApp/MultiUserClient/Source/MultiUserClient/Private/Replication/Submission/AutoSubmissionPolicy.cpp
+++ b/Engine/Plugins/Developer/Concert/ConcertApp/MultiUserClient/Source/MultiUserClient/Private/Replication/Submission/AutoSubmissionPolicy.cpp
@@ -2,20 +2,23 @@
 
 #include "AutoSubmissionPolicy.h"
 
+#include "ChangeRequestBuilder.h"
 #include "ISubmissionWorkflow.h"
+#include "Replication/Authority/AuthorityChangeTracker.h"
 #include "Replication/Editor/Model/IEditableReplicationStreamModel.h"
 
 #include "Misc/CoreDelegates.h"
-#include "Replication/Authority/AuthorityChangeTracker.h"
 
 namespace UE::MultiUserClient
 {
 	FAutoSubmissionPolicy::FAutoSubmissionPolicy(
 		ISubmissionWorkflow& InSubmissionWorkflow,
+		const FChangeRequestBuilder& InRequestBuilder,
 		ConcertClientSharedSlate::IEditableReplicationStreamModel& InStreamEditorModel,
 		FAuthorityChangeTracker& InAuthorityChangeTracker
 		)
 		: SubmissionWorkflow(InSubmissionWorkflow)
+		, RequestBuilder(InRequestBuilder)
 		, StreamEditorModel(InStreamEditorModel)
 		, AuthorityChangeTracker(InAuthorityChangeTracker)
 	{
@@ -42,10 +45,23 @@ namespace UE::MultiUserClient
 
 	void FAutoSubmissionPolicy::OnEndFrame()
 	{
+		// Instead of ticking until CanSubmit return true, consider introducing a queueing system for tasks...
+		// but that only makes sense when there multiple code places that want to call SubmitChanges
 		if (SubmissionWorkflow.CanSubmit())
 		{
+			using namespace UE::ConcertSyncClient::Replication;
 			FCoreDelegates::OnEndFrame.RemoveAll(this);
-			SubmissionWorkflow.SubmitChanges();
+			
+			// Even though authority request is sent after server confirms stream change, the authority request is pre-built to avoid sending changes
+			// the local client makes while we're waiting for the latent server responses.
+			TOptional<FAuthorityChangeRequest> AuthorityChangeRequest = RequestBuilder.BuildAuthorityChange();
+			TOptional<FChangeStreamRequest> StreamRequest = RequestBuilder.BuildStreamChange();
+			
+			// Should be true since a change caused OnEndFrame to be registered in the first place, but in THEORY changes could have been reverted in the same frame (e.g. via a transaction)
+			if (AuthorityChangeRequest || StreamRequest)
+			{
+				SubmissionWorkflow.SubmitChanges({ StreamRequest, AuthorityChangeRequest });
+			}
 		}
 	}
 }
diff --git a/Engine/Plugins/Developer/Concert/ConcertApp/MultiUserClient/Source/MultiUserClient/Private/Replication/Submission/AutoSubmissionPolicy.h b/Engine/Plugins/Developer/Concert/ConcertApp/MultiUserClient/Source/MultiUserClient/Private/Replication/Submission/AutoSubmissionPolicy.h
--- a/Engine/Plugins/Developer/Concert/ConcertApp/MultiUserClient/Source/MultiUserClient/Private/Replication/Submission/AutoSubmissionPolicy.h
+++ b/Engine/Plugins/Developer/Concert/ConcertApp/MultiUserClient/Source/MultiUserClient/Private/Replication/Submission/AutoSubmissionPolicy.h
@@ -5,6 +5,11 @@
 #include "Replication/Editor/Model/IEditableReplicationStreamModel.h"
 #include "Templates/UnrealTemplate.h"
 
+namespace UE::MultiUserClient
+{
+	class FChangeRequestBuilder;
+}
+
 namespace UE::ConcertClientSharedSlate
 {
 	class IEditableReplicationStreamModel;
@@ -12,8 +17,10 @@ namespace UE::ConcertClientSharedSlate
 
 namespace UE::MultiUserClient
 {
-	class FAuthorityChangeTracker;
+	class IClientStreamSynchronizer;
 	class ISubmissionWorkflow;
+	class FAuthorityChangeTracker;
+	class FStreamChangeTracker;
 	
 	/**
 	 * Automatically submits stream and authority changes as the user makes them.
@@ -28,6 +35,7 @@ namespace UE::MultiUserClient
 		
 		FAutoSubmissionPolicy(
 			ISubmissionWorkflow& InSubmissionWorkflow,
+			const FChangeRequestBuilder& InRequestBuilder,
 			ConcertClientSharedSlate::IEditableReplicationStreamModel& InStreamEditorModel,
 			FAuthorityChangeTracker& InAuthorityChangeTracker
 			);
@@ -37,6 +45,8 @@ namespace UE::MultiUserClient
 
 		/** Handles performing the submission */
 		ISubmissionWorkflow& SubmissionWorkflow;
+		/** Used to building the requests that are passed to SubmissionWorkflow. */
+		const FChangeRequestBuilder& RequestBuilder;
 
 		/** Informs us when the stream is structurally changed by the user. */
 		ConcertClientSharedSlate::IEditableReplicationStreamModel& StreamEditorModel;
diff --git a/Engine/Plugins/Developer/Concert/ConcertApp/MultiUserClient/Source/MultiUserClient/Private/Replication/Submission/ChangeRequestBuilder.cpp b/Engine/Plugins/Developer/Concert/ConcertApp/MultiUserClient/Source/MultiUserClient/Private/Replication/Submission/ChangeRequestBuilder.cpp
new file mode 100644
--- /dev/null
+++ b/Engine/Plugins/Developer/Concert/ConcertApp/MultiUserClient/Source/MultiUserClient/Private/Replication/Submission/ChangeRequestBuilder.cpp
@@ -0,0 +1,55 @@
+// Copyright Epic Games, Inc. All Rights Reserved.
+
+#include "ChangeRequestBuilder.h"
+
+#include "Replication/Authority/AuthorityChangeTracker.h"
+#include "Replication/Stream/StreamChangeTracker.h"
+#include "Replication/Util/GlobalAuthorityCache.h"
+#include "Replication/Util/StreamRequestUtils.h"
+
+namespace UE::MultiUserClient
+{
+	FChangeRequestBuilder::FChangeRequestBuilder(
+		const FGuid& InLocalClientId,
+		const FGlobalAuthorityCache& InAuthorityCache,
+		IClientStreamSynchronizer& InStreamSynchronizer,
+		FStreamChangeTracker& InStreamChangeTracker,
+		FAuthorityChangeTracker& InAuthorityChangeTracker
+		)
+		: LocalClientId(InLocalClientId)
+		, AuthorityCache(InAuthorityCache)
+		, StreamSynchronizer(InStreamSynchronizer)
+		, StreamChangeTracker(InStreamChangeTracker)
+		, AuthorityChangeTracker(InAuthorityChangeTracker)
+	{}
+
+	TOptional<ConcertSyncClient::Replication::FChangeStreamRequest> FChangeRequestBuilder::BuildStreamChange() const
+	{
+		using namespace UE::ConcertSyncClient::Replication;
+		TOptional<FChangeStreamRequest> StreamRequest;
+		
+		const FStreamChangelist& Changelist = StreamChangeTracker.GetCachedDeltaChange();
+		const bool bIsChangelistEmpty = Changelist.ObjectsToPut.IsEmpty() && Changelist.ObjectsToRemove.IsEmpty();
+		if (!bIsChangelistEmpty)
+		{
+			StreamRequest = StreamSynchronizer.GetServerState().ReplicatedObjects.IsEmpty()
+				? StreamRequestUtils::BuildChangeRequest_CreateNewStream(GetLocalClientStreamId(), Changelist)
+				: StreamRequestUtils::BuildChangeRequest_UpdateExistingStream(Changelist);
+				
+			// Predict conflicts in case the remote client's streams have changed. Also: while the UI highlights "bad" requests, it does not correct it.
+			AuthorityCache.CleanseConflictsFromStreamRequest(*StreamRequest, LocalClientId);
+		}
+
+		return StreamRequest;
+	}
+
+	TOptional<ConcertSyncClient::Replication::FAuthorityChangeRequest> FChangeRequestBuilder::BuildAuthorityChange() const
+	{
+		return AuthorityChangeTracker.BuildChangeRequest(GetLocalClientStreamId());
+	}
+
+	FGuid FChangeRequestBuilder::GetLocalClientStreamId() const
+	{
+		return StreamSynchronizer.GetStreamId();
+	}
+}
diff --git a/Engine/Plugins/Developer/Concert/ConcertApp/MultiUserClient/Source/MultiUserClient/Private/Replication/Submission/ChangeRequestBuilder.h b/Engine/Plugins/Developer/Concert/ConcertApp/MultiUserClient/Source/MultiUserClient/Private/Replication/Submission/ChangeRequestBuilder.h
new file mode 100644
--- /dev/null
+++ b/Engine/Plugins/Developer/Concert/ConcertApp/MultiUserClient/Source/MultiUserClient/Private/Replication/Submission/ChangeRequestBuilder.h
@@ -0,0 +1,63 @@
+// Copyright Epic Games, Inc. All Rights Reserved.
+
+#pragma once
+
+#include "Misc/Guid.h"
+#include "Misc/Optional.h"
+
+namespace UE::ConcertSyncClient::Replication
+{
+	struct FChangeStreamRequest;
+	struct FAuthorityChangeRequest;
+}
+
+namespace UE::MultiUserClient
+{
+	class IClientStreamSynchronizer;
+	class FAuthorityChangeTracker;
+	class FGlobalAuthorityCache;
+	class FStreamChangeTracker;
+	
+	/**
+	 * Util that knows how to build FChangeStreamRequest and FAuthorityChangeRequest based on the local client's changes
+	 * and knowledge of other remote clients.
+	 */
+	class FChangeRequestBuilder
+	{
+	public:
+
+		FChangeRequestBuilder(
+			const FGuid& InLocalClientId,
+			const FGlobalAuthorityCache& InAuthorityCache,
+			IClientStreamSynchronizer& InStreamSynchronizer,
+			FStreamChangeTracker& InStreamChangeTracker,
+			FAuthorityChangeTracker& AuthorityChangeTracker
+			);
+
+		/** @return Valid request that can be sent to the server, if there are any local changes. */
+		TOptional<ConcertSyncClient::Replication::FChangeStreamRequest> BuildStreamChange() const;
+
+		/** @return Valid request that can be sent to the server, if there are any local changes. */
+		TOptional<ConcertSyncClient::Replication::FAuthorityChangeRequest> BuildAuthorityChange() const;
+
+	private:
+
+		/** Id of the local client in the session. */
+		const FGuid LocalClientId;
+
+		/** Used to filter out changes that would cause a conflict when submitted. */
+		const FGlobalAuthorityCache& AuthorityCache;
+		
+		/** Used to determine whether a new stream needs to be created when submitting (happens if nothing was previously registered). */
+		IClientStreamSynchronizer& StreamSynchronizer;
+		
+		/** Used to get changes made to the stream */
+		FStreamChangeTracker& StreamChangeTracker;
+		/** Informs us when authority is changed by the user. */
+		FAuthorityChangeTracker& AuthorityChangeTracker;
+
+		/** @return The ID MU uses for its single stream. */
+		FGuid GetLocalClientStreamId() const;
+	};
+}
+
diff --git a/Engine/Plugins/Developer/Concert/ConcertApp/MultiUserClient/Source/MultiUserClient/Private/Replication/Submission/Data/EChangeRevertability.h b/Engine/Plugins/Developer/Concert/ConcertApp/MultiUserClient/Source/MultiUserClient/Private/Replication/Submission/Data/EChangeRevertability.h
deleted file mode 100644
diff --git a/Engine/Plugins/Developer/Concert/ConcertApp/MultiUserClient/Source/MultiUserClient/Private/Replication/Submission/Data/EChangeUploadability.h b/Engine/Plugins/Developer/Concert/ConcertApp/MultiUserClient/Source/MultiUserClient/Private/Replication/Submission/Data/EChangeUploadability.h
--- a/Engine/Plugins/Developer/Concert/ConcertApp/MultiUserClient/Source/MultiUserClient/Private/Replication/Submission/Data/EChangeUploadability.h
+++ b/Engine/Plugins/Developer/Concert/ConcertApp/MultiUserClient/Source/MultiUserClient/Private/Replication/Submission/Data/EChangeUploadability.h
@@ -9,8 +9,6 @@ namespace UE::MultiUserClient
 		/** Can click the upload button */
 		Ready,
 
-		/** There are no changes to upload */
-		NoChanges,
 		/** A previous upload operation is in progress */
 		InProgress,
 		
diff --git a/Engine/Plugins/Developer/Concert/ConcertApp/MultiUserClient/Source/MultiUserClient/Private/Replication/Submission/ISubmissionWorkflow.h b/Engine/Plugins/Developer/Concert/ConcertApp/MultiUserClient/Source/MultiUserClient/Private/Replication/Submission/ISubmissionWorkflow.h
--- a/Engine/Plugins/Developer/Concert/ConcertApp/MultiUserClient/Source/MultiUserClient/Private/Replication/Submission/ISubmissionWorkflow.h
+++ b/Engine/Plugins/Developer/Concert/ConcertApp/MultiUserClient/Source/MultiUserClient/Private/Replication/Submission/ISubmissionWorkflow.h
@@ -2,19 +2,33 @@
 
 #pragma once
 
-#include "Data/EChangeRevertability.h"
 #include "Data/EChangeUploadability.h"
+#include "Replication/IConcertClientReplicationManager.h"
 
 #include "Delegates/Delegate.h"
+#include "Misc/Optional.h"
 #include "Templates/Function.h"
 #include "Templates/SharedPointer.h"
 
+namespace UE::ConcertSyncClient::Replication
+{
+	struct FChangeStreamRequest;
+	struct FAuthorityChangeRequest;
+}
+
 namespace UE::MultiUserClient
 {
 	class ISubmissionOperation;
 	struct FSubmitStreamChangesResponse;
 	struct FSubmitAuthorityChangesRequest;
 	struct FSubmitAuthorityChangesResponse;
+
+	/** At least StreamRequest or AuthorityRequest need to be valid to form a valid request */
+	struct FSubmissionParams
+	{
+		TOptional<ConcertSyncClient::Replication::FChangeStreamRequest> StreamRequest;
+		TOptional<ConcertSyncClient::Replication::FAuthorityChangeRequest> AuthorityRequest;
+	};
 	
 	/**
 	 * Manages the flow of changing stream and authority on the server.
@@ -38,7 +52,7 @@ namespace UE::MultiUserClient
 		 * @note The operation might start and instantly stop before SubmitChanges finishes (e.g. a network request fails to be created instantly).
 		 * @return The operation object if the operation was started
 		 */
-		virtual TSharedPtr<ISubmissionOperation> SubmitChanges() = 0;
+		virtual TSharedPtr<ISubmissionOperation> SubmitChanges(FSubmissionParams Params) = 0;
 
 		/** @return Detailed information about whether Submit can be called */
 		virtual EChangeUploadability GetUploadability() const = 0;
diff --git a/Engine/Plugins/Developer/Concert/ConcertApp/MultiUserClient/Source/MultiUserClient/Private/Replication/Submission/SubmissionWorkflow_LocalClient.cpp b/Engine/Plugins/Developer/Concert/ConcertApp/MultiUserClient/Source/MultiUserClient/Private/Replication/Submission/SubmissionWorkflow_LocalClient.cpp
--- a/Engine/Plugins/Developer/Concert/ConcertApp/MultiUserClient/Source/MultiUserClient/Private/Replication/Submission/SubmissionWorkflow_LocalClient.cpp
+++ b/Engine/Plugins/Developer/Concert/ConcertApp/MultiUserClient/Source/MultiUserClient/Private/Replication/Submission/SubmissionWorkflow_LocalClient.cpp
@@ -3,84 +3,59 @@
 #include "SubmissionWorkflow_LocalClient.h"
 
 #include "IConcertSyncClient.h"
-#include "Replication/Authority/AuthorityChangeTracker.h"
-#include "Replication/Stream/StreamChangeTracker.h"
-#include "Replication/Util/GlobalAuthorityCache.h"
-#include "Replication/Util/StreamRequestUtils.h"
+#include "Replication/IConcertClientReplicationManager.h"
 
 namespace UE::MultiUserClient
 {
-	FSubmissionWorkflow_LocalClient::FSubmissionWorkflow_LocalClient(
-		TSharedRef<IConcertSyncClient> InClient,
-		FStreamChangeTracker& InStreamChangeTracker,
-		FAuthorityChangeTracker& InAuthorityChangeTracker,
-		IClientStreamSynchronizer& InStreamSynchronizer,
-		const FGlobalAuthorityCache& InAuthorityCache
-		)
+	FSubmissionWorkflow_LocalClient::FSubmissionWorkflow_LocalClient(TSharedRef<IConcertSyncClient> InClient)
 		: Client(MoveTemp(InClient))
-		, AuthorityCache(InAuthorityCache)
-		, StreamChangeTracker(InStreamChangeTracker)
-		, AuthorityChangeTracker(InAuthorityChangeTracker)
-		, StreamSynchronizer(InStreamSynchronizer)
 	{}
 
 	EChangeUploadability FSubmissionWorkflow_LocalClient::GetUploadability() const
 	{
 		const bool bOperationInProgress = InProgressOperation.IsSet(); 
-		if (bOperationInProgress)
-		{
-			return EChangeUploadability::InProgress;
-		}
-
-		const bool bHasChanges = StreamChangeTracker.HasChanges() || AuthorityChangeTracker.HasChanges();
-		return bHasChanges ? EChangeUploadability::Ready : EChangeUploadability::NoChanges;
+		return bOperationInProgress ? EChangeUploadability::InProgress : EChangeUploadability::Ready;
 	}
 
-	TSharedPtr<ISubmissionOperation> FSubmissionWorkflow_LocalClient::SubmitChanges()
+	TSharedPtr<ISubmissionOperation> FSubmissionWorkflow_LocalClient::SubmitChanges(FSubmissionParams Params)
 	{
 		using namespace ConcertSyncClient::Replication;
+		
+		const TOptional<FChangeStreamRequest>& StreamRequest = Params.StreamRequest;
+		TOptional<FAuthorityChangeRequest>& AuthorityRequest = Params.AuthorityRequest;
+		const bool bIsStreamChangeEmpty = !StreamRequest.IsSet() || StreamRequest->IsEmpty();
+		const bool bIsAuthorityChangeEmpty = !AuthorityRequest.IsSet() || AuthorityRequest->IsEmpty();
+		
+		const bool bPointlessRequest = bIsStreamChangeEmpty && bIsAuthorityChangeEmpty;
 		IConcertClientReplicationManager* ReplicationManager = Client->GetReplicationManager();
-		if (!CanSubmit() || !ensure(ReplicationManager))
+		if (bPointlessRequest || !CanSubmit() || !ensure(ReplicationManager))
 		{
 			return nullptr;
 		}
 		
-		// The authority request is pre-built now to avoid sending changes the local client makes while we're waiting for the latent server responses
-		FAuthorityChangeRequest AuthorityChangeRequest = AuthorityChangeTracker.BuildChangeRequest(GetLocalClientStreamId());
-		
-		const FStreamChangelist& Changelist = StreamChangeTracker.GetCachedDeltaChange();
-		const bool bIsChangelistEmpty = Changelist.ObjectsToPut.IsEmpty() && Changelist.ObjectsToRemove.IsEmpty();
-		const bool bModifyStreams = !bIsChangelistEmpty;
-		
+		const bool bModifyStreams = !bIsStreamChangeEmpty;
 		const TSharedRef<FSingleClientSubmissionOperation> Operation = MakeShared<FSingleClientSubmissionOperation>(bModifyStreams);
 		InProgressOperation.Emplace(Operation);
 		
-		if (bIsChangelistEmpty)
+		if (bIsStreamChangeEmpty)
 		{
 			const FSubmitStreamChangesResponse CompletedChange { EStreamSubmissionErrorCode::NoChange };
 			Operation->EmplaceStreamPromise(CompletedChange);
 			StreamRequestCompletedDelegate.Broadcast(CompletedChange);
 			
-			SendAuthorityChangeRequest(MoveTemp(AuthorityChangeRequest));
+			SendAuthorityChangeRequest(MoveTemp(*AuthorityRequest));
 		}
 		else
 		{
-			FChangeStreamRequest StreamRequest = StreamSynchronizer.GetServerState().ReplicatedObjects.IsEmpty()
-				? StreamRequestUtils::BuildChangeRequest_CreateNewStream(GetLocalClientStreamId(), Changelist)
-				: StreamRequestUtils::BuildChangeRequest_UpdateExistingStream(Changelist);
-			
-			// Predict conflicts in case the remote client's streams have changed. Also: while the UI highlights "bad" requests, it does not correct it.
-			AuthorityCache.CleanseConflictsFromStreamRequest(StreamRequest, GetLocalClientId());
-		
-			ReplicationManager->ChangeStream(StreamRequest)
-				.Next([this, DestructionDetection = LifetimeToken->AsWeak(), StreamRequest, AuthorityChangeRequest = MoveTemp(AuthorityChangeRequest)](FChangeStreamResponse&& Response)
+			ReplicationManager->ChangeStream(*StreamRequest)
+				.Next([this, DestructionDetection = LifetimeToken->AsWeak(), StreamRequest, AuthorityRequest = MoveTemp(*AuthorityRequest)](FChangeStreamResponse&& Response)
 				{
-					const FSubmitStreamChangesResponse SubmissionResult { EStreamSubmissionErrorCode::Success, { FCompletedChangeSubmission{StreamRequest, Response } } };
+					const FSubmitStreamChangesResponse SubmissionResult { EStreamSubmissionErrorCode::Success, { FCompletedChangeSubmission{*StreamRequest, Response } } };
 					// The request might execute after we're destroyed, e.g. by leaving session while request is on the way.
 					// In that case, the Concert session triggers the OnSessionConnectionChanged which destroys us. Only after that, all the requests are timed out.
 					if (DestructionDetection.IsValid())
 					{
-						OnStreamChangeCompleted(StreamRequest, Response, AuthorityChangeRequest);
+						OnStreamChangeCompleted(*StreamRequest, Response, AuthorityRequest);
 					}
 
 					return SubmissionResult;
@@ -148,9 +123,6 @@ namespace UE::MultiUserClient
 			AuthorityRequestCompletedDelegate.Broadcast(Request, Response);
 			return;
 		}
-
-		// Predict conflicts in case the remote client's streams have changed. Also: while the UI highlights "bad" requests, it does not correct it.
-		AuthorityCache.CleanseConflictsFromAuthorityRequest(AuthorityChangeRequest, GetLocalClientId());
 		
 		FSubmitAuthorityChangesRequest Request{ EAuthoritySubmissionRequestErrorCode::Success, AuthorityChangeRequest };
 		Operation->EmplaceAuthorityRequestPromise(Request);
@@ -171,9 +143,4 @@ namespace UE::MultiUserClient
 				}
 			});
 	}
-	
-	FGuid FSubmissionWorkflow_LocalClient::GetLocalClientId() const
-	{
-		return Client->GetConcertClient()->GetCurrentSession()->GetSessionClientEndpointId();
-	}
 }
diff --git a/Engine/Plugins/Developer/Concert/ConcertApp/MultiUserClient/Source/MultiUserClient/Private/Replication/Submission/SubmissionWorkflow_LocalClient.h b/Engine/Plugins/Developer/Concert/ConcertApp/MultiUserClient/Source/MultiUserClient/Private/Replication/Submission/SubmissionWorkflow_LocalClient.h
--- a/Engine/Plugins/Developer/Concert/ConcertApp/MultiUserClient/Source/MultiUserClient/Private/Replication/Submission/SubmissionWorkflow_LocalClient.h
+++ b/Engine/Plugins/Developer/Concert/ConcertApp/MultiUserClient/Source/MultiUserClient/Private/Replication/Submission/SubmissionWorkflow_LocalClient.h
@@ -25,16 +25,10 @@ namespace UE::MultiUserClient
 	{
 	public:
 		
-		FSubmissionWorkflow_LocalClient(
-			TSharedRef<IConcertSyncClient> InClient,
-			FStreamChangeTracker& InStreamChangeTracker,
-			FAuthorityChangeTracker& InAuthorityChangeTracker,
-			IClientStreamSynchronizer& InStreamSynchronizer,
-			const FGlobalAuthorityCache& InAuthorityCache
-			);
+		FSubmissionWorkflow_LocalClient(TSharedRef<IConcertSyncClient> InClient);
 		
 		//~ Begin ISubmissionWorkflow Interface
-		virtual TSharedPtr<ISubmissionOperation> SubmitChanges() override;
+		virtual TSharedPtr<ISubmissionOperation> SubmitChanges(FSubmissionParams Params) override;
 		virtual EChangeUploadability GetUploadability() const override;
 		//~ End ISubmissionWorkflow Interface
 	
@@ -46,25 +40,12 @@ namespace UE::MultiUserClient
 		/** Used to send authority requests to the server. */
 		const TSharedRef<IConcertSyncClient> Client;
 
-		/** Used to filter out changes that would cause a conflict when submitted. */
-		const FGlobalAuthorityCache& AuthorityCache;
-
-		/** Used to get changes made to the stream */
-		FStreamChangeTracker& StreamChangeTracker;
-		/** Used to get changes made to the authority */
-		FAuthorityChangeTracker& AuthorityChangeTracker;
-
-		/** Used to change the streams */
-		IClientStreamSynchronizer& StreamSynchronizer;
-
 		/**
 		 * Set for as long as there is a SubmitChanges operation in progress.
 		 * Automatically cancels pending promises when destroyed.
 		 */
 		TOptional<TSharedRef<FSingleClientSubmissionOperation>> InProgressOperation;
 
-		FGuid GetLocalClientStreamId() const { return StreamSynchronizer.GetStreamId(); }
-
 		/** Advances the request by requesting authority. */
 		void OnStreamChangeCompleted(
 			const ConcertSyncClient::Replication::FChangeStreamRequest& StreamChangeRequest,
@@ -72,8 +53,6 @@ namespace UE::MultiUserClient
 			ConcertSyncClient::Replication::FAuthorityChangeRequest AuthorityChangeRequest
 			);
 		void SendAuthorityChangeRequest(ConcertSyncClient::Replication::FAuthorityChangeRequest AuthorityChangeRequest);
-		
-		FGuid GetLocalClientId() const;
 	};
 }
 
diff --git a/Engine/Plugins/Developer/Concert/ConcertApp/MultiUserClient/Source/MultiUserClient/Private/Replication/Submission/SubmissionWorkflow_RemoteClient.cpp b/Engine/Plugins/Developer/Concert/ConcertApp/MultiUserClient/Source/MultiUserClient/Private/Replication/Submission/SubmissionWorkflow_RemoteClient.cpp
--- a/Engine/Plugins/Developer/Concert/ConcertApp/MultiUserClient/Source/MultiUserClient/Private/Replication/Submission/SubmissionWorkflow_RemoteClient.cpp
+++ b/Engine/Plugins/Developer/Concert/ConcertApp/MultiUserClient/Source/MultiUserClient/Private/Replication/Submission/SubmissionWorkflow_RemoteClient.cpp
@@ -4,7 +4,7 @@
 
 namespace UE::MultiUserClient
 {
-	TSharedPtr<ISubmissionOperation> FSubmissionWorkflow_RemoteClient::SubmitChanges()
+	TSharedPtr<ISubmissionOperation> FSubmissionWorkflow_RemoteClient::SubmitChanges(FSubmissionParams Params)
 	{
 		// TODO DP UE-180657: Changing remote client's stream is not implemented for now
 		// TODO DP UE-198088 changing remote client authority
diff --git a/Engine/Plugins/Developer/Concert/ConcertApp/MultiUserClient/Source/MultiUserClient/Private/Replication/Submission/SubmissionWorkflow_RemoteClient.h b/Engine/Plugins/Developer/Concert/ConcertApp/MultiUserClient/Source/MultiUserClient/Private/Replication/Submission/SubmissionWorkflow_RemoteClient.h
--- a/Engine/Plugins/Developer/Concert/ConcertApp/MultiUserClient/Source/MultiUserClient/Private/Replication/Submission/SubmissionWorkflow_RemoteClient.h
+++ b/Engine/Plugins/Developer/Concert/ConcertApp/MultiUserClient/Source/MultiUserClient/Private/Replication/Submission/SubmissionWorkflow_RemoteClient.h
@@ -17,7 +17,7 @@ namespace UE::MultiUserClient
 	public:
 
 		//~ Begin ISubmissionWorkflow Interface
-		virtual TSharedPtr<ISubmissionOperation> SubmitChanges() override;
+		virtual TSharedPtr<ISubmissionOperation> SubmitChanges(FSubmissionParams Params) override;
 		virtual EChangeUploadability GetUploadability() const override;
 		//~ End ISubmissionWorkflow Interface
 	};
diff --git a/Engine/Plugins/Developer/Concert/ConcertSync/ConcertSyncCore/Source/ConcertSyncCore/Public/Replication/Messages/ChangeAuthority.h b/Engine/Plugins/Developer/Concert/ConcertSync/ConcertSyncCore/Source/ConcertSyncCore/Public/Replication/Messages/ChangeAuthority.h
--- a/Engine/Plugins/Developer/Concert/ConcertSync/ConcertSyncCore/Source/ConcertSyncCore/Public/Replication/Messages/ChangeAuthority.h
+++ b/Engine/Plugins/Developer/Concert/ConcertSync/ConcertSyncCore/Source/ConcertSyncCore/Public/Replication/Messages/ChangeAuthority.h
@@ -44,6 +44,8 @@ struct FConcertReplication_ChangeAuthority_Request
 	 */
 	UPROPERTY()
 	TMap<FSoftObjectPath, FConcertStreamArray> ReleaseAuthority;
+
+	bool IsEmpty() const { return TakeAuthority.IsEmpty() && ReleaseAuthority.IsEmpty(); }
 };
 
 USTRUCT()
diff --git a/Engine/Plugins/Developer/Concert/ConcertSync/ConcertSyncCore/Source/ConcertSyncCore/Public/Replication/Messages/ChangeStream.h b/Engine/Plugins/Developer/Concert/ConcertSync/ConcertSyncCore/Source/ConcertSyncCore/Public/Replication/Messages/ChangeStream.h
--- a/Engine/Plugins/Developer/Concert/ConcertSync/ConcertSyncCore/Source/ConcertSyncCore/Public/Replication/Messages/ChangeStream.h
+++ b/Engine/Plugins/Developer/Concert/ConcertSync/ConcertSyncCore/Source/ConcertSyncCore/Public/Replication/Messages/ChangeStream.h
@@ -118,6 +118,8 @@ struct FConcertReplication_ChangeStream_Request
 	UPROPERTY()
 	TSet<FGuid> StreamsToRemove;
 
+	bool IsEmpty() const { return ObjectsToRemove.IsEmpty() && ObjectsToPut.IsEmpty() && StreamsToAdd.IsEmpty() && StreamsToRemove.IsEmpty(); }
+
 	friend bool operator==(const FConcertReplication_ChangeStream_Request& Left, const FConcertReplication_ChangeStream_Request& Right)
 	{
 		const auto OrderIndependentEquals = [](const auto& Left, const auto& Right){ return Left.Num() == Right.Num() && Left.Includes(Right); };
