Fix bugs:
- crash when checking and then unchecking property
- no stream change request being sent when unchecking property while you have authority

diff --git a/Engine/Plugins/Developer/Concert/ConcertApp/MultiUserClient/Source/MultiUserClient/Private/Replication/Client/ReplicationClient.cpp b/Engine/Plugins/Developer/Concert/ConcertApp/MultiUserClient/Source/MultiUserClient/Private/Replication/Client/ReplicationClient.cpp
--- a/Engine/Plugins/Developer/Concert/ConcertApp/MultiUserClient/Source/MultiUserClient/Private/Replication/Client/ReplicationClient.cpp
+++ b/Engine/Plugins/Developer/Concert/ConcertApp/MultiUserClient/Source/MultiUserClient/Private/Replication/Client/ReplicationClient.cpp
@@ -95,8 +95,19 @@ namespace UE::MultiUserClient
 		LocalClientStreamDiffer.RefreshChangesCache();
 		
 		// Better UX for user: automatically take authority for newly added objects (but only if it is allowed and causes no conflicts)
+		TakeAuthorityOverNewlyAddedObjects(ChangeData);
+		// Refresh because local authority changes may no longer be valid after modifying the stream
+		LocalAuthorityDiffer.RefreshChanges();
+
+		// Finally, let everybody else know.
+		OnModelChangedDelegate.Broadcast();
+		AutoSubmissionPolicy.ProcessAccumulatedChangesAndSubmit();
+	}
+	
+	void FReplicationClient::TakeAuthorityOverNewlyAddedObjects(const FDeferredOnModelChangedData& ChangeData)
+	{
 		TArray<FSoftObjectPath> ObjectPaths;
-		Algo::TransformIf(ChangeData.AccumulatedAddedObjects, ObjectPaths, 
+		Algo::TransformIf(ChangeData.AccumulatedAddedObjects, ObjectPaths,
 			[](const TWeakObjectPtr<UObject>& Object)
 			{
 				// The object might have been made invalid last frame.
@@ -107,12 +118,6 @@ namespace UE::MultiUserClient
 				return FSoftObjectPath(Object.Get()) ;
 			});
 		LocalAuthorityDiffer.SetAuthorityIfAllowed(ObjectPaths, true);
-
-		// Refresh because local authority changes may no longer be valid after modifying the stream
-		LocalAuthorityDiffer.RefreshChanges();
-
-		// Finally, let everybody else know.
-		OnModelChangedDelegate.Broadcast();
 	}
 
 	void FReplicationClient::OnAuthoritySubmissionCompleted(const FSubmitAuthorityChangesRequest& Request, const FSubmitAuthorityChangesResponse& Response)
diff --git a/Engine/Plugins/Developer/Concert/ConcertApp/MultiUserClient/Source/MultiUserClient/Private/Replication/Client/ReplicationClient.h b/Engine/Plugins/Developer/Concert/ConcertApp/MultiUserClient/Source/MultiUserClient/Private/Replication/Client/ReplicationClient.h
--- a/Engine/Plugins/Developer/Concert/ConcertApp/MultiUserClient/Source/MultiUserClient/Private/Replication/Client/ReplicationClient.h
+++ b/Engine/Plugins/Developer/Concert/ConcertApp/MultiUserClient/Source/MultiUserClient/Private/Replication/Client/ReplicationClient.h
@@ -133,6 +133,9 @@ namespace UE::MultiUserClient
 		/** Processes all changes that have happened to the stream this frame. */
 		void ProcessOnModelChanged();
 		
+		/** Takes authority over newly added objects for better UX */
+		void TakeAuthorityOverNewlyAddedObjects(const FDeferredOnModelChangedData& ChangeData);
+		
 		/** Removes authority if request fails */
 		void OnAuthoritySubmissionCompleted(const FSubmitAuthorityChangesRequest& Request, const FSubmitAuthorityChangesResponse& Response);
 	};
diff --git a/Engine/Plugins/Developer/Concert/ConcertApp/MultiUserClient/Source/MultiUserClient/Private/Replication/Submission/AutoSubmissionPolicy.cpp b/Engine/Plugins/Developer/Concert/ConcertApp/MultiUserClient/Source/MultiUserClient/Private/Replication/Submission/AutoSubmissionPolicy.cpp
--- a/Engine/Plugins/Developer/Concert/ConcertApp/MultiUserClient/Source/MultiUserClient/Private/Replication/Submission/AutoSubmissionPolicy.cpp
+++ b/Engine/Plugins/Developer/Concert/ConcertApp/MultiUserClient/Source/MultiUserClient/Private/Replication/Submission/AutoSubmissionPolicy.cpp
@@ -30,38 +30,50 @@ namespace UE::MultiUserClient
 	FAutoSubmissionPolicy::~FAutoSubmissionPolicy()
 	{
 		FCoreDelegates::OnEndFrame.RemoveAll(this);
+		
 		StreamEditorModel.OnObjectsChanged().RemoveAll(this);
 		StreamEditorModel.OnPropertiesChanged().RemoveAll(this);
 		AuthorityChangeTracker.OnAddedOwnedObjects().RemoveAll(this);
+		
+		SubmissionWorkflow.OnStreamRequestCompleted().RemoveAll(this);
 	}
 
-	void FAutoSubmissionPolicy::OnChangesDetected()
+	void FAutoSubmissionPolicy::ProcessAccumulatedChangesAndSubmit()
 	{
-		if (!FCoreDelegates::OnEndFrame.IsBoundToObject(this))
+		if (!bIsDirty)
+		{
+			return;
+		}
+		
+		if (SubmissionWorkflow.CanSubmit())
 		{
-			FCoreDelegates::OnEndFrame.AddRaw(this, &FAutoSubmissionPolicy::OnEndFrame);
+			SubmitChanges();
+		}
+		else if (!SubmissionWorkflow.OnSubmitOperationCompleted().IsBoundToObject(this))
+		{
+			SubmissionWorkflow.OnSubmitOperationCompleted().AddRaw(this, &FAutoSubmissionPolicy::ProcessAccumulatedChangesAndSubmit);
 		}
 	}
 
-	void FAutoSubmissionPolicy::OnEndFrame()
+	void FAutoSubmissionPolicy::SubmitChanges()
 	{
-		// Instead of ticking until CanSubmit return true, consider introducing a queueing system for tasks...
-		// but that only makes sense when there multiple code places that want to call SubmitChanges
-		if (SubmissionWorkflow.CanSubmit())
-		{
-			using namespace UE::ConcertSyncClient::Replication;
-			FCoreDelegates::OnEndFrame.RemoveAll(this);
+		using namespace UE::ConcertSyncClient::Replication;
+		FCoreDelegates::OnEndFrame.RemoveAll(this);
 			
-			// Even though authority request is sent after server confirms stream change, the authority request is pre-built to avoid sending changes
-			// the local client makes while we're waiting for the latent server responses.
-			TOptional<FAuthorityChangeRequest> AuthorityChangeRequest = RequestBuilder.BuildAuthorityChange();
-			TOptional<FChangeStreamRequest> StreamRequest = RequestBuilder.BuildStreamChange();
+		// Even though authority request is sent after server confirms stream change, the authority request is pre-built to avoid sending changes
+		// the local client makes while we're waiting for the latent server responses.
+		TOptional<FAuthorityChangeRequest> AuthorityChangeRequest = RequestBuilder.BuildAuthorityChange();
+		TOptional<FChangeStreamRequest> StreamRequest = RequestBuilder.BuildStreamChange();
 			
-			// Should be true since a change caused OnEndFrame to be registered in the first place, but in THEORY changes could have been reverted in the same frame (e.g. via a transaction)
-			if (AuthorityChangeRequest || StreamRequest)
-			{
-				SubmissionWorkflow.SubmitChanges({ StreamRequest, AuthorityChangeRequest });
-			}
+		// Should be true since a change caused OnEndFrame to be registered in the first place, but in THEORY changes could have been reverted in the same frame (e.g. via a transaction)
+		if (AuthorityChangeRequest || StreamRequest)
+		{
+			SubmissionWorkflow.SubmitChanges({ StreamRequest, AuthorityChangeRequest });
 		}
 	}
+
+	void FAutoSubmissionPolicy::OnChangesDetected()
+	{
+		bIsDirty = true;
+	}
 }
diff --git a/Engine/Plugins/Developer/Concert/ConcertApp/MultiUserClient/Source/MultiUserClient/Private/Replication/Submission/AutoSubmissionPolicy.h b/Engine/Plugins/Developer/Concert/ConcertApp/MultiUserClient/Source/MultiUserClient/Private/Replication/Submission/AutoSubmissionPolicy.h
--- a/Engine/Plugins/Developer/Concert/ConcertApp/MultiUserClient/Source/MultiUserClient/Private/Replication/Submission/AutoSubmissionPolicy.h
+++ b/Engine/Plugins/Developer/Concert/ConcertApp/MultiUserClient/Source/MultiUserClient/Private/Replication/Submission/AutoSubmissionPolicy.h
@@ -40,6 +40,9 @@ namespace UE::MultiUserClient
 			FAuthorityChangeTracker& InAuthorityChangeTracker
 			);
 		~FAutoSubmissionPolicy();
+
+		/** Checks whether any changes have been made since the last call and submits a change request if so. */
+		void ProcessAccumulatedChangesAndSubmit();
 		
 	private:
 
@@ -52,12 +55,14 @@ namespace UE::MultiUserClient
 		ConcertClientSharedSlate::IEditableReplicationStreamModel& StreamEditorModel;
 		/** Informs us when authority is changed by the user. */
 		FAuthorityChangeTracker& AuthorityChangeTracker;
+
+		/** Whether any changes were made. */
+		bool bIsDirty = false;
+		
+		void SubmitChanges();
 		
 		void OnObjectsChanged(TArrayView<UObject* const>, TArrayView<const FSoftObjectPath>, ConcertClientSharedSlate::EReplicatedObjectChangeReason) { OnChangesDetected(); }
 		void OnChangesDetected();
-
-		/** Checks whether any changes were made last frame and submits a change request if so. */
-		void OnEndFrame();
 	};
 }
 
diff --git a/Engine/Plugins/Developer/Concert/ConcertApp/MultiUserClient/Source/MultiUserClient/Private/Replication/Submission/ISubmissionWorkflow.h b/Engine/Plugins/Developer/Concert/ConcertApp/MultiUserClient/Source/MultiUserClient/Private/Replication/Submission/ISubmissionWorkflow.h
--- a/Engine/Plugins/Developer/Concert/ConcertApp/MultiUserClient/Source/MultiUserClient/Private/Replication/Submission/ISubmissionWorkflow.h
+++ b/Engine/Plugins/Developer/Concert/ConcertApp/MultiUserClient/Source/MultiUserClient/Private/Replication/Submission/ISubmissionWorkflow.h
@@ -7,7 +7,6 @@
 
 #include "Delegates/Delegate.h"
 #include "Misc/Optional.h"
-#include "Templates/Function.h"
 #include "Templates/SharedPointer.h"
 
 namespace UE::ConcertSyncClient::Replication
@@ -28,6 +27,10 @@ namespace UE::MultiUserClient
 	{
 		TOptional<ConcertSyncClient::Replication::FChangeStreamRequest> StreamRequest;
 		TOptional<ConcertSyncClient::Replication::FAuthorityChangeRequest> AuthorityRequest;
+		
+		bool IsStreamChangeEmpty() const { return !StreamRequest.IsSet() || StreamRequest->IsEmpty(); }
+		bool IsAuthorityChangeEmpty() const { return !AuthorityRequest.IsSet() || AuthorityRequest->IsEmpty(); }
+		bool IsEmpty() const { return IsStreamChangeEmpty() && IsAuthorityChangeEmpty(); }
 	};
 	
 	/**
@@ -60,18 +63,22 @@ namespace UE::MultiUserClient
 		
 		DECLARE_MULTICAST_DELEGATE_OneParam(FOnStreamRequestCompleted, const FSubmitStreamChangesResponse&);
 		/**
-		 * Called whenever a submit operation completes the stream change request stage.
+		 * Broadcasts whenever a submit operation completes the stream change request stage. Useful for accumulating and counting errors.
 		 * @note No stream changes may have been requested. Check the error code.
 		 */
 		virtual FOnStreamRequestCompleted& OnStreamRequestCompleted() = 0;
 
 		DECLARE_MULTICAST_DELEGATE_TwoParams(FOnAuthorityRequestCompleted, const FSubmitAuthorityChangesRequest&, const FSubmitAuthorityChangesResponse&);
 		/**
-		 * Called whenever a submit operation completes the authority change request stage.
+		 * Broadcasts whenever a submit operation completes the authority change request stage. Useful for accumulating and counting errors.
 		 * @note No authority changes may have been requested. Check the error code.
 		 */
 		virtual FOnAuthorityRequestCompleted& OnAuthorityRequestCompleted() = 0;
 
+		DECLARE_MULTICAST_DELEGATE(FOnSubmitOperationCompleted);
+		/** Broadcasts whenever a submit operation completes. Useful for enqueuing submit operations. */
+		virtual FOnSubmitOperationCompleted& OnSubmitOperationCompleted() = 0;
+
 		virtual ~ISubmissionWorkflow() = default;
 	};
 
@@ -83,11 +90,13 @@ namespace UE::MultiUserClient
 		//~ Begin ISubmissionWorkflow Interface
 		virtual FOnStreamRequestCompleted& OnStreamRequestCompleted() override { return StreamRequestCompletedDelegate; }
 		virtual FOnAuthorityRequestCompleted& OnAuthorityRequestCompleted() override { return AuthorityRequestCompletedDelegate; }
+		virtual FOnSubmitOperationCompleted& OnSubmitOperationCompleted() override { return OnSubmitOperationCompletedDelegate; }
 		//~ End ISubmissionWorkflow Interface
 
 	protected:
 		
 		FOnStreamRequestCompleted StreamRequestCompletedDelegate;
 		FOnAuthorityRequestCompleted AuthorityRequestCompletedDelegate;
+		FOnSubmitOperationCompleted OnSubmitOperationCompletedDelegate;
 	};
 }
diff --git a/Engine/Plugins/Developer/Concert/ConcertApp/MultiUserClient/Source/MultiUserClient/Private/Replication/Submission/SubmissionWorkflow_LocalClient.cpp b/Engine/Plugins/Developer/Concert/ConcertApp/MultiUserClient/Source/MultiUserClient/Private/Replication/Submission/SubmissionWorkflow_LocalClient.cpp
--- a/Engine/Plugins/Developer/Concert/ConcertApp/MultiUserClient/Source/MultiUserClient/Private/Replication/Submission/SubmissionWorkflow_LocalClient.cpp
+++ b/Engine/Plugins/Developer/Concert/ConcertApp/MultiUserClient/Source/MultiUserClient/Private/Replication/Submission/SubmissionWorkflow_LocalClient.cpp
@@ -21,21 +21,20 @@ namespace UE::MultiUserClient
 	{
 		using namespace ConcertSyncClient::Replication;
 		
-		const TOptional<FChangeStreamRequest>& StreamRequest = Params.StreamRequest;
-		TOptional<FAuthorityChangeRequest>& AuthorityRequest = Params.AuthorityRequest;
-		const bool bIsStreamChangeEmpty = !StreamRequest.IsSet() || StreamRequest->IsEmpty();
-		const bool bIsAuthorityChangeEmpty = !AuthorityRequest.IsSet() || AuthorityRequest->IsEmpty();
+		const TOptional<FChangeStreamRequest>& OptionalStreamRequest = Params.StreamRequest;
+		TOptional<FAuthorityChangeRequest>& OptionalAuthorityRequest = Params.AuthorityRequest;
 		
-		const bool bPointlessRequest = bIsStreamChangeEmpty && bIsAuthorityChangeEmpty;
 		IConcertClientReplicationManager* ReplicationManager = Client->GetReplicationManager();
-		if (bPointlessRequest || !CanSubmit() || !ensure(ReplicationManager))
+		if (Params.IsEmpty() || !CanSubmit() || !ensure(ReplicationManager))
 		{
 			return nullptr;
 		}
 		
+		const bool bIsStreamChangeEmpty = Params.IsStreamChangeEmpty();
 		const bool bModifyStreams = !bIsStreamChangeEmpty;
 		const TSharedRef<FSingleClientSubmissionOperation> Operation = MakeShared<FSingleClientSubmissionOperation>(bModifyStreams);
-		InProgressOperation.Emplace(Operation);
+		InProgressOperation = { Operation };
+		InProgressOperation->OnDestroy.BindLambda([this](){ OnSubmitOperationCompletedDelegate.Broadcast(); });
 		
 		if (bIsStreamChangeEmpty)
 		{
@@ -43,19 +42,19 @@ namespace UE::MultiUserClient
 			Operation->EmplaceStreamPromise(CompletedChange);
 			StreamRequestCompletedDelegate.Broadcast(CompletedChange);
 			
-			SendAuthorityChangeRequest(MoveTemp(*AuthorityRequest));
+			HandlePendingAuthorityChangeRequest(MoveTemp(*OptionalAuthorityRequest));
 		}
 		else
 		{
-			ReplicationManager->ChangeStream(*StreamRequest)
-				.Next([this, DestructionDetection = LifetimeToken->AsWeak(), StreamRequest, AuthorityRequest = MoveTemp(*AuthorityRequest)](FChangeStreamResponse&& Response)
+			ReplicationManager->ChangeStream(*OptionalStreamRequest)
+				.Next([this, DestructionDetection = LifetimeToken->AsWeak(), OptionalStreamRequest, AuthorityRequest = MoveTemp(OptionalAuthorityRequest)](FChangeStreamResponse&& Response) mutable
 				{
-					const FSubmitStreamChangesResponse SubmissionResult { EStreamSubmissionErrorCode::Success, { FCompletedChangeSubmission{*StreamRequest, Response } } };
+					const FSubmitStreamChangesResponse SubmissionResult { EStreamSubmissionErrorCode::Success, { FCompletedChangeSubmission{*OptionalStreamRequest, Response } } };
 					// The request might execute after we're destroyed, e.g. by leaving session while request is on the way.
 					// In that case, the Concert session triggers the OnSessionConnectionChanged which destroys us. Only after that, all the requests are timed out.
 					if (DestructionDetection.IsValid())
 					{
-						OnStreamChangeCompleted(*StreamRequest, Response, AuthorityRequest);
+						OnStreamChangeCompleted(*OptionalStreamRequest, Response, MoveTemp(AuthorityRequest));
 					}
 
 					return SubmissionResult;
@@ -69,10 +68,10 @@ namespace UE::MultiUserClient
 	void FSubmissionWorkflow_LocalClient::OnStreamChangeCompleted(
 		const ConcertSyncClient::Replication::FChangeStreamRequest& StreamChangeRequest,
 		const ConcertSyncClient::Replication::FChangeStreamResponse& ChangeStreamResponse,
-		ConcertSyncClient::Replication::FAuthorityChangeRequest AuthorityChangeRequest
+		TOptional<ConcertSyncClient::Replication::FAuthorityChangeRequest> AuthorityChangeRequest
 		)
 	{
-		const TSharedRef<FSingleClientSubmissionOperation> Operation = *InProgressOperation;
+		const TSharedRef<FSingleClientSubmissionOperation> Operation = InProgressOperation->Operation;
 		
 		const EStreamSubmissionErrorCode ErrorCode = ChangeStreamResponse.ErrorCode == EReplicationResponseErrorCode::Handled
 			? EStreamSubmissionErrorCode::Success
@@ -83,7 +82,7 @@ namespace UE::MultiUserClient
 		
 		if (ChangeStreamResponse.IsSuccess())
 		{
-			SendAuthorityChangeRequest(MoveTemp(AuthorityChangeRequest));
+			HandlePendingAuthorityChangeRequest(MoveTemp(AuthorityChangeRequest));
 		}
 		else
 		{
@@ -98,7 +97,7 @@ namespace UE::MultiUserClient
 		}
 	}
 
-	void FSubmissionWorkflow_LocalClient::SendAuthorityChangeRequest(ConcertSyncClient::Replication::FAuthorityChangeRequest AuthorityChangeRequest)
+	void FSubmissionWorkflow_LocalClient::HandlePendingAuthorityChangeRequest(TOptional<ConcertSyncClient::Replication::FAuthorityChangeRequest> AuthorityChangeRequest)
 	{
 		using namespace ConcertSyncClient::Replication;
 		IConcertClientReplicationManager* ReplicationManager = Client->GetReplicationManager();
@@ -109,8 +108,8 @@ namespace UE::MultiUserClient
 			return;
 		}
 
-		const TSharedRef<FSingleClientSubmissionOperation> Operation = *InProgressOperation;
-		const bool bHasNoChanges = AuthorityChangeRequest.ReleaseAuthority.IsEmpty() && AuthorityChangeRequest.TakeAuthority.IsEmpty();
+		const TSharedRef<FSingleClientSubmissionOperation> Operation = InProgressOperation->Operation;
+		const bool bHasNoChanges = !AuthorityChangeRequest || AuthorityChangeRequest->IsEmpty();
 		if (bHasNoChanges)
 		{
 			const FSubmitAuthorityChangesRequest Request{ EAuthoritySubmissionRequestErrorCode::NoChange };
@@ -126,7 +125,7 @@ namespace UE::MultiUserClient
 		
 		FSubmitAuthorityChangesRequest Request{ EAuthoritySubmissionRequestErrorCode::Success, AuthorityChangeRequest };
 		Operation->EmplaceAuthorityRequestPromise(Request);
-		ReplicationManager->RequestAuthorityChange(MoveTemp(AuthorityChangeRequest))
+		ReplicationManager->RequestAuthorityChange(MoveTemp(*AuthorityChangeRequest))
 			.Next([this, Request = MoveTemp(Request), DestructionDetection = LifetimeToken->AsWeak()](FAuthorityChangeResponse&& Response)
 			{
 				if (DestructionDetection.IsValid())
@@ -136,7 +135,7 @@ namespace UE::MultiUserClient
 						: EAuthoritySubmissionResponseErrorCode::Timeout;
 					const FSubmitAuthorityChangesResponse Result { ErrorCode, MoveTemp(Response) };
 					
-					InProgressOperation->Get().EmplaceAuthorityResponsePromise(Result);
+					InProgressOperation->Operation->EmplaceAuthorityResponsePromise(Result);
 					InProgressOperation.Reset();
 					
 					AuthorityRequestCompletedDelegate.Broadcast(Request, Result);
diff --git a/Engine/Plugins/Developer/Concert/ConcertApp/MultiUserClient/Source/MultiUserClient/Private/Replication/Submission/SubmissionWorkflow_LocalClient.h b/Engine/Plugins/Developer/Concert/ConcertApp/MultiUserClient/Source/MultiUserClient/Private/Replication/Submission/SubmissionWorkflow_LocalClient.h
--- a/Engine/Plugins/Developer/Concert/ConcertApp/MultiUserClient/Source/MultiUserClient/Private/Replication/Submission/SubmissionWorkflow_LocalClient.h
+++ b/Engine/Plugins/Developer/Concert/ConcertApp/MultiUserClient/Source/MultiUserClient/Private/Replication/Submission/SubmissionWorkflow_LocalClient.h
@@ -40,19 +40,32 @@ namespace UE::MultiUserClient
 		/** Used to send authority requests to the server. */
 		const TSharedRef<IConcertSyncClient> Client;
 
+		struct FOperationData
+		{
+			TSharedRef<FSingleClientSubmissionOperation> Operation;
+
+			DECLARE_DELEGATE(FOnDestroy);
+			FOnDestroy OnDestroy;
+			
+			~FOperationData()
+			{
+				OnDestroy.ExecuteIfBound();
+			}
+		};
+		
 		/**
 		 * Set for as long as there is a SubmitChanges operation in progress.
 		 * Automatically cancels pending promises when destroyed.
 		 */
-		TOptional<TSharedRef<FSingleClientSubmissionOperation>> InProgressOperation;
+		TOptional<FOperationData> InProgressOperation;
 
 		/** Advances the request by requesting authority. */
 		void OnStreamChangeCompleted(
 			const ConcertSyncClient::Replication::FChangeStreamRequest& StreamChangeRequest,
 			const ConcertSyncClient::Replication::FChangeStreamResponse& ChangeStreamResponse,
-			ConcertSyncClient::Replication::FAuthorityChangeRequest AuthorityChangeRequest
+			TOptional<ConcertSyncClient::Replication::FAuthorityChangeRequest> AuthorityChangeRequest
 			);
-		void SendAuthorityChangeRequest(ConcertSyncClient::Replication::FAuthorityChangeRequest AuthorityChangeRequest);
+		void HandlePendingAuthorityChangeRequest(TOptional<ConcertSyncClient::Replication::FAuthorityChangeRequest> AuthorityChangeRequest);
 	};
 }
 
