Remotely change other clients

diff --git a/Engine/Plugins/Developer/Concert/ConcertApp/MultiUserClient/Source/MultiUserClient/Private/Replication/Authority/AuthorityChangeTracker.cpp b/Engine/Plugins/Developer/Concert/ConcertApp/MultiUserClient/Source/MultiUserClient/Private/Replication/Authority/AuthorityChangeTracker.cpp
--- a/Engine/Plugins/Developer/Concert/ConcertApp/MultiUserClient/Source/MultiUserClient/Private/Replication/Authority/AuthorityChangeTracker.cpp
+++ b/Engine/Plugins/Developer/Concert/ConcertApp/MultiUserClient/Source/MultiUserClient/Private/Replication/Authority/AuthorityChangeTracker.cpp
@@ -2,6 +2,7 @@
 
 #include "AuthorityChangeTracker.h"
 
+#include "EAuthorityMutability.h"
 #include "IClientAuthoritySynchronizer.h"
 #include "Replication/Util/GlobalAuthorityCache.h"
 
@@ -64,7 +65,7 @@ namespace UE::MultiUserClient
 
 			const bool bHasAuthority = AuthoritySynchronizer.HasAuthorityOver(ObjectPath); 
 			if (bHasAuthority == NewAuthorityState.Value
-				|| !AuthoritySynchronizer.CanChangeAuthority(ObjectPath))
+				|| !CanSetAuthorityFor(ObjectPath))
 			{
 				It.RemoveCurrent();
 			}
@@ -79,13 +80,12 @@ namespace UE::MultiUserClient
 
 	bool FAuthorityChangeTracker::CanSetAuthorityFor(const FSoftObjectPath& ObjectPath) const
 	{
-		return AuthoritySynchronizer.CanChangeAuthority(ObjectPath)
-			&& AuthorityCache.CanClientTakeAuthorityAfterSubmission(ObjectPath, ClientId) != FGlobalAuthorityCache::ECanTakeAuthority::Conflict;
+		return GetChangeAuthorityMutability(ObjectPath) == EAuthorityMutability::Allowed;
 	}
 
 	EAuthorityMutability FAuthorityChangeTracker::GetChangeAuthorityMutability(const FSoftObjectPath& ObjectPath) const
 	{
-		return AuthoritySynchronizer.GetChangeAuthorityMutability(ObjectPath);
+		return AuthorityCache.CanClientTakeAuthorityAfterSubmission(ObjectPath, ClientId);
 	}
 
 	TOptional<ConcertSyncClient::Replication::FAuthorityChangeRequest> FAuthorityChangeTracker::BuildChangeRequest(const FGuid& StreamId) const
diff --git a/Engine/Plugins/Developer/Concert/ConcertApp/MultiUserClient/Source/MultiUserClient/Private/Replication/Authority/AuthorityChangeTracker.h b/Engine/Plugins/Developer/Concert/ConcertApp/MultiUserClient/Source/MultiUserClient/Private/Replication/Authority/AuthorityChangeTracker.h
--- a/Engine/Plugins/Developer/Concert/ConcertApp/MultiUserClient/Source/MultiUserClient/Private/Replication/Authority/AuthorityChangeTracker.h
+++ b/Engine/Plugins/Developer/Concert/ConcertApp/MultiUserClient/Source/MultiUserClient/Private/Replication/Authority/AuthorityChangeTracker.h
@@ -16,7 +16,7 @@ struct FSoftObjectPath;
 namespace UE::MultiUserClient
 {
 	class IClientAuthoritySynchronizer;
-	enum class EAuthorityMutability;
+	enum class EAuthorityMutability : uint8;
 	
 	/** Keeps track of authority changes the local editor instance makes to a client. */
 	class FAuthorityChangeTracker
diff --git a/Engine/Plugins/Developer/Concert/ConcertApp/MultiUserClient/Source/MultiUserClient/Private/Replication/Authority/AuthoritySynchronizer_Base.cpp b/Engine/Plugins/Developer/Concert/ConcertApp/MultiUserClient/Source/MultiUserClient/Private/Replication/Authority/AuthoritySynchronizer_Base.cpp
deleted file mode 100644
diff --git a/Engine/Plugins/Developer/Concert/ConcertApp/MultiUserClient/Source/MultiUserClient/Private/Replication/Authority/AuthoritySynchronizer_Base.h b/Engine/Plugins/Developer/Concert/ConcertApp/MultiUserClient/Source/MultiUserClient/Private/Replication/Authority/AuthoritySynchronizer_Base.h
deleted file mode 100644
diff --git a/Engine/Plugins/Developer/Concert/ConcertApp/MultiUserClient/Source/MultiUserClient/Private/Replication/Authority/AuthoritySynchronizer_LocalClient.cpp b/Engine/Plugins/Developer/Concert/ConcertApp/MultiUserClient/Source/MultiUserClient/Private/Replication/Authority/AuthoritySynchronizer_LocalClient.cpp
--- a/Engine/Plugins/Developer/Concert/ConcertApp/MultiUserClient/Source/MultiUserClient/Private/Replication/Authority/AuthoritySynchronizer_LocalClient.cpp
+++ b/Engine/Plugins/Developer/Concert/ConcertApp/MultiUserClient/Source/MultiUserClient/Private/Replication/Authority/AuthoritySynchronizer_LocalClient.cpp
@@ -7,12 +7,8 @@
 
 namespace UE::MultiUserClient
 {
-	FAuthoritySynchronizer_LocalClient::FAuthoritySynchronizer_LocalClient(
-		TSharedRef<IConcertSyncClient> InClient,
-		FDoesObjectHaveProperties InDoesObjectHaveProperties
-		)
-		: FAuthoritySynchronizer_Base(MoveTemp(InDoesObjectHaveProperties))
-		, Client(MoveTemp(InClient))
+	FAuthoritySynchronizer_LocalClient::FAuthoritySynchronizer_LocalClient(TSharedRef<IConcertSyncClient> InClient)
+		: Client(MoveTemp(InClient))
 	{
 		IConcertClientReplicationManager* ReplicationManager = Client->GetReplicationManager();
 		if (ensure(ReplicationManager))
diff --git a/Engine/Plugins/Developer/Concert/ConcertApp/MultiUserClient/Source/MultiUserClient/Private/Replication/Authority/AuthoritySynchronizer_LocalClient.h b/Engine/Plugins/Developer/Concert/ConcertApp/MultiUserClient/Source/MultiUserClient/Private/Replication/Authority/AuthoritySynchronizer_LocalClient.h
--- a/Engine/Plugins/Developer/Concert/ConcertApp/MultiUserClient/Source/MultiUserClient/Private/Replication/Authority/AuthoritySynchronizer_LocalClient.h
+++ b/Engine/Plugins/Developer/Concert/ConcertApp/MultiUserClient/Source/MultiUserClient/Private/Replication/Authority/AuthoritySynchronizer_LocalClient.h
@@ -2,7 +2,7 @@
 
 #pragma once
 
-#include "AuthoritySynchronizer_Base.h"
+#include "IClientAuthoritySynchronizer.h"
 #include "Templates/UnrealTemplate.h"
 
 class IConcertSyncClient;
@@ -13,7 +13,7 @@ namespace UE::MultiUserClient
 	{
 	public:
 		
-		FAuthoritySynchronizer_LocalClient(TSharedRef<IConcertSyncClient> InClient, FDoesObjectHaveProperties InDoesObjectHaveProperties);
+		FAuthoritySynchronizer_LocalClient(TSharedRef<IConcertSyncClient> InClient);
 		virtual ~FAuthoritySynchronizer_LocalClient() override;
 
 		//~ Begin IClientAuthoritySynchronizer Interface
diff --git a/Engine/Plugins/Developer/Concert/ConcertApp/MultiUserClient/Source/MultiUserClient/Private/Replication/Authority/AuthoritySynchronizer_RemoteClient.cpp b/Engine/Plugins/Developer/Concert/ConcertApp/MultiUserClient/Source/MultiUserClient/Private/Replication/Authority/AuthoritySynchronizer_RemoteClient.cpp
--- a/Engine/Plugins/Developer/Concert/ConcertApp/MultiUserClient/Source/MultiUserClient/Private/Replication/Authority/AuthoritySynchronizer_RemoteClient.cpp
+++ b/Engine/Plugins/Developer/Concert/ConcertApp/MultiUserClient/Source/MultiUserClient/Private/Replication/Authority/AuthoritySynchronizer_RemoteClient.cpp
@@ -6,13 +6,8 @@
 
 namespace UE::MultiUserClient
 {
-	FAuthoritySynchronizer_RemoteClient::FAuthoritySynchronizer_RemoteClient(
-		const FGuid& RemoteEndpointId,
-		FRegularQueryService& InQueryService,
-		FDoesObjectHaveProperties InDoesObjectHaveProperties
-		)
-		: FAuthoritySynchronizer_Base(MoveTemp(InDoesObjectHaveProperties))
-		, QueryService(InQueryService)
+	FAuthoritySynchronizer_RemoteClient::FAuthoritySynchronizer_RemoteClient(const FGuid& RemoteEndpointId, FRegularQueryService& InQueryService)
+		: QueryService(InQueryService)
 		, QueryStreamHandle(
     		QueryService.RegisterAuthorityQuery(
     			RemoteEndpointId,
@@ -26,12 +21,6 @@ namespace UE::MultiUserClient
 		QueryService.UnregisterAuthorityQuery(QueryStreamHandle);
 	}
 
-	EAuthorityMutability FAuthoritySynchronizer_RemoteClient::GetChangeAuthorityMutability(const FSoftObjectPath& ObjectPath) const
-	{
-		// TODO DP: UE-198088 return whether it is legal to change the authority
-		return EAuthorityMutability::NotSupported;
-	}
-
 	bool FAuthoritySynchronizer_RemoteClient::HasAnyAuthority() const
 	{
 		return !LastServerState.IsEmpty();
diff --git a/Engine/Plugins/Developer/Concert/ConcertApp/MultiUserClient/Source/MultiUserClient/Private/Replication/Authority/AuthoritySynchronizer_RemoteClient.h b/Engine/Plugins/Developer/Concert/ConcertApp/MultiUserClient/Source/MultiUserClient/Private/Replication/Authority/AuthoritySynchronizer_RemoteClient.h
--- a/Engine/Plugins/Developer/Concert/ConcertApp/MultiUserClient/Source/MultiUserClient/Private/Replication/Authority/AuthoritySynchronizer_RemoteClient.h
+++ b/Engine/Plugins/Developer/Concert/ConcertApp/MultiUserClient/Source/MultiUserClient/Private/Replication/Authority/AuthoritySynchronizer_RemoteClient.h
@@ -2,7 +2,7 @@
 
 #pragma once
 
-#include "AuthoritySynchronizer_Base.h"
+#include "IClientAuthoritySynchronizer.h"
 
 #include "Containers/Set.h"
 #include "UObject/SoftObjectPath.h"
@@ -18,11 +18,10 @@ namespace UE::MultiUserClient
 	{
 	public:
 		
-		FAuthoritySynchronizer_RemoteClient(const FGuid& RemoteEndpointId, FRegularQueryService& InQueryService, FDoesObjectHaveProperties InDoesObjectHaveProperties);
+		FAuthoritySynchronizer_RemoteClient(const FGuid& RemoteEndpointId, FRegularQueryService& InQueryService);
 		virtual ~FAuthoritySynchronizer_RemoteClient();
 
 		//~ Begin IClientAuthoritySynchronizer Interface
-		virtual EAuthorityMutability GetChangeAuthorityMutability(const FSoftObjectPath& ObjectPath) const override;
 		virtual bool HasAnyAuthority() const override;
 		virtual bool HasAuthorityOver(const FSoftObjectPath& ObjectPath) const override { return LastServerState.Contains(ObjectPath); }
 		//~ End IClientAuthoritySynchronizer Interface
diff --git a/Engine/Plugins/Developer/Concert/ConcertApp/MultiUserClient/Source/MultiUserClient/Private/Replication/Authority/EAuthorityMutability.h b/Engine/Plugins/Developer/Concert/ConcertApp/MultiUserClient/Source/MultiUserClient/Private/Replication/Authority/EAuthorityMutability.h
--- a/Engine/Plugins/Developer/Concert/ConcertApp/MultiUserClient/Source/MultiUserClient/Private/Replication/Authority/EAuthorityMutability.h
+++ b/Engine/Plugins/Developer/Concert/ConcertApp/MultiUserClient/Source/MultiUserClient/Private/Replication/Authority/EAuthorityMutability.h
@@ -4,16 +4,16 @@
 
 namespace UE::MultiUserClient
 {
-	/** Various reasons in which an object can be when considering take over its authority */
-	enum class EAuthorityMutability
+	/** Various reasons in which an object can be when considering changing its authority */
+	enum class EAuthorityMutability : uint8
 	{
-		/** Authority can be taken */
-		Allowed,
-
 		/** The object has no registered properties. Cannot take authority. */
-		NoProperties,
+		NotApplicable,
 		
-		/** The operation is not implemented. */
-		NotSupported
+		/** There is a conflict with another client */
+		Conflict,
+		
+		/** Authority can be taken */
+		Allowed
 	};
 }
\ No newline at end of file
diff --git a/Engine/Plugins/Developer/Concert/ConcertApp/MultiUserClient/Source/MultiUserClient/Private/Replication/Authority/IClientAuthoritySynchronizer.h b/Engine/Plugins/Developer/Concert/ConcertApp/MultiUserClient/Source/MultiUserClient/Private/Replication/Authority/IClientAuthoritySynchronizer.h
--- a/Engine/Plugins/Developer/Concert/ConcertApp/MultiUserClient/Source/MultiUserClient/Private/Replication/Authority/IClientAuthoritySynchronizer.h
+++ b/Engine/Plugins/Developer/Concert/ConcertApp/MultiUserClient/Source/MultiUserClient/Private/Replication/Authority/IClientAuthoritySynchronizer.h
@@ -2,9 +2,8 @@
 
 #pragma once
 
-#include "EAuthorityMutability.h"
-
 #include "Delegates/Delegate.h"
+#include "Templates/UnrealTemplate.h"
 
 struct FSoftObjectPath;
 
@@ -15,11 +14,6 @@ namespace UE::MultiUserClient
 	{
 	public:
 		
-		/** @return The reason of why it is or isn't possible to change the authority for the given object. */
-		virtual EAuthorityMutability GetChangeAuthorityMutability(const FSoftObjectPath& ObjectPath) const = 0;
-		/** @return Whether the ObjectPath's authority can be changed. */
-		bool CanChangeAuthority(const FSoftObjectPath& ObjectPath) const { return GetChangeAuthorityMutability(ObjectPath) == EAuthorityMutability::Allowed; }
-
 		/** @return Whether this client is sending anything at all. */
 		virtual bool HasAnyAuthority() const = 0;
 
@@ -32,4 +26,21 @@ namespace UE::MultiUserClient
 
 		virtual ~IClientAuthoritySynchronizer() = default;
 	};
+
+	/** Util base class for implementing the events */
+	class FAuthoritySynchronizer_Base
+		: public IClientAuthoritySynchronizer
+		, public FNoncopyable
+	{
+	public:
+
+		//~ Begin IClientAuthoritySynchronizer Interface
+		virtual FOnServerStateChanged& OnServerStateChanged() override { return OnServerStateChangedDelegate; }
+		//~ End IClientAuthoritySynchronizer Interface
+
+	protected:
+
+		/** Triggered by subclasses when the authority state changes. */
+		FOnServerStateChanged OnServerStateChangedDelegate;
+	};
 }
\ No newline at end of file
diff --git a/Engine/Plugins/Developer/Concert/ConcertApp/MultiUserClient/Source/MultiUserClient/Private/Replication/Client/LocalReplicationClient.cpp b/Engine/Plugins/Developer/Concert/ConcertApp/MultiUserClient/Source/MultiUserClient/Private/Replication/Client/LocalReplicationClient.cpp
--- a/Engine/Plugins/Developer/Concert/ConcertApp/MultiUserClient/Source/MultiUserClient/Private/Replication/Client/LocalReplicationClient.cpp
+++ b/Engine/Plugins/Developer/Concert/ConcertApp/MultiUserClient/Source/MultiUserClient/Private/Replication/Client/LocalReplicationClient.cpp
@@ -19,12 +19,8 @@ namespace UE::MultiUserClient
 			InAuthorityCache,
 			InSessionContent,
 			MoveTemp(InStreamSynchronizer),
-		MakeUnique<FAuthoritySynchronizer_LocalClient>(
-			InClient,
-			FDoesObjectHaveProperties::CreateLambda([this](const FSoftObjectPath& ObjectPath)
-			{
-				return GetStreamDiffer().DoesObjectHavePropertiesAfterSubmit(ObjectPath);
-			})),
+			MakeUnique<FAuthoritySynchronizer_LocalClient>(InClient),
 			MakeUnique<FSubmissionWorkflow_LocalClient>(MoveTemp(InClient)))
+		, RemoteSubmissionListener(InClient->GetConcertClient()->GetCurrentSession().ToSharedRef(), GetStreamSynchronizer(), GetSubmissionWorkflow())
 	{}
 }
diff --git a/Engine/Plugins/Developer/Concert/ConcertApp/MultiUserClient/Source/MultiUserClient/Private/Replication/Client/LocalReplicationClient.h b/Engine/Plugins/Developer/Concert/ConcertApp/MultiUserClient/Source/MultiUserClient/Private/Replication/Client/LocalReplicationClient.h
--- a/Engine/Plugins/Developer/Concert/ConcertApp/MultiUserClient/Source/MultiUserClient/Private/Replication/Client/LocalReplicationClient.h
+++ b/Engine/Plugins/Developer/Concert/ConcertApp/MultiUserClient/Source/MultiUserClient/Private/Replication/Client/LocalReplicationClient.h
@@ -3,6 +3,7 @@
 #pragma once
 
 #include "ReplicationClient.h"
+#include "Replication/Submission/Remote/RemoteSubmissionListener.h"
 
 namespace UE::MultiUserClient
 {
@@ -19,6 +20,11 @@ namespace UE::MultiUserClient
 			TUniquePtr<IClientStreamSynchronizer> InStreamSynchronizer,
 			TSharedRef<IConcertSyncClient> InClient
 			);
+
+	private:
+		
+		/** Listens for and handles for submission request made by remote clients' SubmissionWorkflow.*/
+		FRemoteSubmissionListener RemoteSubmissionListener;
 	};
 }
 
diff --git a/Engine/Plugins/Developer/Concert/ConcertApp/MultiUserClient/Source/MultiUserClient/Private/Replication/Client/RemoteReplicationClient.cpp b/Engine/Plugins/Developer/Concert/ConcertApp/MultiUserClient/Source/MultiUserClient/Private/Replication/Client/RemoteReplicationClient.cpp
--- a/Engine/Plugins/Developer/Concert/ConcertApp/MultiUserClient/Source/MultiUserClient/Private/Replication/Client/RemoteReplicationClient.cpp
+++ b/Engine/Plugins/Developer/Concert/ConcertApp/MultiUserClient/Source/MultiUserClient/Private/Replication/Client/RemoteReplicationClient.cpp
@@ -2,15 +2,17 @@
 
 #include "RemoteReplicationClient.h"
 
+#include "IConcertClient.h"
 #include "Assets/MultiUserReplicationClientPreset.h"
 #include "Replication/Authority/AuthoritySynchronizer_RemoteClient.h"
 #include "Replication/Stream/StreamSynchronizer_RemoteClient.h"
-#include "Replication/Submission/SubmissionWorkflow_RemoteClient.h"
+#include "Replication/Submission/Remote/SubmissionWorkflow_RemoteClient.h"
 
 namespace UE::MultiUserClient
 {
 	FRemoteReplicationClient::FRemoteReplicationClient(
 		const FGuid& InConcertClientId,
+		TSharedRef<IConcertClient> InClient,
 		FGlobalAuthorityCache& InAuthorityCache,
 		UMultiUserReplicationClientPreset& InSessionContent,
 		FRegularQueryService& QueryService
@@ -20,19 +22,7 @@ namespace UE::MultiUserClient
 			InAuthorityCache,
 			InSessionContent,
 			MakeUnique<FStreamSynchronizer_RemoteClient>(InConcertClientId, QueryService),
-			MakeUnique<FAuthoritySynchronizer_RemoteClient>(
-				InConcertClientId,
-				QueryService,
-				FDoesObjectHaveProperties::CreateLambda([this](const FSoftObjectPath& ObjectPath)
-				{
-					return GetStreamDiffer().DoesObjectHavePropertiesAfterSubmit(ObjectPath);
-				})),
-			MakeUnique<FSubmissionWorkflow_RemoteClient>())
-	{
-		// When the remote client's state has changed, refresh the UI.
-		GetStreamSynchronizer().OnServerStateChanged().AddLambda([this]()
-		{
-			GetClientContent()->Stream->ReplicationMap = GetStreamSynchronizer().GetServerState();
-		});
-	}
+			MakeUnique<FAuthoritySynchronizer_RemoteClient>(InConcertClientId,QueryService),
+			MakeUnique<FSubmissionWorkflow_RemoteClient>(InClient->GetCurrentSession().ToSharedRef(), InConcertClientId))
+	{}
 }
diff --git a/Engine/Plugins/Developer/Concert/ConcertApp/MultiUserClient/Source/MultiUserClient/Private/Replication/Client/RemoteReplicationClient.h b/Engine/Plugins/Developer/Concert/ConcertApp/MultiUserClient/Source/MultiUserClient/Private/Replication/Client/RemoteReplicationClient.h
--- a/Engine/Plugins/Developer/Concert/ConcertApp/MultiUserClient/Source/MultiUserClient/Private/Replication/Client/RemoteReplicationClient.h
+++ b/Engine/Plugins/Developer/Concert/ConcertApp/MultiUserClient/Source/MultiUserClient/Private/Replication/Client/RemoteReplicationClient.h
@@ -4,6 +4,8 @@
 
 #include "ReplicationClient.h"
 
+class IConcertClient;
+
 namespace UE::MultiUserClient
 {
 	class FRegularQueryService;
@@ -15,6 +17,7 @@ namespace UE::MultiUserClient
 
 		FRemoteReplicationClient(
 			const FGuid& InConcertClientId,
+			TSharedRef<IConcertClient> InClient,
 			FGlobalAuthorityCache& InAuthorityCache,
 			UMultiUserReplicationClientPreset& InSessionContent,
 			FRegularQueryService& QueryService
diff --git a/Engine/Plugins/Developer/Concert/ConcertApp/MultiUserClient/Source/MultiUserClient/Private/Replication/Client/ReplicationClient.cpp b/Engine/Plugins/Developer/Concert/ConcertApp/MultiUserClient/Source/MultiUserClient/Private/Replication/Client/ReplicationClient.cpp
--- a/Engine/Plugins/Developer/Concert/ConcertApp/MultiUserClient/Source/MultiUserClient/Private/Replication/Client/ReplicationClient.cpp
+++ b/Engine/Plugins/Developer/Concert/ConcertApp/MultiUserClient/Source/MultiUserClient/Private/Replication/Client/ReplicationClient.cpp
@@ -40,9 +40,10 @@ namespace UE::MultiUserClient
 	{
 		LocalClientEditModel->OnObjectsChanged().AddRaw(this, &FReplicationClient::OnObjectsChanged);
 		LocalClientEditModel->OnPropertiesChanged().AddRaw(this, &FReplicationClient::OnPropertiesChanged);
+		LocalAuthorityDiffer.OnAddedOwnedObjects().AddRaw(this, &FReplicationClient::DeferOnModelChanged);
 
-		SubmissionWorkflow->OnAuthorityRequestCompleted().AddRaw(this, &FReplicationClient::OnAuthoritySubmissionCompleted);
-		StreamSynchronizer->OnServerStateChanged().AddRaw(this, &FReplicationClient::DeferOnModelChanged);
+		SubmissionWorkflow->OnAuthorityRequestCompleted_AnyThread().AddRaw(this, &FReplicationClient::OnAuthoritySubmissionCompleted);
+		StreamSynchronizer->OnServerStateChanged().AddRaw(this, &FReplicationClient::OnServerStateChanged);
 	}
 
 	FReplicationClient::~FReplicationClient()
@@ -69,6 +70,14 @@ namespace UE::MultiUserClient
 		DeferOnModelChanged();
 	}
 
+	void FReplicationClient::OnServerStateChanged()
+	{
+		// Whenever this client's server state changes, the UI must be refreshed.
+		GetClientContent()->Stream->ReplicationMap = GetStreamSynchronizer().GetServerState();
+
+		DeferOnModelChanged();
+	}
+
 	void FReplicationClient::DeferOnModelChanged(TConstArrayView<UObject*> AddedObjects)
 	{
 		if (!DeferredOnModelChangedData)
diff --git a/Engine/Plugins/Developer/Concert/ConcertApp/MultiUserClient/Source/MultiUserClient/Private/Replication/Client/ReplicationClient.h b/Engine/Plugins/Developer/Concert/ConcertApp/MultiUserClient/Source/MultiUserClient/Private/Replication/Client/ReplicationClient.h
--- a/Engine/Plugins/Developer/Concert/ConcertApp/MultiUserClient/Source/MultiUserClient/Private/Replication/Client/ReplicationClient.h
+++ b/Engine/Plugins/Developer/Concert/ConcertApp/MultiUserClient/Source/MultiUserClient/Private/Replication/Client/ReplicationClient.h
@@ -13,6 +13,7 @@
 #include "Templates/SharedPointer.h"
 #include "Templates/UnrealTemplate.h"
 
+class IConcertClient;
 class UMultiUserReplicationClientPreset;
 
 namespace UE::ConcertClientSharedSlate
@@ -127,6 +128,8 @@ namespace UE::MultiUserClient
 		void OnObjectsChanged(TConstArrayView<UObject*> AddedObjects, TConstArrayView<FSoftObjectPath> RemovedObjects, ConcertClientSharedSlate::EReplicatedObjectChangeReason ReplicatedObjectChangeReason);
 		void OnPropertiesChanged();
 
+		void OnServerStateChanged();
+
 		/** Defers rebuilding operations, such as refreshing state and calling OnModelChanged, in case there are multiple changes in the same frame. */
 		void DeferOnModelChanged() { DeferOnModelChanged({}); }
 		void DeferOnModelChanged(TConstArrayView<UObject*> AddedObjects);
diff --git a/Engine/Plugins/Developer/Concert/ConcertApp/MultiUserClient/Source/MultiUserClient/Private/Replication/Client/ReplicationClientManager.cpp b/Engine/Plugins/Developer/Concert/ConcertApp/MultiUserClient/Source/MultiUserClient/Private/Replication/Client/ReplicationClientManager.cpp
--- a/Engine/Plugins/Developer/Concert/ConcertApp/MultiUserClient/Source/MultiUserClient/Private/Replication/Client/ReplicationClientManager.cpp
+++ b/Engine/Plugins/Developer/Concert/ConcertApp/MultiUserClient/Source/MultiUserClient/Private/Replication/Client/ReplicationClientManager.cpp
@@ -18,6 +18,7 @@ namespace UE::MultiUserClient
 		TSharedRef<IConcertClientSession> InSession
 		)
 		: SessionContent(NewObject<UMultiUserReplicationSessionPreset>(GetTransientPackage(), NAME_None, RF_Transient))
+		, ConcertClient(InClient)
 		, Session(InSession)
 		, QueryService(InClient)
 		, AuthorityCache(*this)
@@ -135,7 +136,7 @@ namespace UE::MultiUserClient
 	
 	void FReplicationClientManager::CreateRemoteClient(const FGuid& ClientEndpointId, bool bBroadcastDelegate)
 	{
-		TUniquePtr<FRemoteReplicationClient> RemoteClientPtr = MakeUnique<FRemoteReplicationClient>(ClientEndpointId, AuthorityCache, *SessionContent->AddClient(), QueryService);
+		TUniquePtr<FRemoteReplicationClient> RemoteClientPtr = MakeUnique<FRemoteReplicationClient>(ClientEndpointId, ConcertClient->GetConcertClient(), AuthorityCache, *SessionContent->AddClient(), QueryService);
 		FRemoteReplicationClient& RemoteClient = *RemoteClientPtr;
 		RemoteClients.Emplace(
 			MoveTemp(RemoteClientPtr)
diff --git a/Engine/Plugins/Developer/Concert/ConcertApp/MultiUserClient/Source/MultiUserClient/Private/Replication/Client/ReplicationClientManager.h b/Engine/Plugins/Developer/Concert/ConcertApp/MultiUserClient/Source/MultiUserClient/Private/Replication/Client/ReplicationClientManager.h
--- a/Engine/Plugins/Developer/Concert/ConcertApp/MultiUserClient/Source/MultiUserClient/Private/Replication/Client/ReplicationClientManager.h
+++ b/Engine/Plugins/Developer/Concert/ConcertApp/MultiUserClient/Source/MultiUserClient/Private/Replication/Client/ReplicationClientManager.h
@@ -102,6 +102,9 @@ namespace UE::MultiUserClient
 		/** The state of the server is synched up with this object and displayed in the UI. */
 		TObjectPtr<UMultiUserReplicationSessionPreset> SessionContent;
 
+		/** The local Concert client */
+		TSharedRef<IConcertSyncClient> ConcertClient;
+
 		/**
 		 * The session the local client is in.
 		 * 
diff --git a/Engine/Plugins/Developer/Concert/ConcertApp/MultiUserClient/Source/MultiUserClient/Private/Replication/Submission/AutoSubmissionPolicy.cpp b/Engine/Plugins/Developer/Concert/ConcertApp/MultiUserClient/Source/MultiUserClient/Private/Replication/Submission/AutoSubmissionPolicy.cpp
--- a/Engine/Plugins/Developer/Concert/ConcertApp/MultiUserClient/Source/MultiUserClient/Private/Replication/Submission/AutoSubmissionPolicy.cpp
+++ b/Engine/Plugins/Developer/Concert/ConcertApp/MultiUserClient/Source/MultiUserClient/Private/Replication/Submission/AutoSubmissionPolicy.cpp
@@ -7,8 +7,6 @@
 #include "Replication/Authority/AuthorityChangeTracker.h"
 #include "Replication/Editor/Model/IEditableReplicationStreamModel.h"
 
-#include "Misc/CoreDelegates.h"
-
 namespace UE::MultiUserClient
 {
 	FAutoSubmissionPolicy::FAutoSubmissionPolicy(
@@ -29,13 +27,9 @@ namespace UE::MultiUserClient
 
 	FAutoSubmissionPolicy::~FAutoSubmissionPolicy()
 	{
-		FCoreDelegates::OnEndFrame.RemoveAll(this);
-		
 		StreamEditorModel.OnObjectsChanged().RemoveAll(this);
 		StreamEditorModel.OnPropertiesChanged().RemoveAll(this);
 		AuthorityChangeTracker.OnAddedOwnedObjects().RemoveAll(this);
-		
-		SubmissionWorkflow.OnStreamRequestCompleted().RemoveAll(this);
 	}
 
 	void FAutoSubmissionPolicy::ProcessAccumulatedChangesAndSubmit()
@@ -47,25 +41,25 @@ namespace UE::MultiUserClient
 		
 		if (SubmissionWorkflow.CanSubmit())
 		{
+			SubmissionWorkflow.OnSubmitOperationCompleted_AnyThread().RemoveAll(this);
 			SubmitChanges();
 		}
-		else if (!SubmissionWorkflow.OnSubmitOperationCompleted().IsBoundToObject(this))
+		else if (!SubmissionWorkflow.OnSubmitOperationCompleted_AnyThread().IsBoundToObject(this))
 		{
-			SubmissionWorkflow.OnSubmitOperationCompleted().AddRaw(this, &FAutoSubmissionPolicy::ProcessAccumulatedChangesAndSubmit);
+			SubmissionWorkflow.OnSubmitOperationCompleted_AnyThread().AddRaw(this, &FAutoSubmissionPolicy::ProcessAccumulatedChangesAndSubmit);
 		}
 	}
 
 	void FAutoSubmissionPolicy::SubmitChanges()
 	{
 		using namespace UE::ConcertSyncClient::Replication;
-		FCoreDelegates::OnEndFrame.RemoveAll(this);
+		bIsDirty = false;
 			
 		// Even though authority request is sent after server confirms stream change, the authority request is pre-built to avoid sending changes
 		// the local client makes while we're waiting for the latent server responses.
 		TOptional<FAuthorityChangeRequest> AuthorityChangeRequest = RequestBuilder.BuildAuthorityChange();
 		TOptional<FChangeStreamRequest> StreamRequest = RequestBuilder.BuildStreamChange();
 			
-		// Should be true since a change caused OnEndFrame to be registered in the first place, but in THEORY changes could have been reverted in the same frame (e.g. via a transaction)
 		if (AuthorityChangeRequest || StreamRequest)
 		{
 			SubmissionWorkflow.SubmitChanges({ StreamRequest, AuthorityChangeRequest });
diff --git a/Engine/Plugins/Developer/Concert/ConcertApp/MultiUserClient/Source/MultiUserClient/Private/Replication/Submission/Data/EChangeUploadability.h b/Engine/Plugins/Developer/Concert/ConcertApp/MultiUserClient/Source/MultiUserClient/Private/Replication/Submission/Data/EChangeUploadability.h
--- a/Engine/Plugins/Developer/Concert/ConcertApp/MultiUserClient/Source/MultiUserClient/Private/Replication/Submission/Data/EChangeUploadability.h
+++ b/Engine/Plugins/Developer/Concert/ConcertApp/MultiUserClient/Source/MultiUserClient/Private/Replication/Submission/Data/EChangeUploadability.h
@@ -11,14 +11,13 @@ namespace UE::MultiUserClient
 
 		/** A previous upload operation is in progress */
 		InProgress,
-		
-		// TODO DP UE-197435: Remove once remote client streams can be changed
-		// TODO DP UE-198088: Remove once remote changing is implemented
-		NotImplemented
+
+		/** Changing is generally not available, e.g. remote client does not allow changes. */
+		NotAvailable
 	};
 
 	inline bool CanEverSubmit(EChangeUploadability Uploadability)
 	{
-		return Uploadability != EChangeUploadability::NotImplemented;
+		return Uploadability != EChangeUploadability::NotAvailable;
 	}
 }
\ No newline at end of file
diff --git a/Engine/Plugins/Developer/Concert/ConcertApp/MultiUserClient/Source/MultiUserClient/Private/Replication/Submission/Data/StreamSubmission.h b/Engine/Plugins/Developer/Concert/ConcertApp/MultiUserClient/Source/MultiUserClient/Private/Replication/Submission/Data/StreamSubmission.h
--- a/Engine/Plugins/Developer/Concert/ConcertApp/MultiUserClient/Source/MultiUserClient/Private/Replication/Submission/Data/StreamSubmission.h
+++ b/Engine/Plugins/Developer/Concert/ConcertApp/MultiUserClient/Source/MultiUserClient/Private/Replication/Submission/Data/StreamSubmission.h
@@ -36,7 +36,7 @@ namespace UE::MultiUserClient
 		/** Error code of the submission. Determines whether Response is valid. */
 		EStreamSubmissionErrorCode ErrorCode;
 
-		/** Valid if ErrorCode != ESubmitChangesErrorCode::AlreadyInProgress. */
+		/** Valid if ErrorCode == ESubmitChangesErrorCode::Success. */
 		TOptional<FCompletedChangeSubmission> SubmissionInfo;
 	};
 }
diff --git a/Engine/Plugins/Developer/Concert/ConcertApp/MultiUserClient/Source/MultiUserClient/Private/Replication/Submission/ISubmissionWorkflow.h b/Engine/Plugins/Developer/Concert/ConcertApp/MultiUserClient/Source/MultiUserClient/Private/Replication/Submission/ISubmissionWorkflow.h
--- a/Engine/Plugins/Developer/Concert/ConcertApp/MultiUserClient/Source/MultiUserClient/Private/Replication/Submission/ISubmissionWorkflow.h
+++ b/Engine/Plugins/Developer/Concert/ConcertApp/MultiUserClient/Source/MultiUserClient/Private/Replication/Submission/ISubmissionWorkflow.h
@@ -66,18 +66,18 @@ namespace UE::MultiUserClient
 		 * Broadcasts whenever a submit operation completes the stream change request stage. Useful for accumulating and counting errors.
 		 * @note No stream changes may have been requested. Check the error code.
 		 */
-		virtual FOnStreamRequestCompleted& OnStreamRequestCompleted() = 0;
+		virtual FOnStreamRequestCompleted& OnStreamRequestCompleted_AnyThread() = 0;
 
 		DECLARE_MULTICAST_DELEGATE_TwoParams(FOnAuthorityRequestCompleted, const FSubmitAuthorityChangesRequest&, const FSubmitAuthorityChangesResponse&);
 		/**
 		 * Broadcasts whenever a submit operation completes the authority change request stage. Useful for accumulating and counting errors.
 		 * @note No authority changes may have been requested. Check the error code.
 		 */
-		virtual FOnAuthorityRequestCompleted& OnAuthorityRequestCompleted() = 0;
+		virtual FOnAuthorityRequestCompleted& OnAuthorityRequestCompleted_AnyThread() = 0;
 
 		DECLARE_MULTICAST_DELEGATE(FOnSubmitOperationCompleted);
 		/** Broadcasts whenever a submit operation completes. Useful for enqueuing submit operations. */
-		virtual FOnSubmitOperationCompleted& OnSubmitOperationCompleted() = 0;
+		virtual FOnSubmitOperationCompleted& OnSubmitOperationCompleted_AnyThread() = 0;
 
 		virtual ~ISubmissionWorkflow() = default;
 	};
@@ -88,9 +88,9 @@ namespace UE::MultiUserClient
 	public:
 
 		//~ Begin ISubmissionWorkflow Interface
-		virtual FOnStreamRequestCompleted& OnStreamRequestCompleted() override { return StreamRequestCompletedDelegate; }
-		virtual FOnAuthorityRequestCompleted& OnAuthorityRequestCompleted() override { return AuthorityRequestCompletedDelegate; }
-		virtual FOnSubmitOperationCompleted& OnSubmitOperationCompleted() override { return OnSubmitOperationCompletedDelegate; }
+		virtual FOnStreamRequestCompleted& OnStreamRequestCompleted_AnyThread() override { return StreamRequestCompletedDelegate; }
+		virtual FOnAuthorityRequestCompleted& OnAuthorityRequestCompleted_AnyThread() override { return AuthorityRequestCompletedDelegate; }
+		virtual FOnSubmitOperationCompleted& OnSubmitOperationCompleted_AnyThread() override { return OnSubmitOperationCompletedDelegate; }
 		//~ End ISubmissionWorkflow Interface
 
 	protected:
diff --git a/Engine/Plugins/Developer/Concert/ConcertApp/MultiUserClient/Source/MultiUserClient/Private/Replication/Submission/Notification/SubmissionNotifier.cpp b/Engine/Plugins/Developer/Concert/ConcertApp/MultiUserClient/Source/MultiUserClient/Private/Replication/Submission/Notification/SubmissionNotifier.cpp
--- a/Engine/Plugins/Developer/Concert/ConcertApp/MultiUserClient/Source/MultiUserClient/Private/Replication/Submission/Notification/SubmissionNotifier.cpp
+++ b/Engine/Plugins/Developer/Concert/ConcertApp/MultiUserClient/Source/MultiUserClient/Private/Replication/Submission/Notification/SubmissionNotifier.cpp
@@ -11,6 +11,7 @@
 
 #include "Framework/Notifications/NotificationManager.h"
 #include "Input/Reply.h"
+#include "Stats/Stats2.h"
 #include "Widgets/Notifications/SNotificationList.h"
 
 namespace UE::MultiUserClient
@@ -39,6 +40,17 @@ namespace UE::MultiUserClient
 		}
 	}
 
+	void FSubmissionNotifier::Tick(float DeltaTime)
+	{
+		ProcessCompletedStreamChanges();
+		ProcessCompletedAuthorityChanges();
+	}
+
+	TStatId FSubmissionNotifier::GetStatId() const
+	{
+		RETURN_QUICK_DECLARE_CYCLE_STAT(FSubmissionNotifier, STATGROUP_Tickables);
+	}
+
 	void FSubmissionNotifier::OnPostRemoteClientAdded(FRemoteReplicationClient& RemoteReplicationClient)
 	{
 		RegisterClient(RemoteReplicationClient);
@@ -46,27 +58,48 @@ namespace UE::MultiUserClient
 
 	void FSubmissionNotifier::RegisterClient(FReplicationClient& Client)
 	{
-		Client.GetSubmissionWorkflow().OnStreamRequestCompleted().AddRaw(this, &FSubmissionNotifier::OnStreamRequestCompleted);
-		Client.GetSubmissionWorkflow().OnAuthorityRequestCompleted().AddRaw(this, &FSubmissionNotifier::OnAuthorityRequestCompleted);
+		Client.GetSubmissionWorkflow().OnStreamRequestCompleted_AnyThread().AddRaw(this, &FSubmissionNotifier::OnStreamRequestCompleted_AnyThread);
+		Client.GetSubmissionWorkflow().OnAuthorityRequestCompleted_AnyThread().AddRaw(this, &FSubmissionNotifier::OnAuthorityRequestCompleted_AnyThread);
 	}
 
 	void FSubmissionNotifier::UnregisterClient(FReplicationClient& Client)
 	{
-		Client.GetSubmissionWorkflow().OnStreamRequestCompleted().RemoveAll(this);
-		Client.GetSubmissionWorkflow().OnAuthorityRequestCompleted().RemoveAll(this);
+		Client.GetSubmissionWorkflow().OnStreamRequestCompleted_AnyThread().RemoveAll(this);
+		Client.GetSubmissionWorkflow().OnAuthorityRequestCompleted_AnyThread().RemoveAll(this);
 	}
 
-	void FSubmissionNotifier::OnStreamRequestCompleted(const FSubmitStreamChangesResponse& Request)
+	void FSubmissionNotifier::OnStreamRequestCompleted_AnyThread(const FSubmitStreamChangesResponse& Response)
 	{
-		AccumulateStreamRejections(Request);
+		// Must be processed on the game thread because it might create a SNotification
+		StreamRequestQueue.Enqueue(Response);
+	}
+
+	void FSubmissionNotifier::OnAuthorityRequestCompleted_AnyThread(const FSubmitAuthorityChangesRequest& Request, const FSubmitAuthorityChangesResponse& Response)
+	{
+		// Must be processed on the game thread because it might create a SNotification
+		AuthorityRequestQueue.Enqueue({ Request, Response });
+	}
+
+	void FSubmissionNotifier::ProcessCompletedStreamChanges()
+	{
+		FSubmitStreamChangesResponse Response;
+		while (StreamRequestQueue.Dequeue(Response))
+		{
+			AccumulateStreamRejections(Response);
+		}
+		
 		CreateOrUpdateStreamNotification();
 	}
 
-	void FSubmissionNotifier::OnAuthorityRequestCompleted(const FSubmitAuthorityChangesRequest& Request, const FSubmitAuthorityChangesResponse& Response)
+	void FSubmissionNotifier::ProcessCompletedAuthorityChanges()
 	{
-		AccumulateAuthorityRejections(Request, Response);
-		CreateOrUpdateAuthorityNotification();
+		TPair<FSubmitAuthorityChangesRequest, FSubmitAuthorityChangesResponse> PendingPair;
+		while (AuthorityRequestQueue.Dequeue(PendingPair))
+		{
+			AccumulateAuthorityRejections(PendingPair.Key, PendingPair.Value);
+		}
 		
+		CreateOrUpdateAuthorityNotification();
 	}
 
 	void FSubmissionNotifier::AccumulateStreamRejections(const FSubmitStreamChangesResponse& CompletedOp)
diff --git a/Engine/Plugins/Developer/Concert/ConcertApp/MultiUserClient/Source/MultiUserClient/Private/Replication/Submission/Notification/SubmissionNotifier.h b/Engine/Plugins/Developer/Concert/ConcertApp/MultiUserClient/Source/MultiUserClient/Private/Replication/Submission/Notification/SubmissionNotifier.h
--- a/Engine/Plugins/Developer/Concert/ConcertApp/MultiUserClient/Source/MultiUserClient/Private/Replication/Submission/Notification/SubmissionNotifier.h
+++ b/Engine/Plugins/Developer/Concert/ConcertApp/MultiUserClient/Source/MultiUserClient/Private/Replication/Submission/Notification/SubmissionNotifier.h
@@ -3,7 +3,10 @@
 #pragma once
 
 #include "AccumulatedSubmissionErrors.h"
+
+#include "Containers/Queue.h"
 #include "Templates/UnrealTemplate.h"
+#include "TickableEditorObject.h"
 
 class SNotificationItem;
 class FReply;
@@ -25,13 +28,18 @@ namespace UE::MultiUserClient
 	 * Adds a SNotificationItem when ISubmissionWorkflow operations fail.
 	 * Failing operations are bundled into a single notification item until it is manually dismissed.
 	 */
-	class FSubmissionNotifier : public FNoncopyable
+	class FSubmissionNotifier : public FNoncopyable, public FTickableEditorObject
 	{
 	public:
 		
 		FSubmissionNotifier(FReplicationClientManager& InReplicationClientManager);
 		~FSubmissionNotifier();
 
+		//~ Begin FTickableEditorObject Interface
+		virtual void Tick(float DeltaTime) override;
+		virtual TStatId GetStatId() const override;
+		//~ End FTickableEditorObject Interface
+
 	private:
 
 		/** Emits events when remote clients are discovered. */
@@ -51,6 +59,12 @@ namespace UE::MultiUserClient
 		FAccumulatedStreamErrors StreamErrors;
 		/** Displayed by AuthorityRejectedNotification */
 		FAccumulatedAuthorityErrors AuthorityErrors;
+
+		/** Queues stream completed requests since they may be completed from any thread (usually on game thread except if it times out). */
+		TQueue<FSubmitStreamChangesResponse, EQueueMode::Mpsc> StreamRequestQueue;
+		/** Queues authority completed requests since they may be completed from any thread (usually on game thread except if it times out). */
+		TQueue<TPair<FSubmitAuthorityChangesRequest, FSubmitAuthorityChangesResponse>, EQueueMode::Mpsc> AuthorityRequestQueue;
+		
 		
 		// Client events
 		void OnPostRemoteClientAdded(FRemoteReplicationClient& RemoteReplicationClient);
@@ -58,8 +72,12 @@ namespace UE::MultiUserClient
 		void UnregisterClient(FReplicationClient& Client);
 
 		// Per-client submission events
-		void OnStreamRequestCompleted(const FSubmitStreamChangesResponse& Request);
-		void OnAuthorityRequestCompleted(const FSubmitAuthorityChangesRequest& Request, const FSubmitAuthorityChangesResponse& Response);
+		void OnStreamRequestCompleted_AnyThread(const FSubmitStreamChangesResponse& Response);
+		void OnAuthorityRequestCompleted_AnyThread(const FSubmitAuthorityChangesRequest& Request, const FSubmitAuthorityChangesResponse& Response);
+
+		// Process asynchronously received completions
+		void ProcessCompletedStreamChanges();
+		void ProcessCompletedAuthorityChanges();
 
 		// Handle rejections
 		void AccumulateStreamRejections(const FSubmitStreamChangesResponse& CompletedOp);
diff --git a/Engine/Plugins/Developer/Concert/ConcertApp/MultiUserClient/Source/MultiUserClient/Private/Replication/Submission/Remote/RemoteSubmissionListener.cpp b/Engine/Plugins/Developer/Concert/ConcertApp/MultiUserClient/Source/MultiUserClient/Private/Replication/Submission/Remote/RemoteSubmissionListener.cpp
new file mode 100644
--- /dev/null
+++ b/Engine/Plugins/Developer/Concert/ConcertApp/MultiUserClient/Source/MultiUserClient/Private/Replication/Submission/Remote/RemoteSubmissionListener.cpp
@@ -0,0 +1,180 @@
+// Copyright Epic Games, Inc. All Rights Reserved.
+
+#include "RemoteSubmissionListener.h"
+
+#include "ConcertLogGlobal.h"
+#include "IConcertSession.h"
+#include "RemoteSubmissionMessages.h"
+#include "Replication/Stream/IClientStreamSynchronizer.h"
+#include "Replication/Submission/ISubmissionOperation.h"
+#include "Replication/Submission/ISubmissionWorkflow.h"
+#include "Widgets/ActiveSession/Replication/Client/ClientUtils.h"
+
+namespace UE::MultiUserClient
+{
+	FRemoteSubmissionListener::FRemoteSubmissionListener(
+		TSharedRef<IConcertClientSession> InConcertSession,
+		IClientStreamSynchronizer& InStreamSynchronizer,
+		ISubmissionWorkflow& InSubmissionWorkflow
+		)
+		: ConcertSession(MoveTemp(InConcertSession))
+		, StreamSynchronizer(InStreamSynchronizer)
+		, SubmissionWorkflow(InSubmissionWorkflow)
+	{
+		ConcertSession->RegisterCustomRequestHandler<FMultiUser_ChangeRemote_Request, FMultiUser_ChangeRemote_Response>(this, &FRemoteSubmissionListener::HandleChangeRemoteRequest);
+	}
+
+	FRemoteSubmissionListener::~FRemoteSubmissionListener()
+	{
+		ConcertSession->UnregisterCustomRequestHandler<FMultiUser_ChangeRemote_Request>();
+		SubmissionWorkflow.OnSubmitOperationCompleted_AnyThread().RemoveAll(this);
+	}
+
+	EConcertSessionResponseCode FRemoteSubmissionListener::HandleChangeRemoteRequest(
+		const FConcertSessionContext& Context,
+		const FMultiUser_ChangeRemote_Request& Request,
+		FMultiUser_ChangeRemote_Response& Response
+		)
+	{
+		UE_LOG(LogConcert, Log, TEXT("Received remote change request from %s"), *ClientUtils::GetClientDisplayName(*ConcertSession, Context.SourceEndpointId));
+		Response.bTimeout = false;
+
+		if (InProgressOperation.IsSet())
+		{
+			Response.Error = EMultiUserChangeRemoteRequestError::OtherInProgress;
+			return EConcertSessionResponseCode::Failed;
+		}
+		
+		if (!CanProcessRequest(Request))
+		{
+			Response.Error = EMultiUserChangeRemoteRequestError::RejectedChange;
+			return EConcertSessionResponseCode::Failed;
+		}
+		
+		InProgressOperation = { Request, Context.SourceEndpointId };
+		SubmitNowOrWaitUntilReady();
+		return EConcertSessionResponseCode::Success;
+	}
+
+	void FRemoteSubmissionListener::SubmitNowOrWaitUntilReady(int32 NumTriesSoFar)
+	{
+		if (SubmissionWorkflow.GetUploadability() == EChangeUploadability::Ready)
+		{
+			SubmissionWorkflow.OnSubmitOperationCompleted_AnyThread().RemoveAll(this);
+			SubmitRequest();
+		}
+		else if (ensureMsgf(SubmissionWorkflow.GetUploadability() == EChangeUploadability::InProgress, TEXT("Unsupported state")))
+		{
+			++NumTriesSoFar;
+			UE_LOG(LogConcert, Log, TEXT("Post-boning requested remote update since a local change is already in progress (%dth try)"), NumTriesSoFar);
+			SubmissionWorkflow.OnSubmitOperationCompleted_AnyThread().AddRaw(this, &FRemoteSubmissionListener::SubmitNowOrWaitUntilReady, NumTriesSoFar);
+		}
+	}
+
+	void FRemoteSubmissionListener::SubmitRequest()
+	{
+		using namespace ConcertSyncClient::Replication;
+
+		const FMultiUser_ChangeRemote_Request& RequestData = InProgressOperation->Request;
+		if (!CanProcessRequest(RequestData))
+		{
+			CancelOperation(EMultiUserCancelRemoteChangeReason::PredictedConflict);
+			return;
+		}
+
+		const bool bHasStreamOnServer = !StreamSynchronizer.GetServerState().ReplicatedObjects.IsEmpty();
+		const bool bIsAddingStream = !RequestData.StreamChangeRequest.StreamsToAdd.IsEmpty();
+		const bool bHasStreamCreationConflict = (bHasStreamOnServer && bIsAddingStream) || (!bHasStreamOnServer && !bIsAddingStream);
+		if (bHasStreamCreationConflict)
+		{
+			// Problem: Requesting client and local client disagree about whether the local client already has a registered stream. In theory we can resolve this but for this WIP we won't.
+			CancelOperation(EMultiUserCancelRemoteChangeReason::PredictedConflict);
+			return;
+		}
+		
+		const TSharedPtr<ISubmissionOperation> Operation = SubmissionWorkflow.SubmitChanges({
+			FChangeStreamRequest { RequestData.StreamChangeRequest }, FAuthorityChangeRequest{ RequestData.AuthorityRequest } }
+			);
+		if (!ensure(Operation))
+		{
+			CancelOperation(EMultiUserCancelRemoteChangeReason::FailedToCreate);
+			return;
+		}
+
+		if (!RequestData.StreamChangeRequest.IsEmpty())
+		{
+			Operation->OnStreamChangesSubmittedFuture().Next([this, Token = LifetimeToken.ToWeakPtr()](FSubmitStreamChangesResponse&& Response)
+			{
+				if (const TSharedPtr<FToken> TokenPin = Token.Pin())
+				{
+					HandleStreamChangeDone(Response);
+				}
+			});
+		}
+		if (!RequestData.AuthorityRequest.IsEmpty())
+		{
+			Operation->OnAuthorityChangeResponseReceivedFuture().Next([this, Token = LifetimeToken.ToWeakPtr()](FSubmitAuthorityChangesResponse&& Response)
+			{
+				if (const TSharedPtr<FToken> TokenPin = Token.Pin())
+				{
+					HandleAuthorityChangeDone(Response);
+				}
+			});
+		}
+	}
+
+	void FRemoteSubmissionListener::HandleStreamChangeDone(const FSubmitStreamChangesResponse& Response)
+	{
+		const bool bProcessed = Response.ErrorCode == EStreamSubmissionErrorCode::Success && ensure(Response.SubmissionInfo.IsSet());
+		FMultiUser_ChangeRemote_StreamUpdatedEvent ProgressEvent { bProcessed };
+		if (bProcessed)
+		{
+			ProgressEvent.StreamChangeResponse = Response.SubmissionInfo->Response;
+		}
+		
+		InProgressOperation->bFinishedStream = true;
+		ConcertSession->SendCustomEvent(ProgressEvent, InProgressOperation->RequestingEndpointId, EConcertMessageFlags::ReliableOrdered);
+		FinishOperationIfDone();
+	}
+
+	void FRemoteSubmissionListener::HandleAuthorityChangeDone(const FSubmitAuthorityChangesResponse& Response)
+	{
+		const bool bProcessed = Response.ErrorCode == EAuthoritySubmissionResponseErrorCode::Success && ensure(Response.Response.IsSet());
+		FMultiUser_ChangeRemote_AuthorityUpdatedEvent ProgressEvent { bProcessed };
+		if (bProcessed)
+		{
+			ProgressEvent.AuthorityChangeResponse = *Response.Response;
+		}
+
+		InProgressOperation->bFinishedAuthority = true;
+		ConcertSession->SendCustomEvent(ProgressEvent, InProgressOperation->RequestingEndpointId, EConcertMessageFlags::ReliableOrdered);
+		FinishOperationIfDone();
+	}
+
+	void FRemoteSubmissionListener::FinishOperationIfDone()
+	{
+		const bool bIsStreamDone = InProgressOperation->bFinishedStream || InProgressOperation->Request.StreamChangeRequest.IsEmpty();
+		const bool bIsAuthorityDone = InProgressOperation->bFinishedAuthority || InProgressOperation->Request.AuthorityRequest.IsEmpty();
+		if (bIsStreamDone && bIsAuthorityDone)
+		{
+			CleanUpOperation();
+		}
+	}
+
+	bool FRemoteSubmissionListener::CanProcessRequest(const FMultiUser_ChangeRemote_Request& Request)
+	{
+		// TODO UE-201118: Are these changes free of any authority conflicts?
+		return true;
+	}
+
+	void FRemoteSubmissionListener::CancelOperation(EMultiUserCancelRemoteChangeReason Reason)
+	{
+		ConcertSession->SendCustomEvent(FMultiUser_ChangeRemote_Cancelled{ Reason }, InProgressOperation->RequestingEndpointId, EConcertMessageFlags::ReliableOrdered);
+		CleanUpOperation();
+	}
+
+	void FRemoteSubmissionListener::CleanUpOperation()
+	{
+		InProgressOperation.Reset();
+	}
+}
diff --git a/Engine/Plugins/Developer/Concert/ConcertApp/MultiUserClient/Source/MultiUserClient/Private/Replication/Submission/Remote/RemoteSubmissionListener.h b/Engine/Plugins/Developer/Concert/ConcertApp/MultiUserClient/Source/MultiUserClient/Private/Replication/Submission/Remote/RemoteSubmissionListener.h
new file mode 100644
--- /dev/null
+++ b/Engine/Plugins/Developer/Concert/ConcertApp/MultiUserClient/Source/MultiUserClient/Private/Replication/Submission/Remote/RemoteSubmissionListener.h
@@ -0,0 +1,91 @@
+// Copyright Epic Games, Inc. All Rights Reserved.
+
+#pragma once
+
+#include "RemoteSubmissionMessages.h"
+#include "Replication/IToken.h"
+#include "Templates/SharedPointer.h"
+#include "Templates/UnrealTemplate.h"
+
+namespace UE::MultiUserClient
+{
+	class IClientStreamSynchronizer;
+}
+
+class IConcertClientSession;
+
+enum class EConcertSessionResponseCode : uint8;
+
+struct FConcertSessionContext;
+
+namespace UE::MultiUserClient
+{
+	struct FSubmitAuthorityChangesResponse;
+	struct FSubmitStreamChangesResponse;
+	class ISubmissionWorkflow;
+	
+	/**
+	 * Handles stream and authority requests made by other clients towards the local client.
+	 *
+	 * For security reasons, the server only allows clients to change their own streams. Hence if client 1 wants to change something about client 2,
+	 * client 1 must kindly ask client 2 to do the change. In this case client 1's FSubmissionWorkflow_RemoteClient does the asking and client 2's
+	 * FRemoteSubmissionListener the answering.
+	 * 
+	 * The typical source of such requests is FSubmissionWorkflow_RemoteClient though theoretically the request can from any system with access
+	 * to a Concert endpoint.
+	 *
+	 * @see FSubmissionWorkflow_RemoteClient
+	 */
+	class FRemoteSubmissionListener : public FNoncopyable
+	{
+	public:
+		
+		FRemoteSubmissionListener(TSharedRef<IConcertClientSession> InConcertSession, IClientStreamSynchronizer& InStreamSynchronizer, ISubmissionWorkflow& InSubmissionWorkflow);
+		~FRemoteSubmissionListener();
+
+	private:
+		
+		/** Latent operations keep a weak reference to this to detect whether this FSubmissionWorkflow_LocalClient was destroyed. */
+		const TSharedRef<FToken> LifetimeToken = FToken::Make();
+
+		/** Used to register and unregister event handlers. This object's owner is supposed to destroy us before the session is destroyed. */
+		TSharedRef<IConcertClientSession> ConcertSession;
+
+		/** Used to validate a request against the client's server state. */
+		IClientStreamSynchronizer& StreamSynchronizer;
+		/** Used to serve the incoming change requests. */
+		ISubmissionWorkflow& SubmissionWorkflow;
+
+		struct FOperationData
+		{
+			const FMultiUser_ChangeRemote_Request Request;
+			const FGuid RequestingEndpointId;
+
+			bool bFinishedStream = false;
+			bool bFinishedAuthority = false;
+		};
+		/** Data for current in progress operation */
+		TOptional<FOperationData> InProgressOperation;
+
+		/** Handles remote change event. */
+		EConcertSessionResponseCode HandleChangeRemoteRequest(const FConcertSessionContext& Context, const FMultiUser_ChangeRemote_Request& Request, FMultiUser_ChangeRemote_Response& Response);
+		
+		/** Checks whether the local submit operation is ready. If not, we subscribe to it finishing and try again. */
+		void SubmitNowOrWaitUntilReady(int32 NumTriesSoFar = 0);
+		/** The submission workflow is ready for submission. Do the submission. */
+		void SubmitRequest();
+
+		void HandleStreamChangeDone(const FSubmitStreamChangesResponse& Response);
+		void HandleAuthorityChangeDone(const FSubmitAuthorityChangesResponse& Response);
+		void FinishOperationIfDone();
+
+		/** @return Whether this remote request should be processed. Returns false if there are any authority conflicts. */
+		bool CanProcessRequest(const FMultiUser_ChangeRemote_Request& Request);
+
+		/** Sends a cancel event to the requester and cleans up the operation.*/
+		void CancelOperation(EMultiUserCancelRemoteChangeReason Reason);
+		/** Resets InProgressOperation and does any other clean-up. */
+		void CleanUpOperation();
+	};
+}
+
diff --git a/Engine/Plugins/Developer/Concert/ConcertApp/MultiUserClient/Source/MultiUserClient/Private/Replication/Submission/Remote/RemoteSubmissionMessages.h b/Engine/Plugins/Developer/Concert/ConcertApp/MultiUserClient/Source/MultiUserClient/Private/Replication/Submission/Remote/RemoteSubmissionMessages.h
new file mode 100644
--- /dev/null
+++ b/Engine/Plugins/Developer/Concert/ConcertApp/MultiUserClient/Source/MultiUserClient/Private/Replication/Submission/Remote/RemoteSubmissionMessages.h
@@ -0,0 +1,111 @@
+// Copyright Epic Games, Inc. All Rights Reserved.
+
+#pragma once
+
+#include "Replication/Messages/ChangeAuthority.h"
+#include "Replication/Messages/ChangeStream.h"
+#include "RemoteSubmissionMessages.generated.h"
+
+// IMPORTANT: This peer-to-peer workflow is subject to be replaced, possibly in 5.5, by a server feature.
+// The peer-to-peer approach cannot guarantee atomicity if client A is trying to edit clients B and C as one operation: How to revert when B succeeds and C fails? You can't.
+
+/** A kind request from one client to another client to change their server registered state. The receiver can reject the request. */
+USTRUCT()
+struct FMultiUser_ChangeRemote_Request
+{
+	GENERATED_BODY()
+
+	/** Stream changes to apply. Applied before authority changes. */
+	UPROPERTY()
+	FConcertReplication_ChangeStream_Request StreamChangeRequest;
+
+	/** Authority changes to apply. Applied after stream changes. */
+	UPROPERTY()
+	FConcertReplication_ChangeAuthority_Request AuthorityRequest;
+};
+
+UENUM()
+enum class EMultiUserChangeRemoteRequestError : uint8
+{
+	/** The requests will be executed. */
+	Accepted,
+
+	/** The client rejected the change. */
+	RejectedChange,
+	
+	/** The request was invalid when it was received because the client determined that the change would generate a conflict given the client's known server state. */
+	PredictedConflict,
+	/** The change was rejected because the client was progressing a request from another client. */
+	OtherInProgress
+};
+
+/** Indicates whether the request was enqueued. */
+USTRUCT()
+struct FMultiUser_ChangeRemote_Response
+{
+	GENERATED_BODY()
+
+	/** Default result for default constructed response. Concert default constructs timed out responses. */
+	UPROPERTY()
+	bool bTimeout = true;
+	
+	/** Indicates whether the request was enqueued */
+	UPROPERTY()
+	EMultiUserChangeRemoteRequestError Error = EMultiUserChangeRemoteRequestError::Accepted;
+};
+
+/** Sent by serving client to the requesting client to confirm completion of the stream part of the request. */
+USTRUCT()
+struct FMultiUser_ChangeRemote_StreamUpdatedEvent
+{
+	GENERATED_BODY()
+
+	/** Whether the change was submitted to the server. False indicates that there was some kind of submission failure. True does NOT imply that StreamChangeResponse did not fail. */
+	UPROPERTY()
+	bool bProcessedSuccessfully = true;
+
+	/** Set if bProcessedSuccessfully == true. */
+	UPROPERTY()
+	FConcertReplication_ChangeStream_Response StreamChangeResponse;
+};
+
+/** Sent by serving client to the requesting client to confirm completion of the authority part of the request. */
+USTRUCT()
+struct FMultiUser_ChangeRemote_AuthorityUpdatedEvent
+{
+	GENERATED_BODY()
+	
+	/** Whether the change was submitted to the server. False indicates that there was some kind of submission failure. True does NOT imply that AuthorityChangeResponse did not fail. */
+	UPROPERTY()
+	bool bProcessedSuccessfully = true;
+
+	/** Set if bProcessedSuccessfully == true. */
+	UPROPERTY()
+	FConcertReplication_ChangeAuthority_Response AuthorityChangeResponse;
+};
+
+UENUM()
+enum class EMultiUserCancelRemoteChangeReason : uint8
+{
+	/**
+	 * Internal error: SubmitChanges failed.
+	 * No changes were made.
+	 */
+	FailedToCreate,
+
+	/**
+	 * The serving client had to wait with executing the event and after so doing determined that the request would a conflict. 
+	 * No changes were made.
+	 */
+	PredictedConflict
+};
+
+/** Send by serving client to the requesting client to cancel the request after it was confirmed. */
+USTRUCT()
+struct FMultiUser_ChangeRemote_Cancelled
+{
+	GENERATED_BODY()
+
+	UPROPERTY()
+	EMultiUserCancelRemoteChangeReason Reason;
+};
\ No newline at end of file
diff --git a/Engine/Plugins/Developer/Concert/ConcertApp/MultiUserClient/Source/MultiUserClient/Private/Replication/Submission/Remote/SubmissionWorkflow_RemoteClient.cpp b/Engine/Plugins/Developer/Concert/ConcertApp/MultiUserClient/Source/MultiUserClient/Private/Replication/Submission/Remote/SubmissionWorkflow_RemoteClient.cpp
new file mode 100644
--- /dev/null
+++ b/Engine/Plugins/Developer/Concert/ConcertApp/MultiUserClient/Source/MultiUserClient/Private/Replication/Submission/Remote/SubmissionWorkflow_RemoteClient.cpp
@@ -0,0 +1,215 @@
+// Copyright Epic Games, Inc. All Rights Reserved.
+
+#include "SubmissionWorkflow_RemoteClient.h"
+
+#include "ConcertLogGlobal.h"
+#include "IConcertSession.h"
+#include "RemoteSubmissionMessages.h"
+#include "Replication/Submission/SingleClientSubmissionOperation.h"
+#include "Widgets/ActiveSession/Replication/Client/ClientUtils.h"
+
+namespace UE::MultiUserClient
+{
+	FSubmissionWorkflow_RemoteClient::FSubmissionWorkflow_RemoteClient(
+		TSharedRef<IConcertClientSession> InConcertSession,
+		const FGuid& InRemoteClientEndpointId
+		)
+		: RemoteClientEndpointId(InRemoteClientEndpointId)
+		, ConcertSession(MoveTemp(InConcertSession))
+	{
+		// The other client will update as progress is made.
+		// Note that handles cannot be registered & unregistered dynamically while the handler is executing (that would change the handler array while it is being iterated).
+		ConcertSession->RegisterCustomEventHandler<FMultiUser_ChangeRemote_StreamUpdatedEvent>(this, &FSubmissionWorkflow_RemoteClient::OnStreamRemoteChangeEvent);
+		ConcertSession->RegisterCustomEventHandler<FMultiUser_ChangeRemote_AuthorityUpdatedEvent>(this, &FSubmissionWorkflow_RemoteClient::OnAuthorityRemoteChangeEvent);
+	}
+
+	FSubmissionWorkflow_RemoteClient::~FSubmissionWorkflow_RemoteClient()
+	{
+		ConcertSession->UnregisterCustomEventHandler<FMultiUser_ChangeRemote_StreamUpdatedEvent>(this);
+		ConcertSession->UnregisterCustomEventHandler<FMultiUser_ChangeRemote_AuthorityUpdatedEvent>(this);
+		CleanUpSubmissionProcess();
+	}
+
+	TSharedPtr<ISubmissionOperation> FSubmissionWorkflow_RemoteClient::SubmitChanges(FSubmissionParams Params)
+	{
+		if (Params.IsEmpty() || !CanSubmit())
+		{
+			return nullptr;
+		}
+		
+		const bool bIsModifyingStream = !Params.IsStreamChangeEmpty();
+		const TSharedRef<FSingleClientSubmissionOperation> OperationResult = MakeShared<FSingleClientSubmissionOperation>(bIsModifyingStream);
+
+		// Since SendCustomRequest can return immediately, emplace relevant promises before 
+		if (Params.IsStreamChangeEmpty())
+		{
+			const FSubmitStreamChangesResponse StreamResponse { EStreamSubmissionErrorCode::NoChange };
+			OperationResult->EmplaceStreamPromise(StreamResponse);
+			StreamRequestCompletedDelegate.Broadcast(StreamResponse);
+		}
+		// Opposed to the local case, sending authority request is sent instantly in the remote case, so emplace the promise now.
+		OperationResult->EmplaceAuthorityRequestPromise({ EAuthoritySubmissionRequestErrorCode::Success, Params.AuthorityRequest });
+		
+		const FMultiUser_ChangeRemote_Request Request
+		{
+			Params.StreamRequest ? *Params.StreamRequest : FConcertReplication_ChangeStream_Request{},
+			Params.AuthorityRequest ? *Params.AuthorityRequest : FConcertReplication_ChangeAuthority_Request{},
+		};
+		InProgressOperation = { MoveTemp(Params), OperationResult };
+		
+		ConcertSession->SendCustomRequest<FMultiUser_ChangeRemote_Request, FMultiUser_ChangeRemote_Response>(Request, RemoteClientEndpointId)
+			.Next([this, WeakToken = LifetimeToken.ToWeakPtr()](FMultiUser_ChangeRemote_Response&& Response)
+			{
+				if (WeakToken.Pin())
+				{
+					ProcessSubmissionResponse(Response);
+				}
+			});
+		
+		return OperationResult;
+	}
+
+	EChangeUploadability FSubmissionWorkflow_RemoteClient::GetUploadability() const
+	{
+		if (InProgressOperation.IsSet())
+		{
+			return EChangeUploadability::InProgress;
+		}
+
+		// TODO DP UE-200927: Figure out whether remote client is configured to reject changes
+		return EChangeUploadability::Ready;
+	}
+
+	void FSubmissionWorkflow_RemoteClient::ProcessSubmissionResponse(const FMultiUser_ChangeRemote_Response& Response)
+	{
+		if (Response.bTimeout)
+		{
+			TimeoutAndClean();
+		}
+		else if (Response.Error != EMultiUserChangeRemoteRequestError::Accepted)
+		{
+			HandleFailureResponse(Response.Error);
+		}
+
+		//
+		UE_LOG(LogConcert, Log, TEXT("Remote client %s accepted change request"), *GetRemoteClientName());
+	}
+	
+	void FSubmissionWorkflow_RemoteClient::TimeoutAndClean()
+	{
+		UE_LOG(LogConcert, Warning, TEXT("Change request to client %s timed out"), *GetRemoteClientName());
+		
+		// SubmitChanges sets the
+		if (!InProgressOperation->Parameters.IsStreamChangeEmpty())
+		{
+			const FSubmitStreamChangesResponse StreamResponse { EStreamSubmissionErrorCode::Timeout };
+			GetOperation().EmplaceStreamPromise(StreamResponse);
+			StreamRequestCompletedDelegate.Broadcast(StreamResponse);
+		}
+
+		// Sending authority request was successful.
+		const FSubmitAuthorityChangesRequest AuthorityRequest { EAuthoritySubmissionRequestErrorCode::Success, InProgressOperation->Parameters.AuthorityRequest };
+		const FSubmitAuthorityChangesResponse AuthorityResponse { EAuthoritySubmissionResponseErrorCode::Timeout };
+		GetOperation().EmplaceAuthorityResponsePromise(AuthorityResponse);
+		AuthorityRequestCompletedDelegate.Broadcast(AuthorityRequest, AuthorityResponse);
+			
+		CleanUpSubmissionProcess();
+	}
+
+	void FSubmissionWorkflow_RemoteClient::HandleFailureResponse(EMultiUserChangeRemoteRequestError ErrorFlags)
+	{
+		UE_LOG(LogConcert, Warning, TEXT("Change request to client %s failed (error flags: %d)"), *GetRemoteClientName(), static_cast<int32>(ErrorFlags));
+
+		switch (ErrorFlags)
+		{
+		case EMultiUserChangeRemoteRequestError::PredictedConflict:
+			// TODO DP UE-200926: Retry when remote client predicts conflict, if applicable
+		case EMultiUserChangeRemoteRequestError::OtherInProgress:
+			// TODO DP UE-200924: Retry when client is busy serving another request
+		case EMultiUserChangeRemoteRequestError::RejectedChange:
+			CleanUpSubmissionProcess();
+			break;
+			
+		case EMultiUserChangeRemoteRequestError::Accepted:
+		default:
+			checkNoEntry();
+		}
+	}
+
+	void FSubmissionWorkflow_RemoteClient::CleanUpSubmissionProcess()
+	{
+		InProgressOperation.Reset();
+		OnSubmitOperationCompletedDelegate.Broadcast();
+	}
+
+	void FSubmissionWorkflow_RemoteClient::OnStreamRemoteChangeEvent(const FConcertSessionContext& Context, const FMultiUser_ChangeRemote_StreamUpdatedEvent& EventData)
+	{
+		// Might receive from other clients due to other registered FSubmissionWorkflow_RemoteClients on this same editor instance
+		if (!InProgressOperation.IsSet() ||Context.SourceEndpointId != RemoteClientEndpointId)
+		{
+			return;
+		}
+
+		FSingleClientSubmissionOperation& Operation = GetOperation();
+		if (EventData.bProcessedSuccessfully)
+		{
+			const FCompletedChangeSubmission StreamChangeSubmission { *InProgressOperation->Parameters.StreamRequest, { EventData.StreamChangeResponse } };
+			const FSubmitStreamChangesResponse StreamResponse { EStreamSubmissionErrorCode::Success, StreamChangeSubmission };
+			Operation.EmplaceStreamPromise(StreamResponse);
+			StreamRequestCompletedDelegate.Broadcast(StreamResponse);
+		}
+		else
+		{
+			const FSubmitStreamChangesResponse StreamResponse{ EStreamSubmissionErrorCode::Timeout };
+			Operation.EmplaceStreamPromise(StreamResponse);
+			StreamRequestCompletedDelegate.Broadcast(StreamResponse);
+		}
+		
+		// TODO UE-200925: Display better error handling to user
+		UE_CLOG(!EventData.bProcessedSuccessfully, LogConcert, Error, TEXT("Remote stream change to client %s failed"), *GetRemoteClientName());
+		FinishSubmissionIfDone();
+	}
+
+	void FSubmissionWorkflow_RemoteClient::OnAuthorityRemoteChangeEvent(const FConcertSessionContext& Context, const FMultiUser_ChangeRemote_AuthorityUpdatedEvent& EventData)
+	{
+		// Might receive from other clients due to other registered FSubmissionWorkflow_RemoteClients on this same editor instance
+		if (!InProgressOperation.IsSet() || Context.SourceEndpointId != RemoteClientEndpointId)
+		{
+			return;
+		}
+		
+		FSingleClientSubmissionOperation& Operation = GetOperation();
+		const FSubmitAuthorityChangesRequest AuthorityRequest { EAuthoritySubmissionRequestErrorCode::Success, InProgressOperation->Parameters.AuthorityRequest };
+		if (EventData.bProcessedSuccessfully)
+		{
+			const ConcertSyncClient::Replication::FAuthorityChangeResponse ReceivedResponse { EventData.AuthorityChangeResponse };
+			const FSubmitAuthorityChangesResponse AuthorityResponse { EAuthoritySubmissionResponseErrorCode::Success, ReceivedResponse };
+			Operation.EmplaceAuthorityResponsePromise(AuthorityResponse);
+			AuthorityRequestCompletedDelegate.Broadcast(AuthorityRequest, AuthorityResponse);
+		}
+		else
+		{
+			const FSubmitAuthorityChangesResponse AuthorityResponse { EAuthoritySubmissionResponseErrorCode::Timeout };
+			Operation.EmplaceAuthorityResponsePromise(AuthorityResponse);
+			AuthorityRequestCompletedDelegate.Broadcast(AuthorityRequest, AuthorityResponse);
+		}
+
+		FinishSubmissionIfDone();
+	}
+
+	void FSubmissionWorkflow_RemoteClient::FinishSubmissionIfDone()
+	{
+		FSingleClientSubmissionOperation& Operation = GetOperation();
+		const bool bIsStreamChangeDone = Operation.HasSetStreamPromise() || InProgressOperation->Parameters.StreamRequest->IsEmpty();
+		const bool bIsAuthorityChangeDone = Operation.HasSetAuthorityRequestPromise() || InProgressOperation->Parameters.AuthorityRequest->IsEmpty();
+		if (bIsStreamChangeDone && bIsAuthorityChangeDone)
+		{
+			CleanUpSubmissionProcess();
+		}
+	}
+
+	FString FSubmissionWorkflow_RemoteClient::GetRemoteClientName() const
+	{
+		return ClientUtils::GetClientDisplayName(*ConcertSession, RemoteClientEndpointId);
+	}
+}
diff --git a/Engine/Plugins/Developer/Concert/ConcertApp/MultiUserClient/Source/MultiUserClient/Private/Replication/Submission/Remote/SubmissionWorkflow_RemoteClient.h b/Engine/Plugins/Developer/Concert/ConcertApp/MultiUserClient/Source/MultiUserClient/Private/Replication/Submission/Remote/SubmissionWorkflow_RemoteClient.h
new file mode 100644
--- /dev/null
+++ b/Engine/Plugins/Developer/Concert/ConcertApp/MultiUserClient/Source/MultiUserClient/Private/Replication/Submission/Remote/SubmissionWorkflow_RemoteClient.h
@@ -0,0 +1,81 @@
+// Copyright Epic Games, Inc. All Rights Reserved.
+
+#pragma once
+
+#include "IConcertSessionHandler.h"
+#include "RemoteSubmissionMessages.h"
+#include "Replication/IToken.h"
+#include "Replication/Submission/ISubmissionWorkflow.h"
+
+#include "Misc/Optional.h"
+#include "Templates/UnrealTemplate.h"
+
+class IConcertClientSession;
+class IConcertSyncClient;
+
+enum class EConcertClientStatus : uint8;
+
+namespace UE::MultiUserClient
+{
+	class FSingleClientSubmissionOperation;
+	
+	/**
+	 * Sends change requests to a remote client, whose FRemoteSubmissionListener will handle the requests.
+	 * @see FRemoteSubmissionListener
+	 */
+	class FSubmissionWorkflow_RemoteClient
+		: public FSubmissionWorkflowBase
+		, public FNoncopyable
+	{
+	public:
+		
+		FSubmissionWorkflow_RemoteClient(TSharedRef<IConcertClientSession> InConcertSession, const FGuid& InRemoteClientEndpointId);
+		virtual ~FSubmissionWorkflow_RemoteClient() override;
+
+		//~ Begin ISubmissionWorkflow Interface
+		virtual TSharedPtr<ISubmissionOperation> SubmitChanges(FSubmissionParams Params) override;
+		virtual EChangeUploadability GetUploadability() const override;
+		//~ End ISubmissionWorkflow Interface
+
+	private:
+		
+		/** Latent operations keep a weak reference to this to detect whether this FSubmissionWorkflow_LocalClient was destroyed. */
+		const TSharedRef<FToken> LifetimeToken = FToken::Make();
+
+		/** Id of the remote client that requests are sent to. */
+		const FGuid RemoteClientEndpointId;
+		
+		/** Used to sending messages */
+		const TSharedRef<IConcertClientSession> ConcertSession;
+
+		struct FOperationData
+		{
+			/** Parameters the operation was started with. */
+			const FSubmissionParams Parameters;
+			/** Automatically cancels pending promises when destroyed. */
+			const TSharedRef<FSingleClientSubmissionOperation> ExposedOperation;
+		};
+		/** Set for as long as there is a SubmitChanges operation in progress. */
+		TOptional<FOperationData> InProgressOperation;
+
+		/** Handles a response from a remote client. */
+		void ProcessSubmissionResponse(const FMultiUser_ChangeRemote_Response& Response);
+		
+		/** Handles timeout of submission response. */
+		void TimeoutAndClean();
+		/** Handles a response from the client that indicates failure. */
+		void HandleFailureResponse(EMultiUserChangeRemoteRequestError ErrorFlags);
+		/** Cleans up side effects from last submission process */
+		void CleanUpSubmissionProcess();
+
+		FSingleClientSubmissionOperation& GetOperation() { return *InProgressOperation->ExposedOperation; }
+
+		/** Called by remote client to update us about progress. */
+		void OnStreamRemoteChangeEvent(const FConcertSessionContext& Context, const FMultiUser_ChangeRemote_StreamUpdatedEvent& EventData);
+		void OnAuthorityRemoteChangeEvent(const FConcertSessionContext& Context, const FMultiUser_ChangeRemote_AuthorityUpdatedEvent& EventData);
+		void FinishSubmissionIfDone();
+
+		FString GetRemoteClientName() const;
+	};
+}
+
diff --git a/Engine/Plugins/Developer/Concert/ConcertApp/MultiUserClient/Source/MultiUserClient/Private/Replication/Submission/SingleClientSubmissionOperation.h b/Engine/Plugins/Developer/Concert/ConcertApp/MultiUserClient/Source/MultiUserClient/Private/Replication/Submission/SingleClientSubmissionOperation.h
--- a/Engine/Plugins/Developer/Concert/ConcertApp/MultiUserClient/Source/MultiUserClient/Private/Replication/Submission/SingleClientSubmissionOperation.h
+++ b/Engine/Plugins/Developer/Concert/ConcertApp/MultiUserClient/Source/MultiUserClient/Private/Replication/Submission/SingleClientSubmissionOperation.h
@@ -52,6 +52,10 @@ namespace UE::MultiUserClient
 			AuthorityChangeResponsePromise.EmplaceValue(MoveTemp(Result));
 		}
 
+		bool HasSetStreamPromise() const { return bStreamPromiseWasSet; }
+		bool HasSetAuthorityRequestPromise() const { return bAuthorityRequestPromiseWasSet; }
+		bool HasSetAuthorityResponsePromise() const { return bAuthorityResponsePromiseWasSet; }
+		
 		//~ Begin ISubmissionOperation Interface
 		virtual bool IsModifyingStreams() const override { return bModifiesStreams; }
 		virtual TFuture<FSubmitStreamChangesResponse> OnStreamChangesSubmittedFuture() override { return StreamChangesPromise.GetFuture();  }
diff --git a/Engine/Plugins/Developer/Concert/ConcertApp/MultiUserClient/Source/MultiUserClient/Private/Replication/Submission/SubmissionWorkflow_LocalClient.cpp b/Engine/Plugins/Developer/Concert/ConcertApp/MultiUserClient/Source/MultiUserClient/Private/Replication/Submission/SubmissionWorkflow_LocalClient.cpp
--- a/Engine/Plugins/Developer/Concert/ConcertApp/MultiUserClient/Source/MultiUserClient/Private/Replication/Submission/SubmissionWorkflow_LocalClient.cpp
+++ b/Engine/Plugins/Developer/Concert/ConcertApp/MultiUserClient/Source/MultiUserClient/Private/Replication/Submission/SubmissionWorkflow_LocalClient.cpp
@@ -11,6 +11,11 @@ namespace UE::MultiUserClient
 		: Client(MoveTemp(InClient))
 	{}
 
+	FSubmissionWorkflow_LocalClient::~FSubmissionWorkflow_LocalClient()
+	{
+		
+	}
+
 	EChangeUploadability FSubmissionWorkflow_LocalClient::GetUploadability() const
 	{
 		const bool bOperationInProgress = InProgressOperation.IsSet(); 
@@ -21,20 +26,19 @@ namespace UE::MultiUserClient
 	{
 		using namespace ConcertSyncClient::Replication;
 		
-		const TOptional<FChangeStreamRequest>& OptionalStreamRequest = Params.StreamRequest;
-		TOptional<FAuthorityChangeRequest>& OptionalAuthorityRequest = Params.AuthorityRequest;
-		
 		IConcertClientReplicationManager* ReplicationManager = Client->GetReplicationManager();
 		if (Params.IsEmpty() || !CanSubmit() || !ensure(ReplicationManager))
 		{
 			return nullptr;
 		}
 		
+		const TOptional<FChangeStreamRequest>& OptionalStreamRequest = Params.StreamRequest;
+		TOptional<FAuthorityChangeRequest>& OptionalAuthorityRequest = Params.AuthorityRequest;
+		
 		const bool bIsStreamChangeEmpty = Params.IsStreamChangeEmpty();
 		const bool bModifyStreams = !bIsStreamChangeEmpty;
 		const TSharedRef<FSingleClientSubmissionOperation> Operation = MakeShared<FSingleClientSubmissionOperation>(bModifyStreams);
 		InProgressOperation = { Operation };
-		InProgressOperation->OnDestroy.BindLambda([this](){ OnSubmitOperationCompletedDelegate.Broadcast(); });
 		
 		if (bIsStreamChangeEmpty)
 		{
@@ -91,9 +95,9 @@ namespace UE::MultiUserClient
 			
 			Operation->EmplaceAuthorityRequestPromise(Request);
 			Operation->EmplaceAuthorityResponsePromise(Response);
-			InProgressOperation.Reset();
-			
 			AuthorityRequestCompletedDelegate.Broadcast(Request, Response);
+			
+			CleanUpSubmissionOperation();
 		}
 	}
 
@@ -103,8 +107,8 @@ namespace UE::MultiUserClient
 		IConcertClientReplicationManager* ReplicationManager = Client->GetReplicationManager();
 		if (!ensure(ReplicationManager))
 		{
-			InProgressOperation.Reset(); // Automatically cancels the pending promises
 			AuthorityRequestCompletedDelegate.Broadcast({ EAuthoritySubmissionRequestErrorCode::Cancelled }, { EAuthoritySubmissionResponseErrorCode::Cancelled });
+			CleanUpSubmissionOperation(); // Automatically cancels the pending promises
 			return;
 		}
 
@@ -117,9 +121,8 @@ namespace UE::MultiUserClient
 			
 			Operation->EmplaceAuthorityRequestPromise(Request);
 			Operation->EmplaceAuthorityResponsePromise(Response);
-			InProgressOperation.Reset();
-			
 			AuthorityRequestCompletedDelegate.Broadcast(Request, Response);
+			CleanUpSubmissionOperation();
 			return;
 		}
 		
@@ -136,10 +139,15 @@ namespace UE::MultiUserClient
 					const FSubmitAuthorityChangesResponse Result { ErrorCode, MoveTemp(Response) };
 					
 					InProgressOperation->Operation->EmplaceAuthorityResponsePromise(Result);
-					InProgressOperation.Reset();
-					
 					AuthorityRequestCompletedDelegate.Broadcast(Request, Result);
+					CleanUpSubmissionOperation();
 				}
 			});
 	}
+
+	void FSubmissionWorkflow_LocalClient::CleanUpSubmissionOperation()
+	{
+		InProgressOperation.Reset();
+		OnSubmitOperationCompletedDelegate.Broadcast();
+	}
 }
diff --git a/Engine/Plugins/Developer/Concert/ConcertApp/MultiUserClient/Source/MultiUserClient/Private/Replication/Submission/SubmissionWorkflow_LocalClient.h b/Engine/Plugins/Developer/Concert/ConcertApp/MultiUserClient/Source/MultiUserClient/Private/Replication/Submission/SubmissionWorkflow_LocalClient.h
--- a/Engine/Plugins/Developer/Concert/ConcertApp/MultiUserClient/Source/MultiUserClient/Private/Replication/Submission/SubmissionWorkflow_LocalClient.h
+++ b/Engine/Plugins/Developer/Concert/ConcertApp/MultiUserClient/Source/MultiUserClient/Private/Replication/Submission/SubmissionWorkflow_LocalClient.h
@@ -26,6 +26,7 @@ namespace UE::MultiUserClient
 	public:
 		
 		FSubmissionWorkflow_LocalClient(TSharedRef<IConcertSyncClient> InClient);
+		virtual ~FSubmissionWorkflow_LocalClient() override;
 		
 		//~ Begin ISubmissionWorkflow Interface
 		virtual TSharedPtr<ISubmissionOperation> SubmitChanges(FSubmissionParams Params) override;
@@ -43,16 +44,7 @@ namespace UE::MultiUserClient
 		struct FOperationData
 		{
 			TSharedRef<FSingleClientSubmissionOperation> Operation;
-
-			DECLARE_DELEGATE(FOnDestroy);
-			FOnDestroy OnDestroy;
-			
-			~FOperationData()
-			{
-				OnDestroy.ExecuteIfBound();
-			}
 		};
-		
 		/**
 		 * Set for as long as there is a SubmitChanges operation in progress.
 		 * Automatically cancels pending promises when destroyed.
@@ -66,6 +58,8 @@ namespace UE::MultiUserClient
 			TOptional<ConcertSyncClient::Replication::FAuthorityChangeRequest> AuthorityChangeRequest
 			);
 		void HandlePendingAuthorityChangeRequest(TOptional<ConcertSyncClient::Replication::FAuthorityChangeRequest> AuthorityChangeRequest);
+
+		void CleanUpSubmissionOperation();
 	};
 }
 
diff --git a/Engine/Plugins/Developer/Concert/ConcertApp/MultiUserClient/Source/MultiUserClient/Private/Replication/Submission/SubmissionWorkflow_RemoteClient.cpp b/Engine/Plugins/Developer/Concert/ConcertApp/MultiUserClient/Source/MultiUserClient/Private/Replication/Submission/SubmissionWorkflow_RemoteClient.cpp
deleted file mode 100644
diff --git a/Engine/Plugins/Developer/Concert/ConcertApp/MultiUserClient/Source/MultiUserClient/Private/Replication/Submission/SubmissionWorkflow_RemoteClient.h b/Engine/Plugins/Developer/Concert/ConcertApp/MultiUserClient/Source/MultiUserClient/Private/Replication/Submission/SubmissionWorkflow_RemoteClient.h
deleted file mode 100644
diff --git a/Engine/Plugins/Developer/Concert/ConcertApp/MultiUserClient/Source/MultiUserClient/Private/Replication/Util/GlobalAuthorityCache.cpp b/Engine/Plugins/Developer/Concert/ConcertApp/MultiUserClient/Source/MultiUserClient/Private/Replication/Util/GlobalAuthorityCache.cpp
--- a/Engine/Plugins/Developer/Concert/ConcertApp/MultiUserClient/Source/MultiUserClient/Private/Replication/Util/GlobalAuthorityCache.cpp
+++ b/Engine/Plugins/Developer/Concert/ConcertApp/MultiUserClient/Source/MultiUserClient/Private/Replication/Util/GlobalAuthorityCache.cpp
@@ -3,6 +3,7 @@
 #include "GlobalAuthorityCache.h"
 
 #include "Replication/AuthorityConflictSharedUtils.h"
+#include "Replication/Authority/EAuthorityMutability.h"
 #include "Replication/Client/RemoteReplicationClient.h"
 #include "Replication/Client/ReplicationClientManager.h"
 
@@ -65,12 +66,12 @@ namespace UE::MultiUserClient
 		return Result;
 	}
 
-	FGlobalAuthorityCache::ECanTakeAuthority FGlobalAuthorityCache::CanClientTakeAuthorityAfterSubmission(const FSoftObjectPath& Object, const FGuid& ClientId, FProcessPropertyConflict ProcessConflict) const
+	EAuthorityMutability FGlobalAuthorityCache::CanClientTakeAuthorityAfterSubmission(const FSoftObjectPath& Object, const FGuid& ClientId, FProcessPropertyConflict ProcessConflict) const
 	{
 		const FReplicationClient* Client = ClientManager.FindClient(ClientId);
 		if (!ensure(Client))
 		{
-			return ECanTakeAuthority::NotApplicable;
+			return EAuthorityMutability::NotApplicable;
 		}
 
 		// Important: Get server state with local changes applied to it! CanClientTakeAuthority answers: "Can the client take authority after submitting?" 
@@ -78,7 +79,7 @@ namespace UE::MultiUserClient
 		if (!PropertySelection)
 		{
 			// Nothing to take authority over
-			return ECanTakeAuthority::NotApplicable;
+			return EAuthorityMutability::NotApplicable;
 		}
 		
 		using namespace ConcertSyncCore::Replication::AuthorityConflictUtils;
@@ -91,7 +92,7 @@ namespace UE::MultiUserClient
 			{
 				return ProcessConflict(ClientId, ConflictingProperty);
 			});
-		return Conflict == EAuthorityConflict::Allowed ? ECanTakeAuthority::Allowed : ECanTakeAuthority::Conflict;
+		return Conflict == EAuthorityConflict::Allowed ? EAuthorityMutability::Allowed : EAuthorityMutability::Conflict;
 	}
 
 	bool FGlobalAuthorityCache::CanClientAddProperty(const FSoftObjectPath& Object, const FGuid& ClientId, const FConcertPropertyChain& Chain) const
diff --git a/Engine/Plugins/Developer/Concert/ConcertApp/MultiUserClient/Source/MultiUserClient/Private/Replication/Util/GlobalAuthorityCache.h b/Engine/Plugins/Developer/Concert/ConcertApp/MultiUserClient/Source/MultiUserClient/Private/Replication/Util/GlobalAuthorityCache.h
--- a/Engine/Plugins/Developer/Concert/ConcertApp/MultiUserClient/Source/MultiUserClient/Private/Replication/Util/GlobalAuthorityCache.h
+++ b/Engine/Plugins/Developer/Concert/ConcertApp/MultiUserClient/Source/MultiUserClient/Private/Replication/Util/GlobalAuthorityCache.h
@@ -43,17 +43,9 @@ namespace UE::MultiUserClient
 		void ForEachClientWithAuthorityOverObject(const FSoftObjectPath& Object, TFunctionRef<EBreakBehavior(const FGuid& ClientId)> Callback) const;
 		/** Util that uses ForEachClientWithAuthorityOverObject to make an array. */
 		TArray<FGuid> GetClientsWithAuthorityOverObject(const FSoftObjectPath& Object) const;
-
-		enum class ECanTakeAuthority
-		{
-			/** The specified client has no properties, does not exist, etc. */
-			NotApplicable,
-			Conflict,
-			Allowed
-		};
 		
 		/** @return Whether the given client can take authority over the object without causing any conflicts. This also considers the changes made to the stream after submission. */
-		ECanTakeAuthority CanClientTakeAuthorityAfterSubmission(const FSoftObjectPath& Object, const FGuid& ClientId, FProcessPropertyConflict ProcessConflict = [](auto&, auto&){ return EBreakBehavior::Break; }) const;
+		EAuthorityMutability CanClientTakeAuthorityAfterSubmission(const FSoftObjectPath& Object, const FGuid& ClientId, FProcessPropertyConflict ProcessConflict = [](auto&, auto&){ return EBreakBehavior::Break; }) const;
 		/** @return Whether the given client add the given property to the object without causing any conflicts. */
 		bool CanClientAddProperty(const FSoftObjectPath& Object, const FGuid& ClientId, const FConcertPropertyChain& Chain) const;
 
diff --git a/Engine/Plugins/Developer/Concert/ConcertApp/MultiUserClient/Source/MultiUserClient/Private/Widgets/ActiveSession/Replication/Client/ClientUtils.cpp b/Engine/Plugins/Developer/Concert/ConcertApp/MultiUserClient/Source/MultiUserClient/Private/Widgets/ActiveSession/Replication/Client/ClientUtils.cpp
--- a/Engine/Plugins/Developer/Concert/ConcertApp/MultiUserClient/Source/MultiUserClient/Private/Widgets/ActiveSession/Replication/Client/ClientUtils.cpp
+++ b/Engine/Plugins/Developer/Concert/ConcertApp/MultiUserClient/Source/MultiUserClient/Private/Widgets/ActiveSession/Replication/Client/ClientUtils.cpp
@@ -12,20 +12,25 @@ namespace UE::MultiUserClient::ClientUtils
 	FString GetClientDisplayName(const IConcertClient& InLocalClientInstance, const FGuid& InClientEndpointId)
 	{
 		const TSharedPtr<IConcertClientSession> Session = InLocalClientInstance.GetCurrentSession();
-		const bool bIsLocalClient = ensure(Session) && Session->GetSessionClientEndpointId() == InClientEndpointId;
+		return ensure(Session) ? GetClientDisplayName(*Session, InClientEndpointId) : FString{};
+	}
+	
+	FString GetClientDisplayName(const IConcertClientSession& InSession, const FGuid& InClientEndpointId)
+	{
+		const bool bIsLocalClient = InSession.GetSessionClientEndpointId() == InClientEndpointId;
 		if (bIsLocalClient)
 		{
-			return ConcertClientSharedSlate::SClientName::GetDisplayText(Session->GetLocalClientInfo(), bIsLocalClient).ToString();
+			return ConcertClientSharedSlate::SClientName::GetDisplayText(InSession.GetLocalClientInfo(), bIsLocalClient).ToString();
 		}
 
 		FConcertSessionClientInfo ClientInfo;
-		if (ensure(Session) && Session->FindSessionClient(InClientEndpointId, ClientInfo))
+		if (InSession.FindSessionClient(InClientEndpointId, ClientInfo))
 		{
 			return ConcertClientSharedSlate::SClientName::GetDisplayText(ClientInfo.ClientInfo, bIsLocalClient).ToString();
 		}
 
 		ensureMsgf(false, TEXT("Bad args"));
-		return {};
+		return InClientEndpointId.ToString(EGuidFormats::DigitsWithHyphens);
 	}
 	
 	bool GetClientDisplayInfo(const IConcertClient& InLocalClientInstance, const FGuid& InClientEndpointId, FConcertClientInfo& OutClientInfo)
diff --git a/Engine/Plugins/Developer/Concert/ConcertApp/MultiUserClient/Source/MultiUserClient/Private/Widgets/ActiveSession/Replication/Client/ClientUtils.h b/Engine/Plugins/Developer/Concert/ConcertApp/MultiUserClient/Source/MultiUserClient/Private/Widgets/ActiveSession/Replication/Client/ClientUtils.h
--- a/Engine/Plugins/Developer/Concert/ConcertApp/MultiUserClient/Source/MultiUserClient/Private/Widgets/ActiveSession/Replication/Client/ClientUtils.h
+++ b/Engine/Plugins/Developer/Concert/ConcertApp/MultiUserClient/Source/MultiUserClient/Private/Widgets/ActiveSession/Replication/Client/ClientUtils.h
@@ -6,6 +6,8 @@
 
 class FString;
 class IConcertClient;
+class IConcertClientSession;
+
 struct FConcertClientInfo;
 struct FGuid;
 
@@ -25,6 +27,7 @@ namespace UE::MultiUserClient::ClientUtils
 	 * @return The display name or empty
 	 */
 	FString GetClientDisplayName(const IConcertClient& InLocalClientInstance, const FGuid& InClientEndpointId);
+	FString GetClientDisplayName(const IConcertClientSession& InSession, const FGuid& InClientEndpointId);
 
 	/**
 	 * Gets the display info for a given client
diff --git a/Engine/Plugins/Developer/Concert/ConcertApp/MultiUserClient/Source/MultiUserClient/Private/Widgets/ActiveSession/Replication/Client/Multi/MultiStreamModel.cpp b/Engine/Plugins/Developer/Concert/ConcertApp/MultiUserClient/Source/MultiUserClient/Private/Widgets/ActiveSession/Replication/Client/Multi/MultiStreamModel.cpp
--- a/Engine/Plugins/Developer/Concert/ConcertApp/MultiUserClient/Source/MultiUserClient/Private/Widgets/ActiveSession/Replication/Client/Multi/MultiStreamModel.cpp
+++ b/Engine/Plugins/Developer/Concert/ConcertApp/MultiUserClient/Source/MultiUserClient/Private/Widgets/ActiveSession/Replication/Client/Multi/MultiStreamModel.cpp
@@ -63,6 +63,9 @@ namespace UE::MultiUserClient
 		ClientSelectionModel.ForEachSelectedClient([this, &ReadOnlyClients, &WritableClients](FReplicationClient& Client)
 		{
 			const bool bIsUploadable = CanEverSubmit(Client.GetSubmissionWorkflow().GetUploadability());
+			const TSharedRef<ConcertClientSharedSlate::IEditableReplicationStreamModel> Stream = Client.GetClientEditModel();
+			Client.OnModelChanged().AddRaw(this, &FMultiStreamModel::OnStreamExternallyChanged, Stream.ToWeakPtr());
+			
 			if (bIsUploadable)
 			{
 				WritableClients.Add(&Client);
@@ -70,8 +73,6 @@ namespace UE::MultiUserClient
 			else
 			{
 				ReadOnlyClients.Add(&Client);
-				const TSharedRef<ConcertClientSharedSlate::IEditableReplicationStreamModel> Stream = Client.GetClientEditModel();
-				Client.OnModelChanged().AddRaw(this, &FMultiStreamModel::OnReadOnlyStreamChanged, Stream.ToWeakPtr());
 			}
 			
 			return EBreakBehavior::Continue;
@@ -94,7 +95,7 @@ namespace UE::MultiUserClient
 		}
 	}
 
-	void FMultiStreamModel::OnReadOnlyStreamChanged(TWeakPtr<ConcertClientSharedSlate::IEditableReplicationStreamModel> ChangedStream)
+	void FMultiStreamModel::OnStreamExternallyChanged(TWeakPtr<ConcertClientSharedSlate::IEditableReplicationStreamModel> ChangedStream)
 	{
 		if (const TSharedPtr<ConcertClientSharedSlate::IEditableReplicationStreamModel> ChangedStreamPin = ChangedStream.Pin())
 		{
diff --git a/Engine/Plugins/Developer/Concert/ConcertApp/MultiUserClient/Source/MultiUserClient/Private/Widgets/ActiveSession/Replication/Client/Multi/MultiStreamModel.h b/Engine/Plugins/Developer/Concert/ConcertApp/MultiUserClient/Source/MultiUserClient/Private/Widgets/ActiveSession/Replication/Client/Multi/MultiStreamModel.h
--- a/Engine/Plugins/Developer/Concert/ConcertApp/MultiUserClient/Source/MultiUserClient/Private/Widgets/ActiveSession/Replication/Client/Multi/MultiStreamModel.h
+++ b/Engine/Plugins/Developer/Concert/ConcertApp/MultiUserClient/Source/MultiUserClient/Private/Widgets/ActiveSession/Replication/Client/Multi/MultiStreamModel.h
@@ -30,7 +30,7 @@ namespace UE::MultiUserClient
 		//~ Begin IEditableMultiReplicationStreamModel Interface
 		virtual TSet<TSharedRef<ConcertClientSharedSlate::IReplicationStreamModel>> GetReadOnlyStreams() const override;
 		virtual TSet<TSharedRef<ConcertClientSharedSlate::IEditableReplicationStreamModel>> GetEditableStreams() const override;
-		virtual FOnReadOnlyStreamChanged& OnReadOnlyStreamChanged() override { return OnReadOnlyStreamChangedDelegate; }
+		virtual FOnStreamExternallyChanged& OnStreamExternallyChanged() override { return OnReadOnlyStreamChangedDelegate; }
 		virtual FOnStreamSetChanged& OnStreamSetChanged() override { return OnStreamSetChangedDelegate; }
 		//~ End IEditableMultiReplicationStreamModel Interface
 
@@ -44,13 +44,13 @@ namespace UE::MultiUserClient
 		TSet<const FReplicationClient*> CachedReadOnlyClients;
 		TSet<const FReplicationClient*> CachedWritableClients;
 
-		FOnReadOnlyStreamChanged OnReadOnlyStreamChangedDelegate;
+		FOnStreamExternallyChanged OnReadOnlyStreamChangedDelegate;
 		FOnStreamSetChanged OnStreamSetChangedDelegate;
 		
 		void RebuildStreamsSets();
 
 		/** Handle read-only streams changing */
-		void OnReadOnlyStreamChanged(TWeakPtr<ConcertClientSharedSlate::IEditableReplicationStreamModel> ChangedStream);
+		void OnStreamExternallyChanged(TWeakPtr<ConcertClientSharedSlate::IEditableReplicationStreamModel> ChangedStream);
 	};
 }
 
diff --git a/Engine/Plugins/Developer/Concert/ConcertApp/MultiUserClient/Source/MultiUserClient/Private/Widgets/ActiveSession/Replication/Client/Single/Columns/SingleClientColumns.h b/Engine/Plugins/Developer/Concert/ConcertApp/MultiUserClient/Source/MultiUserClient/Private/Widgets/ActiveSession/Replication/Client/Single/Columns/SingleClientColumns.h
--- a/Engine/Plugins/Developer/Concert/ConcertApp/MultiUserClient/Source/MultiUserClient/Private/Widgets/ActiveSession/Replication/Client/Single/Columns/SingleClientColumns.h
+++ b/Engine/Plugins/Developer/Concert/ConcertApp/MultiUserClient/Source/MultiUserClient/Private/Widgets/ActiveSession/Replication/Client/Single/Columns/SingleClientColumns.h
@@ -41,7 +41,6 @@ namespace UE::MultiUserClient::SingleClientColumns
 
 	/********** Toggle authority **********/
 	extern const FName ToggleTopLevelAuthorityColumnId;
-	extern const FName ToggleSubobjectAuthorityColumnId;
 
 	/**
 	 * Checkbox placed in the top-level view.
diff --git a/Engine/Plugins/Developer/Concert/ConcertApp/MultiUserClient/Source/MultiUserClient/Private/Widgets/ActiveSession/Replication/Client/Single/Columns/SingleClientColumns_ConflictWarning.cpp b/Engine/Plugins/Developer/Concert/ConcertApp/MultiUserClient/Source/MultiUserClient/Private/Widgets/ActiveSession/Replication/Client/Single/Columns/SingleClientColumns_ConflictWarning.cpp
--- a/Engine/Plugins/Developer/Concert/ConcertApp/MultiUserClient/Source/MultiUserClient/Private/Widgets/ActiveSession/Replication/Client/Single/Columns/SingleClientColumns_ConflictWarning.cpp
+++ b/Engine/Plugins/Developer/Concert/ConcertApp/MultiUserClient/Source/MultiUserClient/Private/Widgets/ActiveSession/Replication/Client/Single/Columns/SingleClientColumns_ConflictWarning.cpp
@@ -3,6 +3,7 @@
 #include "SingleClientColumns.h"
 
 #include "IConcertClient.h"
+#include "Replication/Authority/EAuthorityMutability.h"
 #include "Replication/Editor/Model/ReplicatedObjectData.h"
 #include "Replication/Editor/View/IReplicationStreamViewer.h"
 #include "Replication/Util/GlobalAuthorityCache.h"
@@ -34,22 +35,22 @@ namespace UE::MultiUserClient::SingleClientColumns
 				{
 					const auto GetVisibility = [&InAuthorityCache, ClientId, ObjectPath = InArgs.RowData.GetObjectPath()]()
 					{
-						const FGlobalAuthorityCache::ECanTakeAuthority TakeAuthorityResult = InAuthorityCache.CanClientTakeAuthorityAfterSubmission(ObjectPath, ClientId);
-						return TakeAuthorityResult == FGlobalAuthorityCache::ECanTakeAuthority::Conflict
+						const EAuthorityMutability TakeAuthorityResult = InAuthorityCache.CanClientTakeAuthorityAfterSubmission(ObjectPath, ClientId);
+						return TakeAuthorityResult == EAuthorityMutability::Conflict
 							? EVisibility::Visible
 							: EVisibility::Collapsed;
 					};
 					const auto GetToolTip = [InClient = MoveTemp(InClient), &InAuthorityCache, ClientId, ObjectPath = InArgs.RowData.GetObjectPath()]()
 					{
 						TSet<FString> ClientNames;
-						const FGlobalAuthorityCache::ECanTakeAuthority TakeAuthorityResult = InAuthorityCache.CanClientTakeAuthorityAfterSubmission(ObjectPath, ClientId,
+						const EAuthorityMutability TakeAuthorityResult = InAuthorityCache.CanClientTakeAuthorityAfterSubmission(ObjectPath, ClientId,
 						[InClient, &ClientNames](const FGuid& ClientId, const FConcertPropertyChain& ConflictingProperty)
 						{
 							ClientNames.Add(ClientUtils::GetClientDisplayName(*InClient, ClientId));
 							return EBreakBehavior::Continue;
 						});
 						
-						return TakeAuthorityResult != FGlobalAuthorityCache::ECanTakeAuthority::Conflict
+						return TakeAuthorityResult != EAuthorityMutability::Conflict
 							? FText::GetEmpty()
 							: FText::Format(
 								LOCTEXT("Subobject.WarningFmt", "{0} {1}|plural(one=is,other=are) replicating some of the assigned properties. This change will not be sent to the server."),
diff --git a/Engine/Plugins/Developer/Concert/ConcertApp/MultiUserClient/Source/MultiUserClient/Private/Widgets/ActiveSession/Replication/Client/Single/Columns/SingleClientColumns_ToggleAuthority.cpp b/Engine/Plugins/Developer/Concert/ConcertApp/MultiUserClient/Source/MultiUserClient/Private/Widgets/ActiveSession/Replication/Client/Single/Columns/SingleClientColumns_ToggleAuthority.cpp
--- a/Engine/Plugins/Developer/Concert/ConcertApp/MultiUserClient/Source/MultiUserClient/Private/Widgets/ActiveSession/Replication/Client/Single/Columns/SingleClientColumns_ToggleAuthority.cpp
+++ b/Engine/Plugins/Developer/Concert/ConcertApp/MultiUserClient/Source/MultiUserClient/Private/Widgets/ActiveSession/Replication/Client/Single/Columns/SingleClientColumns_ToggleAuthority.cpp
@@ -19,149 +19,6 @@
 namespace UE::MultiUserClient::SingleClientColumns
 {
 	const FName ToggleTopLevelAuthorityColumnId = TEXT("ToggleTopLevelAuthorityColumn");
-	const FName ToggleSubobjectAuthorityColumnId = TEXT("ToggleSubobjectAuthorityColumn");
-
-	namespace Private::ToggleTopLevelAuthority
-	{
-		static ECheckBoxState GetCheckboxState(
-			const ConcertClientSharedSlate::FReplicatedObjectData& ObjectData,
-			ConcertClientSharedSlate::IReplicationStreamModel* ClientStreamModel,
-			FAuthorityChangeTracker* ChangeTracker
-			)
-		{
-			const FSoftObjectPath& ObjectPath = ObjectData.GetObjectPath();
-
-			// Init the checkbox state to the top level object, if the top level object has properties associated with it...
-			const bool bCanSetTopLevelState = ChangeTracker->CanSetAuthorityFor(ObjectPath)
-				|| ChangeTracker->GetChangeAuthorityMutability(ObjectPath) == EAuthorityMutability::NotSupported;
-			TOptional<ECheckBoxState> CheckBoxState = bCanSetTopLevelState
-				? ChangeTracker->GetAuthorityStateAfterApplied(ObjectPath) ? ECheckBoxState::Checked : ECheckBoxState::Unchecked
-				: TOptional<ECheckBoxState>{};
-
-			// ... and then determine whether all the other subobjects share that state (Check / Unchecked) or not (Undetermined)
-			ClientStreamModel->ForEachSubobject(ObjectPath, [&CheckBoxState, &ChangeTracker](const FSoftObjectPath& Child)
-			{
-				if (!ChangeTracker->CanSetAuthorityFor(Child)
-					&& ChangeTracker->GetChangeAuthorityMutability(Child) != EAuthorityMutability::NotSupported)
-				{
-					return EBreakBehavior::Continue;
-				}
-
-				// Init checkbox state
-				const bool bCurrentAuthorityState = ChangeTracker->GetAuthorityStateAfterApplied(Child);
-				const bool bIsFirstObject = !CheckBoxState.IsSet(); 
-				if (bIsFirstObject)
-				{
-					CheckBoxState = bCurrentAuthorityState ? ECheckBoxState::Checked : ECheckBoxState::Unchecked;
-					return EBreakBehavior::Continue;
-				}
-
-				// Return the checkbox state that all subobjects share, or Undetermined if mixed
-				const bool bMustBeUndetermined = (*CheckBoxState == ECheckBoxState::Unchecked && bCurrentAuthorityState)
-					|| (*CheckBoxState == ECheckBoxState::Checked && !bCurrentAuthorityState);
-				if (bMustBeUndetermined)
-				{
-					CheckBoxState = ECheckBoxState::Undetermined;
-					return EBreakBehavior::Break;
-				}
-				return EBreakBehavior::Continue;
-			});
-						
-			return CheckBoxState
-				? *CheckBoxState
-				// If this case happens, there are no objects to check. In that case, it visually "looks better" to be unchecked
-				: ECheckBoxState::Unchecked;
-		}
-
-		static void OnCheckboxStateChanged(
-			bool bIsChecked,
-			const ConcertClientSharedSlate::FReplicatedObjectData& ObjectData,
-			ConcertClientSharedSlate::IReplicationStreamModel* ClientStreamModel,
-			FAuthorityChangeTracker* ChangeTracker
-			)
-		{
-			const FSoftObjectPath& ObjectPath = ObjectData.GetObjectPath();
-			TArray<FSoftObjectPath> Paths { ObjectPath };
-			ClientStreamModel->ForEachSubobject(ObjectPath, [&Paths](const FSoftObjectPath& Child)
-			{
-				Paths.Add(Child);
-				return EBreakBehavior::Continue;
-			});
-			
-			ChangeTracker->SetAuthorityIfAllowed(Paths, bIsChecked);
-		}
-
-		static bool IsEnabled(
-			const ConcertClientSharedSlate::FReplicatedObjectData& ObjectData,
-			ConcertClientSharedSlate::IReplicationStreamModel* ClientStreamModel,
-			FAuthorityChangeTracker* ChangeTracker
-			)
-		{
-			const FSoftObjectPath& TopLevelObjectPath = ObjectData.GetObjectPath();
-			const bool bCanChangeTopLevel = ChangeTracker->CanSetAuthorityFor(TopLevelObjectPath);
-			if (bCanChangeTopLevel)
-			{
-				return true;
-			}
-						
-			bool bCanChangeAnySubobject = false;
-			ClientStreamModel->ForEachSubobject(TopLevelObjectPath, [&ChangeTracker, &bCanChangeAnySubobject](const FSoftObjectPath& Child)
-			{
-				bCanChangeAnySubobject |= ChangeTracker->CanSetAuthorityFor(Child);
-				return bCanChangeAnySubobject ? EBreakBehavior::Break : EBreakBehavior::Continue;
-			});
-			return bCanChangeAnySubobject;
-		}
-	}
-	
-	ConcertClientSharedSlate::ReplicationColumns::FReplicationTopLevelObjectColumn ToggleTopLevelAuthority(
-		ConcertClientSharedSlate::IReplicationStreamModel& ClientStreamModel,
-		FAuthorityChangeTracker& ChangeTracker,
-		ISubmissionWorkflow& SubmissionWorkflow
-		)
-	{
-		using namespace ConcertClientSharedSlate;
-		using namespace Private::ToggleTopLevelAuthority;
-		using FTopLevelColumnDelegates = TReplicationColumnDelegates<FReplicatedObjectData>;
-		
-		return MakeCheckboxColumn<FReplicatedObjectData>(
-			ToggleTopLevelAuthorityColumnId,
-				FTopLevelColumnDelegates(
-					FTopLevelColumnDelegates::FGetColumnCheckboxState::CreateStatic(&GetCheckboxState, &ClientStreamModel, &ChangeTracker),
-					FTopLevelColumnDelegates::FOnColumnCheckboxChanged::CreateStatic(&OnCheckboxStateChanged, &ClientStreamModel, &ChangeTracker),
-					FTopLevelColumnDelegates::FGetToolTipText::CreateLambda([&ClientStreamModel, &ChangeTracker, &SubmissionWorkflow](const FReplicatedObjectData& ObjectData)
-					{
-						
-						if (SubmissionWorkflow.GetUploadability() == EChangeUploadability::NotImplemented)
-						{
-							const ECheckBoxState CheckBoxState = GetCheckboxState(ObjectData, &ClientStreamModel, &ChangeTracker);
-							const FText CheckBoxText = [CheckBoxState]()
-							{
-								switch (CheckBoxState)
-								{
-								case ECheckBoxState::Unchecked: return LOCTEXT("TopLevel.ChangeAuthority.ToolTip.NotSupported.Unchecked", "None of the subobjects are being replicated.");
-								case ECheckBoxState::Checked: return LOCTEXT("TopLevel.ChangeAuthority.ToolTip.NotSupported.Checked", "All of the subobjects are being replicated.");
-								case ECheckBoxState::Undetermined: return LOCTEXT("TopLevel.ChangeAuthority.ToolTip.NotSupported.Undetermined", "Some of the subobjects are being replicated.");
-								default: checkNoEntry(); return FText::GetEmpty();
-								}
-							}();
-							
-							return FText::Format(
-								LOCTEXT("TopLevel.ChangeAuthority.ToolTip.NotSupportedFmt", "{0}\nEditing remote clients is not implemented. You can only edit the local client."),
-								CheckBoxText
-								);
-						}
-						
-						return IsEnabled(ObjectData, &ClientStreamModel, &ChangeTracker)
-							? LOCTEXT("TopLevel.ChangeAuthority.Allowed", "Toggles whether this object and its subobjects will replicate the assigned properties.")
-							: LOCTEXT("TopLevel.ChangeAuthority.Disallowed", "Toggles whether this object and its subobjects will replicate the assigned properties.\nDisabled because none of the subobjects can be replicated with their current configuration.");
-					}),
-					FTopLevelColumnDelegates::FIsEnabled::CreateStatic(&IsEnabled, &ClientStreamModel, &ChangeTracker)
-				),
-			FText::GetEmpty(),
-			static_cast<int32>(ETopLevelObjectColumnOrder::ToggleAuthority)
-		);
-	}
 
 	ConcertClientSharedSlate::ReplicationColumns::FReplicationTopLevelObjectColumn ToggleObjectAuthority(
 		FAuthorityChangeTracker& ChangeTracker,
@@ -169,42 +26,37 @@ namespace UE::MultiUserClient::SingleClientColumns
 		)
 	{
 		using namespace ConcertClientSharedSlate;
-		using FSubobjectColumnDelegates = TReplicationColumnDelegates<FReplicatedObjectData>;
+		using FColumnDelegates = TReplicationColumnDelegates<FReplicatedObjectData>;
 		return MakeCheckboxColumn<FReplicatedObjectData>(
 			ToggleTopLevelAuthorityColumnId,
-				FSubobjectColumnDelegates(
-					FSubobjectColumnDelegates::FGetColumnCheckboxState::CreateLambda(
+				FColumnDelegates(
+					FColumnDelegates::FGetColumnCheckboxState::CreateLambda(
 					[&ChangeTracker](const FReplicatedObjectData& ObjectData)
 					{
 						const bool bHasAuthority = ChangeTracker.GetAuthorityStateAfterApplied(ObjectData.GetObjectPath());
 						return bHasAuthority ? ECheckBoxState::Checked : ECheckBoxState::Unchecked;
 					}),
-					FSubobjectColumnDelegates::FOnColumnCheckboxChanged::CreateLambda(
+					FColumnDelegates::FOnColumnCheckboxChanged::CreateLambda(
 					[&ChangeTracker](bool bIsChecked, const FReplicatedObjectData& ObjectData)
 					{
 						ChangeTracker.SetAuthorityIfAllowed({ ObjectData.GetObjectPath() }, bIsChecked);
 					}),
-					FSubobjectColumnDelegates::FGetToolTipText::CreateLambda([&ChangeTracker, &SubmissionWorkflow](const FReplicatedObjectData& ObjectData)
+					FColumnDelegates::FGetToolTipText::CreateLambda([&ChangeTracker, &SubmissionWorkflow](const FReplicatedObjectData& ObjectData)
 					{
-						const bool bHasAuthority = ChangeTracker.GetAuthorityStateAfterApplied(ObjectData.GetObjectPath());
-						const FText NotSupportedText = FText::Format(
-							LOCTEXT("Subobject.ChangeAuthority.ToolTip.NotSupportedFmt", "{0}\nEditing remote clients is not implemented. You can only edit the local client."),
-							bHasAuthority ? LOCTEXT("Subobject.ChangeAuthority.Replicating.True", "This object is being replicated.") : LOCTEXT("Subobject.ChangeAuthority.Replicating.False", "This object is not currently being replicated.")
-							);
-						if (SubmissionWorkflow.GetUploadability() == EChangeUploadability::NotImplemented)
+						if (!CanEverSubmit(SubmissionWorkflow.GetUploadability()))
 						{
-							return NotSupportedText;
+							return LOCTEXT("ToggleAuthority.ToolTip.NotSupported", "This client cannot be remotely edited.");
 						}
 						
 						switch (ChangeTracker.GetChangeAuthorityMutability(ObjectData.GetObjectPath()))
 						{
-						case EAuthorityMutability::Allowed: return LOCTEXT("Subobject.ChangeAuthority.ToolTip.Allowed", "Toggles whether this object should be replicated.");
-						case EAuthorityMutability::NoProperties: return LOCTEXT("Subobject.ChangeAuthority.ToolTip.NoProperties", "Toggles whether this object should be replicated.\nAssign properties to this object first.\n");
-						case EAuthorityMutability::NotSupported: return NotSupportedText;
+						case EAuthorityMutability::Allowed: return LOCTEXT("ToggleAuthority.ToolTip.Allowed", "Whether to replicate this object.");
+						case EAuthorityMutability::NotApplicable: return LOCTEXT("ToggleAuthority.ToolTip.NotApplicable", "Assign properties to replicate first.");
+						case EAuthorityMutability::Conflict: return LOCTEXT("ToggleAuthority.ToolTip.Conflict", "Another client is replicating this property already.");
 						default: checkNoEntry(); return FText::GetEmpty();
 						}
 					}),
-					FSubobjectColumnDelegates::FIsEnabled::CreateLambda([&ChangeTracker](const FReplicatedObjectData& ObjectData)
+					FColumnDelegates::FIsEnabled::CreateLambda([&ChangeTracker](const FReplicatedObjectData& ObjectData)
 					{
 						return ChangeTracker.CanSetAuthorityFor(ObjectData.GetObjectPath());
 					})
diff --git a/Engine/Plugins/Developer/Concert/ConcertApp/MultiUserClient/Source/MultiUserClient/Private/Widgets/ActiveSession/Replication/Client/Single/SReplicationClientView.cpp b/Engine/Plugins/Developer/Concert/ConcertApp/MultiUserClient/Source/MultiUserClient/Private/Widgets/ActiveSession/Replication/Client/Single/SReplicationClientView.cpp
--- a/Engine/Plugins/Developer/Concert/ConcertApp/MultiUserClient/Source/MultiUserClient/Private/Widgets/ActiveSession/Replication/Client/Single/SReplicationClientView.cpp
+++ b/Engine/Plugins/Developer/Concert/ConcertApp/MultiUserClient/Source/MultiUserClient/Private/Widgets/ActiveSession/Replication/Client/Single/SReplicationClientView.cpp
@@ -76,7 +76,7 @@ namespace UE::MultiUserClient
 			.DataModel = InReplicationClient.GetClientEditModel(),
 			.ObjectSource = MakeShared<FActorSelectionSourceModel>(),
 			.PropertySource = MakeShared<FSelectPropertyFromUClassModel>(),
-			.IsEditingEnabled = TAttribute<bool>::CreateLambda([&SubmissionWorkflow](){ return SubmissionWorkflow.GetUploadability() != EChangeUploadability::NotImplemented; }),
+			.IsEditingEnabled = TAttribute<bool>::CreateLambda([&SubmissionWorkflow](){ return CanEverSubmit(SubmissionWorkflow.GetUploadability()); }),
 			.EditingDisabledToolTipText = LOCTEXT("Editing.NotImplemented", "Editing remote clients is not implemented. You can only edit the local client."),
 			.ReplicationSettingsAttribute = MoveTemp(ReplicationSettingsAttribute),
 			.ViewerParams =
diff --git a/Engine/Plugins/Developer/Concert/ConcertSync/ConcertSyncServer/Source/ConcertSyncServer/Private/Replication/ConcertServerReplicationManager_ChangeStream.cpp b/Engine/Plugins/Developer/Concert/ConcertSync/ConcertSyncServer/Source/ConcertSyncServer/Private/Replication/ConcertServerReplicationManager_ChangeStream.cpp
--- a/Engine/Plugins/Developer/Concert/ConcertSync/ConcertSyncServer/Source/ConcertSyncServer/Private/Replication/ConcertServerReplicationManager_ChangeStream.cpp
+++ b/Engine/Plugins/Developer/Concert/ConcertSync/ConcertSyncServer/Source/ConcertSyncServer/Private/Replication/ConcertServerReplicationManager_ChangeStream.cpp
@@ -57,6 +57,7 @@ namespace UE::ConcertSyncServer::Replication
 				else
 				{
 					// ObjectsToPut can only write to pre-existing streams
+					UE_LOG(LogConcert, Log, TEXT("Semantic error: unknown stream %s"), *StreamToModify.ToString(EGuidFormats::Short));
 					OutResponse.ObjectsToPutSemanticErrors.Add(ChangedObject, EConcertPutObjectErrorCode::UnresolvedStream);
 				}
 			}
@@ -78,6 +79,7 @@ namespace UE::ConcertSyncServer::Replication
 				
 				if (bIdAlreadyExists || bIsDuplicateEntry)
 				{
+					UE_LOG(LogConcert, Log, TEXT("Duplicate stream entry %s"), *NewStreamId.ToString(EGuidFormats::Short));
 					OutResponse.FailedStreamCreation.Add(NewStreamId);
 				}
 			}
@@ -106,10 +108,12 @@ namespace UE::ConcertSyncServer::Replication
 
 				// Simply check whether any other client is already sending any of the requested properties.
 				AuthorityManager.EnumerateAuthorityConflicts(ReplicatedObjectInfo, &PropertySelection,
-					[&OutResponse, &ReplicatedObjectInfo](const FGuid& ClientId, const FGuid& StreamId, const FConcertPropertyChain&)
+					[&OutResponse, &ReplicatedObjectInfo](const FGuid& ClientId, const FGuid& StreamId, const FConcertPropertyChain& Property)
 					{
 						const FReplicatedObjectId ConflictingObject = { { StreamId, ReplicatedObjectInfo.Object }, ClientId };
 						OutResponse.AuthorityConflicts.Add(ReplicatedObjectInfo, ConflictingObject);
+						
+						UE_LOG(LogConcert, Log, TEXT("Authority conflict with client %s for stream %s for property %s"), *ClientId.ToString(EGuidFormats::Short), *StreamId.ToString(EGuidFormats::Short), *Property.ToString());
 						return EBreakBehavior::Continue;
 					});
 			}
@@ -122,6 +126,7 @@ namespace UE::ConcertSyncServer::Replication
 			{
 				if (!ToUnpack.Unpack())
 				{
+					UE_LOG(LogConcert, Log, TEXT("Failed to unpack stream %s"), *ToUnpack.BaseDescription.Identifier.ToString(EGuidFormats::Short));
 					OutResponse.FailedStreamCreation.Add(ToUnpack.BaseDescription.Identifier);
 				}
 			}
diff --git a/Engine/Plugins/Developer/Concert/ConcertUI/ConcertClientSharedSlate/Source/ConcertClientSharedSlate/Private/Replication/Editor/View/MultiEditor/ConsolidatedMultiStreamModel.cpp b/Engine/Plugins/Developer/Concert/ConcertUI/ConcertClientSharedSlate/Source/ConcertClientSharedSlate/Private/Replication/Editor/View/MultiEditor/ConsolidatedMultiStreamModel.cpp
--- a/Engine/Plugins/Developer/Concert/ConcertUI/ConcertClientSharedSlate/Source/ConcertClientSharedSlate/Private/Replication/Editor/View/MultiEditor/ConsolidatedMultiStreamModel.cpp
+++ b/Engine/Plugins/Developer/Concert/ConcertUI/ConcertClientSharedSlate/Source/ConcertClientSharedSlate/Private/Replication/Editor/View/MultiEditor/ConsolidatedMultiStreamModel.cpp
@@ -46,7 +46,7 @@ namespace UE::ConcertClientSharedSlate
 		, StreamForAdding(Private::MakeModel(StreamForAddingObject, bShouldSupportTransactions))
 		, MultiStreamModel(MoveTemp(InMultiStreamModel))
 	{
-		MultiStreamModel->OnReadOnlyStreamChanged().AddRaw(this, &FConsolidatedMultiStreamModel::OnReadOnlyStreamChanged);
+		MultiStreamModel->OnStreamExternallyChanged().AddRaw(this, &FConsolidatedMultiStreamModel::OnStreamExternallyChanged);
 		MultiStreamModel->OnStreamSetChanged().AddRaw(this, &FConsolidatedMultiStreamModel::RebuildStreamSubscriptions);
 		StreamForAdding->OnObjectsChanged().AddRaw(this, &FConsolidatedMultiStreamModel::OnObjectsChanged_StreamForAdding);
 		RebuildStreamSubscriptions();
@@ -54,7 +54,7 @@ namespace UE::ConcertClientSharedSlate
 
 	FConsolidatedMultiStreamModel::~FConsolidatedMultiStreamModel()
 	{
-		MultiStreamModel->OnReadOnlyStreamChanged().RemoveAll(this);
+		MultiStreamModel->OnStreamExternallyChanged().RemoveAll(this);
 		MultiStreamModel->OnStreamSetChanged().RemoveAll(this);
 		StreamForAdding->OnObjectsChanged().RemoveAll(this);
 		ClearStreamSubscriptions();
@@ -152,7 +152,7 @@ namespace UE::ConcertClientSharedSlate
 		OnObjectsChangedDelegate.Broadcast(AddedObjects, RemovedObjects, ChangeReason);
 	}
 
-	void FConsolidatedMultiStreamModel::OnReadOnlyStreamChanged(TSharedRef<IReplicationStreamModel> Stream)
+	void FConsolidatedMultiStreamModel::OnStreamExternallyChanged(TSharedRef<IReplicationStreamModel> Stream)
 	{
 		OnObjectsChangedDelegate.Broadcast({}, {}, EReplicatedObjectChangeReason::ExternalChange);
 	}
diff --git a/Engine/Plugins/Developer/Concert/ConcertUI/ConcertClientSharedSlate/Source/ConcertClientSharedSlate/Private/Replication/Editor/View/MultiEditor/ConsolidatedMultiStreamModel.h b/Engine/Plugins/Developer/Concert/ConcertUI/ConcertClientSharedSlate/Source/ConcertClientSharedSlate/Private/Replication/Editor/View/MultiEditor/ConsolidatedMultiStreamModel.h
--- a/Engine/Plugins/Developer/Concert/ConcertUI/ConcertClientSharedSlate/Source/ConcertClientSharedSlate/Private/Replication/Editor/View/MultiEditor/ConsolidatedMultiStreamModel.h
+++ b/Engine/Plugins/Developer/Concert/ConcertUI/ConcertClientSharedSlate/Source/ConcertClientSharedSlate/Private/Replication/Editor/View/MultiEditor/ConsolidatedMultiStreamModel.h
@@ -91,7 +91,7 @@ namespace UE::ConcertClientSharedSlate
 		void OnObjectsChanged_ClientStream(TConstArrayView<UObject*> AddedObjects, TConstArrayView<FSoftObjectPath> RemovedObjects, EReplicatedObjectChangeReason ChangeReason);
 		
 		/** Trigger subscription rebuild when streams are added / removed. */
-		void OnReadOnlyStreamChanged(TSharedRef<IReplicationStreamModel> Stream);
+		void OnStreamExternallyChanged(TSharedRef<IReplicationStreamModel> Stream);
 		void RebuildStreamSubscriptions();
 		void ClearStreamSubscriptions();
 	};
diff --git a/Engine/Plugins/Developer/Concert/ConcertUI/ConcertClientSharedSlate/Source/ConcertClientSharedSlate/Public/Replication/Editor/Model/IMultiReplicationStreamModel.h b/Engine/Plugins/Developer/Concert/ConcertUI/ConcertClientSharedSlate/Source/ConcertClientSharedSlate/Public/Replication/Editor/Model/IMultiReplicationStreamModel.h
--- a/Engine/Plugins/Developer/Concert/ConcertUI/ConcertClientSharedSlate/Source/ConcertClientSharedSlate/Public/Replication/Editor/Model/IMultiReplicationStreamModel.h
+++ b/Engine/Plugins/Developer/Concert/ConcertUI/ConcertClientSharedSlate/Source/ConcertClientSharedSlate/Public/Replication/Editor/Model/IMultiReplicationStreamModel.h
@@ -17,9 +17,9 @@ namespace UE::ConcertClientSharedSlate
 		/** @return Streams in this model that can only be read from. */
 		virtual TSet<TSharedRef<IReplicationStreamModel>> GetReadOnlyStreams() const = 0;
 
-		DECLARE_MULTICAST_DELEGATE_OneParam(FOnReadOnlyStreamChanged, TSharedRef<IReplicationStreamModel> /* ChangedStream */);
+		DECLARE_MULTICAST_DELEGATE_OneParam(FOnStreamExternallyChanged, TSharedRef<IReplicationStreamModel> /* ChangedStream */);
 		/** Called when a read-only stream has changed. */
-		virtual FOnReadOnlyStreamChanged& OnReadOnlyStreamChanged() = 0;
+		virtual FOnStreamExternallyChanged& OnStreamExternallyChanged() = 0;
 		
 		DECLARE_MULTICAST_DELEGATE(FOnStreamSetChanged);
 		/** Broadcasts when the result of GetStreams has changed. */
