Fix struct not initialized warning

diff --git a/Engine/Plugins/Developer/Concert/ConcertApp/MultiUserClient/Source/MultiUserClient/Private/Replication/Submission/Remote/RemoteSubmissionMessages.h b/Engine/Plugins/Developer/Concert/ConcertApp/MultiUserClient/Source/MultiUserClient/Private/Replication/Submission/Remote/RemoteSubmissionMessages.h
--- a/Engine/Plugins/Developer/Concert/ConcertApp/MultiUserClient/Source/MultiUserClient/Private/Replication/Submission/Remote/RemoteSubmissionMessages.h
+++ b/Engine/Plugins/Developer/Concert/ConcertApp/MultiUserClient/Source/MultiUserClient/Private/Replication/Submission/Remote/RemoteSubmissionMessages.h
@@ -8,6 +8,7 @@
 
 // IMPORTANT: This peer-to-peer workflow is subject to be replaced, possibly in 5.5, by a server feature.
 // The peer-to-peer approach cannot guarantee atomicity if client A is trying to edit clients B and C as one operation: How to revert when B succeeds and C fails? You can't.
+// TODO UE-201136.
 
 /** A kind request from one client to another client to change their server registered state. The receiver can reject the request. */
 USTRUCT()
@@ -107,5 +108,5 @@ struct FMultiUser_ChangeRemote_Cancelled
 	GENERATED_BODY()
 
 	UPROPERTY()
-	EMultiUserCancelRemoteChangeReason Reason;
+	EMultiUserCancelRemoteChangeReason Reason = EMultiUserCancelRemoteChangeReason::FailedToCreate;
 };
\ No newline at end of file
