Add queuing for submissions

diff --git a/Engine/Plugins/Developer/Concert/ConcertApp/MultiUserClient/Source/MultiUserClient/Private/Replication/Client/LocalReplicationClient.cpp b/Engine/Plugins/Developer/Concert/ConcertApp/MultiUserClient/Source/MultiUserClient/Private/Replication/Client/LocalReplicationClient.cpp
--- a/Engine/Plugins/Developer/Concert/ConcertApp/MultiUserClient/Source/MultiUserClient/Private/Replication/Client/LocalReplicationClient.cpp
+++ b/Engine/Plugins/Developer/Concert/ConcertApp/MultiUserClient/Source/MultiUserClient/Private/Replication/Client/LocalReplicationClient.cpp
@@ -21,6 +21,6 @@ namespace UE::MultiUserClient
 			MoveTemp(InStreamSynchronizer),
 			MakeUnique<FAuthoritySynchronizer_LocalClient>(InClient),
 			MakeUnique<FSubmissionWorkflow_LocalClient>(MoveTemp(InClient)))
-		, RemoteSubmissionListener(InClient->GetConcertClient()->GetCurrentSession().ToSharedRef(), GetStreamSynchronizer(), GetSubmissionWorkflow())
+		, RemoteSubmissionListener(InClient->GetConcertClient()->GetCurrentSession().ToSharedRef(), GetStreamSynchronizer(), GetSubmissionQueue())
 	{}
 }
diff --git a/Engine/Plugins/Developer/Concert/ConcertApp/MultiUserClient/Source/MultiUserClient/Private/Replication/Client/ReplicationClient.cpp b/Engine/Plugins/Developer/Concert/ConcertApp/MultiUserClient/Source/MultiUserClient/Private/Replication/Client/ReplicationClient.cpp
--- a/Engine/Plugins/Developer/Concert/ConcertApp/MultiUserClient/Source/MultiUserClient/Private/Replication/Client/ReplicationClient.cpp
+++ b/Engine/Plugins/Developer/Concert/ConcertApp/MultiUserClient/Source/MultiUserClient/Private/Replication/Client/ReplicationClient.cpp
@@ -25,6 +25,7 @@ namespace UE::MultiUserClient
 		, StreamSynchronizer(MoveTemp(InStreamSynchronizer))
 		, AuthoritySynchronizer(MoveTemp(InAuthoritySynchronizer))
 		, SubmissionWorkflow(MoveTemp(InSubmissionWorkflow))
+		, SubmissionQueue(*SubmissionWorkflow)
 		, LocalClientEditModel(ConcertClientSharedSlate::CreatePropertySelectionModel(
 			*ClientContentStorage->Stream,
 			ClientContentStorage->Stream->MakeReplicationMapGetterAttribute()
@@ -36,7 +37,7 @@ namespace UE::MultiUserClient
 			)
 		, LocalAuthorityDiffer(EndpointId, *AuthoritySynchronizer, InAuthorityCache)
 		, ChangeRequestBuilder(EndpointId, InAuthorityCache, *StreamSynchronizer, LocalClientStreamDiffer, LocalAuthorityDiffer)
-		, AutoSubmissionPolicy(*SubmissionWorkflow.Get(), ChangeRequestBuilder, LocalClientEditModel.Get(), LocalAuthorityDiffer)
+		, AutoSubmissionPolicy(SubmissionQueue, ChangeRequestBuilder, LocalClientEditModel.Get(), LocalAuthorityDiffer)
 	{
 		LocalClientEditModel->OnObjectsChanged().AddRaw(this, &FReplicationClient::OnObjectsChanged);
 		LocalClientEditModel->OnPropertiesChanged().AddRaw(this, &FReplicationClient::OnPropertiesChanged);
diff --git a/Engine/Plugins/Developer/Concert/ConcertApp/MultiUserClient/Source/MultiUserClient/Private/Replication/Client/ReplicationClient.h b/Engine/Plugins/Developer/Concert/ConcertApp/MultiUserClient/Source/MultiUserClient/Private/Replication/Client/ReplicationClient.h
--- a/Engine/Plugins/Developer/Concert/ConcertApp/MultiUserClient/Source/MultiUserClient/Private/Replication/Client/ReplicationClient.h
+++ b/Engine/Plugins/Developer/Concert/ConcertApp/MultiUserClient/Source/MultiUserClient/Private/Replication/Client/ReplicationClient.h
@@ -9,6 +9,7 @@
 #include "Replication/Submission/AutoSubmissionPolicy.h"
 #include "Replication/Submission/ChangeRequestBuilder.h"
 #include "Replication/Submission/ISubmissionWorkflow.h"
+#include "Replication/Submission/Queue/SubmissionQueue.h"
 
 #include "Templates/SharedPointer.h"
 #include "Templates/UnrealTemplate.h"
@@ -54,9 +55,11 @@ namespace UE::MultiUserClient
 		
 		const FAuthorityChangeTracker& GetAuthorityDiffer() const { return LocalAuthorityDiffer; }
 		FAuthorityChangeTracker& GetAuthorityDiffer() { return LocalAuthorityDiffer; }
-		
-		const ISubmissionWorkflow& GetSubmissionWorkflow() const { return *SubmissionWorkflow; }
+
+		/** @return The object with which you can submit changes. You can use it to listen to general events. If you want to submit changes, you should prefer to use the SubmissionQueue. */
 		ISubmissionWorkflow& GetSubmissionWorkflow() { return *SubmissionWorkflow; }
+		/** @return Implements simple game-thread based queue for SubmissionWorkflow. Only one submission can be in progress at any given time. */
+		FSubmissionQueue& GetSubmissionQueue() { return SubmissionQueue; }
 
 		/** @return The endpoint ID of this client in the Concert session. */
 		const FGuid& GetEndpointId() const { return EndpointId; }
@@ -88,6 +91,9 @@ namespace UE::MultiUserClient
 		TUniquePtr<IClientAuthoritySynchronizer> AuthoritySynchronizer;
 		/** Handles the logic of submitting and reverting for this client. */
 		TUniquePtr<ISubmissionWorkflow> SubmissionWorkflow;
+
+		/** Allows systems to queue pending submissions to SubmissionWorkflow in case a submission is in progress. */
+		FSubmissionQueue SubmissionQueue;
 		
 		/**
 		 * Used to detect changes made to the client's config by the local editor.
diff --git a/Engine/Plugins/Developer/Concert/ConcertApp/MultiUserClient/Source/MultiUserClient/Private/Replication/Submission/AutoSubmissionPolicy.cpp b/Engine/Plugins/Developer/Concert/ConcertApp/MultiUserClient/Source/MultiUserClient/Private/Replication/Submission/AutoSubmissionPolicy.cpp
--- a/Engine/Plugins/Developer/Concert/ConcertApp/MultiUserClient/Source/MultiUserClient/Private/Replication/Submission/AutoSubmissionPolicy.cpp
+++ b/Engine/Plugins/Developer/Concert/ConcertApp/MultiUserClient/Source/MultiUserClient/Private/Replication/Submission/AutoSubmissionPolicy.cpp
@@ -6,16 +6,18 @@
 #include "ISubmissionWorkflow.h"
 #include "Replication/Authority/AuthorityChangeTracker.h"
 #include "Replication/Editor/Model/IEditableReplicationStreamModel.h"
+#include "Replication/Submission/Queue/SubmissionQueue.h"
 
 namespace UE::MultiUserClient
 {
 	FAutoSubmissionPolicy::FAutoSubmissionPolicy(
-		ISubmissionWorkflow& InSubmissionWorkflow,
+		FSubmissionQueue& InSubmissionQueue,
 		const FChangeRequestBuilder& InRequestBuilder,
 		ConcertClientSharedSlate::IEditableReplicationStreamModel& InStreamEditorModel,
 		FAuthorityChangeTracker& InAuthorityChangeTracker
 		)
-		: SubmissionWorkflow(InSubmissionWorkflow)
+		: FSelfUnregisteringDeferredSubmitter(InSubmissionQueue)
+		, SubmissionQueue(InSubmissionQueue)
 		, RequestBuilder(InRequestBuilder)
 		, StreamEditorModel(InStreamEditorModel)
 		, AuthorityChangeTracker(InAuthorityChangeTracker)
@@ -34,23 +36,13 @@ namespace UE::MultiUserClient
 
 	void FAutoSubmissionPolicy::ProcessAccumulatedChangesAndSubmit()
 	{
-		if (!bIsDirty)
+		if (bIsDirty)
 		{
-			return;
-		}
-		
-		if (SubmissionWorkflow.CanSubmit())
-		{
-			SubmissionWorkflow.OnSubmitOperationCompleted_AnyThread().RemoveAll(this);
-			SubmitChanges();
-		}
-		else if (!SubmissionWorkflow.OnSubmitOperationCompleted_AnyThread().IsBoundToObject(this))
-		{
-			SubmissionWorkflow.OnSubmitOperationCompleted_AnyThread().AddRaw(this, &FAutoSubmissionPolicy::ProcessAccumulatedChangesAndSubmit);
+			SubmissionQueue.SubmitNowOrEnqueue_GameThread(*this);
 		}
 	}
 
-	void FAutoSubmissionPolicy::SubmitChanges()
+	void FAutoSubmissionPolicy::PerformSubmission_GameThread(ISubmissionWorkflow& Workflow)
 	{
 		using namespace UE::ConcertSyncClient::Replication;
 		bIsDirty = false;
@@ -62,7 +54,7 @@ namespace UE::MultiUserClient
 			
 		if (AuthorityChangeRequest || StreamRequest)
 		{
-			SubmissionWorkflow.SubmitChanges({ StreamRequest, AuthorityChangeRequest });
+			Workflow.SubmitChanges({ StreamRequest, AuthorityChangeRequest });
 		}
 	}
 
diff --git a/Engine/Plugins/Developer/Concert/ConcertApp/MultiUserClient/Source/MultiUserClient/Private/Replication/Submission/AutoSubmissionPolicy.h b/Engine/Plugins/Developer/Concert/ConcertApp/MultiUserClient/Source/MultiUserClient/Private/Replication/Submission/AutoSubmissionPolicy.h
--- a/Engine/Plugins/Developer/Concert/ConcertApp/MultiUserClient/Source/MultiUserClient/Private/Replication/Submission/AutoSubmissionPolicy.h
+++ b/Engine/Plugins/Developer/Concert/ConcertApp/MultiUserClient/Source/MultiUserClient/Private/Replication/Submission/AutoSubmissionPolicy.h
@@ -3,6 +3,7 @@
 #pragma once
 
 #include "Replication/Editor/Model/IEditableReplicationStreamModel.h"
+#include "Replication/Submission/Queue/DeferredSubmitter.h"
 #include "Templates/UnrealTemplate.h"
 
 namespace UE::MultiUserClient
@@ -29,12 +30,14 @@ namespace UE::MultiUserClient
 	 * This policy accumulates all changes until the end of frame and then sends them. This means multiple changes are made in the same frame,
 	 * only one request will be sent containing all changes.
 	 */
-	class FAutoSubmissionPolicy : public FNoncopyable
+	class FAutoSubmissionPolicy
+		: public FNoncopyable
+		, FSelfUnregisteringDeferredSubmitter
 	{
 	public:
 		
 		FAutoSubmissionPolicy(
-			ISubmissionWorkflow& InSubmissionWorkflow,
+			FSubmissionQueue& InSubmissionQueue,
 			const FChangeRequestBuilder& InRequestBuilder,
 			ConcertClientSharedSlate::IEditableReplicationStreamModel& InStreamEditorModel,
 			FAuthorityChangeTracker& InAuthorityChangeTracker
@@ -47,7 +50,7 @@ namespace UE::MultiUserClient
 	private:
 
 		/** Handles performing the submission */
-		ISubmissionWorkflow& SubmissionWorkflow;
+		FSubmissionQueue& SubmissionQueue;
 		/** Used to building the requests that are passed to SubmissionWorkflow. */
 		const FChangeRequestBuilder& RequestBuilder;
 
@@ -58,9 +61,11 @@ namespace UE::MultiUserClient
 
 		/** Whether any changes were made. */
 		bool bIsDirty = false;
-		
-		void SubmitChanges();
-		
+
+		//~ Begin FSelfUnregisteringDeferredSubmitter Interface
+		virtual void PerformSubmission_GameThread(ISubmissionWorkflow& Workflow) override;
+		//~ End FSelfUnregisteringDeferredSubmitter Interface
+
 		void OnObjectsChanged(TArrayView<UObject* const>, TArrayView<const FSoftObjectPath>, ConcertClientSharedSlate::EReplicatedObjectChangeReason) { OnChangesDetected(); }
 		void OnChangesDetected();
 	};
diff --git a/Engine/Plugins/Developer/Concert/ConcertApp/MultiUserClient/Source/MultiUserClient/Private/Replication/Submission/Queue/DeferredSubmitter.h b/Engine/Plugins/Developer/Concert/ConcertApp/MultiUserClient/Source/MultiUserClient/Private/Replication/Submission/Queue/DeferredSubmitter.h
new file mode 100644
--- /dev/null
+++ b/Engine/Plugins/Developer/Concert/ConcertApp/MultiUserClient/Source/MultiUserClient/Private/Replication/Submission/Queue/DeferredSubmitter.h
@@ -0,0 +1,47 @@
+// Copyright Epic Games, Inc. All Rights Reserved.
+
+#pragma once
+
+#include "SubmissionQueue.h"
+#include "Templates/UnrealTemplate.h"
+
+namespace UE::MultiUserClient
+{
+	class ISubmissionWorkflow;
+	
+	/**
+	 * Implemented by systems that want to generate a submission once the ISubmissionWorkflow is ready.
+	 * This is an interface instead of a simple callback so the FSubmissionQueue can use the address of this object to avoid double enqueuing.
+	 */
+	class IDeferredSubmitter
+	{
+	public:
+
+		/**
+		 * Called when the ISubmissionWorkflow is ready.
+		 *
+		 * You are expected to call SubmitChanges but the queue handles the case if you don't (by dequeuing)
+		 * but you should prefer to dequeue once you know no more changes are needed.
+		 */
+		virtual void PerformSubmission_GameThread(ISubmissionWorkflow& Workflow) = 0;
+
+		virtual ~IDeferredSubmitter() = default;
+	};
+
+	/** Util for systems to auto-unregister when they are destroyed. */
+	class FSelfUnregisteringDeferredSubmitter : public IDeferredSubmitter
+	{
+		FSubmissionQueue& SubmissionQueue;
+	public:
+
+		FSelfUnregisteringDeferredSubmitter(FSubmissionQueue& InSubmissionQueue)
+			: SubmissionQueue(InSubmissionQueue)
+		{}
+		
+		virtual ~FSelfUnregisteringDeferredSubmitter()
+		{
+			SubmissionQueue.Dequeue_GameThread(*this);
+		}
+	};
+}
+
diff --git a/Engine/Plugins/Developer/Concert/ConcertApp/MultiUserClient/Source/MultiUserClient/Private/Replication/Submission/Queue/SubmissionQueue.cpp b/Engine/Plugins/Developer/Concert/ConcertApp/MultiUserClient/Source/MultiUserClient/Private/Replication/Submission/Queue/SubmissionQueue.cpp
new file mode 100644
--- /dev/null
+++ b/Engine/Plugins/Developer/Concert/ConcertApp/MultiUserClient/Source/MultiUserClient/Private/Replication/Submission/Queue/SubmissionQueue.cpp
@@ -0,0 +1,88 @@
+// Copyright Epic Games, Inc. All Rights Reserved.
+
+#include "SubmissionQueue.h"
+
+#include "DeferredSubmitter.h"
+#include "Replication/Submission/ISubmissionWorkflow.h"
+
+#include "Stats/Stats2.h"
+
+namespace UE::MultiUserClient
+{
+	FSubmissionQueue::FSubmissionQueue(ISubmissionWorkflow& InWorkflow)
+		: Workflow(InWorkflow)
+	{
+		Workflow.OnSubmitOperationCompleted_AnyThread().AddRaw(this, &FSubmissionQueue::OnSubmissionCompleted);
+	}
+
+	FSubmissionQueue::~FSubmissionQueue()
+	{
+		Workflow.OnSubmitOperationCompleted_AnyThread().RemoveAll(this);
+	}
+
+	void FSubmissionQueue::SubmitNowOrEnqueue_GameThread(IDeferredSubmitter& DeferredSubmission)
+	{
+		check(IsInGameThread());
+
+		if (Workflow.CanSubmit())
+		{
+			DeferredSubmission.PerformSubmission_GameThread(Workflow);
+		}
+		else if (!Queue.Contains(&DeferredSubmission))
+		{
+			Queue.AddTail(&DeferredSubmission);
+		}
+	}
+
+	void FSubmissionQueue::Dequeue_GameThread(IDeferredSubmitter& DeferredSubmission)
+	{
+		check(IsInGameThread());
+		Queue.RemoveNode(&DeferredSubmission);
+	}
+
+	void FSubmissionQueue::Tick(float DeltaTime)
+	{
+		if (bTaskFinishedOutsideGameThread)
+		{
+			ProcessNextInQueue();
+		}
+	}
+
+	TStatId FSubmissionQueue::GetStatId() const
+	{
+		RETURN_QUICK_DECLARE_CYCLE_STAT(FSubmissionNotifier, STATGROUP_Tickables);
+	}
+
+	void FSubmissionQueue::OnSubmissionCompleted()
+	{
+		if (IsInGameThread())
+		{
+			ProcessNextInQueue();
+		}
+		else
+		{
+			bTaskFinishedOutsideGameThread = true;
+		}
+	}
+
+	void FSubmissionQueue::ProcessNextInQueue()
+	{
+		check(IsInGameThread());
+		bTaskFinishedOutsideGameThread = false;
+		
+		if (!Queue.GetHead())
+		{
+			return;
+		}
+
+		IDeferredSubmitter* Submitter = Queue.GetHead()->GetValue();
+		Queue.RemoveNode(Queue.GetHead());
+		Submitter->PerformSubmission_GameThread(Workflow);
+
+		// PerformSubmission_GameThread did nothing? Progress the queue to avoid infinite waiting.
+		if (Workflow.CanSubmit())
+		{
+			ProcessNextInQueue();
+		}
+	}
+}
diff --git a/Engine/Plugins/Developer/Concert/ConcertApp/MultiUserClient/Source/MultiUserClient/Private/Replication/Submission/Queue/SubmissionQueue.h b/Engine/Plugins/Developer/Concert/ConcertApp/MultiUserClient/Source/MultiUserClient/Private/Replication/Submission/Queue/SubmissionQueue.h
new file mode 100644
--- /dev/null
+++ b/Engine/Plugins/Developer/Concert/ConcertApp/MultiUserClient/Source/MultiUserClient/Private/Replication/Submission/Queue/SubmissionQueue.h
@@ -0,0 +1,66 @@
+// Copyright Epic Games, Inc. All Rights Reserved.
+
+#pragma once
+
+#include "Containers/List.h"
+#include "Templates/UnrealTemplate.h"
+#include "TickableEditorObject.h"
+
+namespace UE::MultiUserClient
+{
+	class IDeferredSubmitter;
+	class ISubmissionWorkflow;
+	
+	/**
+	 * The submission queue is used to enforce a deterministic queue for calling ISubmissionWorkflow::SubmitChanges.
+	 * 
+	 * ISubmissionWorkflow only allows at most one pending change in progress at a time.
+	 * Hence, if one is in progress you must wait. This queue implements shared waiting logic.
+	 *
+	 * Since submissions are not real-time sensitive and to simplify concurrency design, the queue will always call
+	 * IDeferredSubmission::PerformSubmission_GameThread on the game thread.
+	 * This is not a big restriction either: usually requests complete on the game thread.
+	 * Only timeouts take place on the UDP thread, in which case we synchronize.
+	 */
+	class FSubmissionQueue : public FNoncopyable, public FTickableEditorObject
+	{
+	public:
+		
+		FSubmissionQueue(ISubmissionWorkflow& InWorkflow);
+		virtual ~FSubmissionQueue();
+
+		/**
+		 * Enqueues a submission.
+		 *
+		 * For convenience, enqueueing the same object multiple times is ignored.
+		 * Suppose you call Enqueue(A), Enqueue(B), Enqueue(A), then the queue will look like this A > B.
+		 *
+		 * @note You are responsible for calling Dequeue_GameThread when your IDeferredSubmission is destroyed.
+		 * @see FSelfUnregisteringDeferredSubmission
+		 */
+		void SubmitNowOrEnqueue_GameThread(IDeferredSubmitter& DeferredSubmission);
+
+		/** Dequeues a latent submission. It does not need to executed anymore. */
+		void Dequeue_GameThread(IDeferredSubmitter& DeferredSubmission);
+
+		//~ Begin FTickableEditorObject Interface
+		virtual void Tick(float DeltaTime) override;
+		virtual TStatId GetStatId() const override;
+		//~ End FTickableEditorObject Interface
+
+	private:
+
+		/** Implements the queue. The head is the next to execute. Used instead of TQueue because TQueue does not support removing. */
+		TDoubleLinkedList<IDeferredSubmitter*> Queue;
+
+		/** The workflow that is being enqueued for. */
+		ISubmissionWorkflow& Workflow;
+
+		/** Set to true when a task finishes outside of the game thread. Written exactly once by other threads other than game thread. Read by game thread. */
+		bool bTaskFinishedOutsideGameThread = false;
+
+		void OnSubmissionCompleted();
+		void ProcessNextInQueue();
+	};
+}
+
diff --git a/Engine/Plugins/Developer/Concert/ConcertApp/MultiUserClient/Source/MultiUserClient/Private/Replication/Submission/Remote/RemoteSubmissionListener.cpp b/Engine/Plugins/Developer/Concert/ConcertApp/MultiUserClient/Source/MultiUserClient/Private/Replication/Submission/Remote/RemoteSubmissionListener.cpp
--- a/Engine/Plugins/Developer/Concert/ConcertApp/MultiUserClient/Source/MultiUserClient/Private/Replication/Submission/Remote/RemoteSubmissionListener.cpp
+++ b/Engine/Plugins/Developer/Concert/ConcertApp/MultiUserClient/Source/MultiUserClient/Private/Replication/Submission/Remote/RemoteSubmissionListener.cpp
@@ -8,6 +8,7 @@
 #include "Replication/Stream/IClientStreamSynchronizer.h"
 #include "Replication/Submission/ISubmissionOperation.h"
 #include "Replication/Submission/ISubmissionWorkflow.h"
+#include "Replication/Submission/Queue/DeferredSubmitter.h"
 #include "Widgets/ActiveSession/Replication/Client/ClientUtils.h"
 
 namespace UE::MultiUserClient
@@ -15,11 +16,12 @@ namespace UE::MultiUserClient
 	FRemoteSubmissionListener::FRemoteSubmissionListener(
 		TSharedRef<IConcertClientSession> InConcertSession,
 		IClientStreamSynchronizer& InStreamSynchronizer,
-		ISubmissionWorkflow& InSubmissionWorkflow
+		FSubmissionQueue& InSubmissionQueue
 		)
-		: ConcertSession(MoveTemp(InConcertSession))
+		: FSelfUnregisteringDeferredSubmitter(InSubmissionQueue)
+		, ConcertSession(MoveTemp(InConcertSession))
 		, StreamSynchronizer(InStreamSynchronizer)
-		, SubmissionWorkflow(InSubmissionWorkflow)
+		, SubmissionQueue(InSubmissionQueue)
 	{
 		ConcertSession->RegisterCustomRequestHandler<FMultiUser_ChangeRemote_Request, FMultiUser_ChangeRemote_Response>(this, &FRemoteSubmissionListener::HandleChangeRemoteRequest);
 	}
@@ -27,7 +29,6 @@ namespace UE::MultiUserClient
 	FRemoteSubmissionListener::~FRemoteSubmissionListener()
 	{
 		ConcertSession->UnregisterCustomRequestHandler<FMultiUser_ChangeRemote_Request>();
-		SubmissionWorkflow.OnSubmitOperationCompleted_AnyThread().RemoveAll(this);
 	}
 
 	EConcertSessionResponseCode FRemoteSubmissionListener::HandleChangeRemoteRequest(
@@ -52,26 +53,11 @@ namespace UE::MultiUserClient
 		}
 		
 		InProgressOperation = { Request, Context.SourceEndpointId };
-		SubmitNowOrWaitUntilReady();
+		SubmissionQueue.SubmitNowOrEnqueue_GameThread(*this);
 		return EConcertSessionResponseCode::Success;
 	}
 
-	void FRemoteSubmissionListener::SubmitNowOrWaitUntilReady(int32 NumTriesSoFar)
-	{
-		if (SubmissionWorkflow.GetUploadability() == EChangeUploadability::Ready)
-		{
-			SubmissionWorkflow.OnSubmitOperationCompleted_AnyThread().RemoveAll(this);
-			SubmitRequest();
-		}
-		else if (ensureMsgf(SubmissionWorkflow.GetUploadability() == EChangeUploadability::InProgress, TEXT("Unsupported state")))
-		{
-			++NumTriesSoFar;
-			UE_LOG(LogConcert, Log, TEXT("Post-boning requested remote update since a local change is already in progress (%dth try)"), NumTriesSoFar);
-			SubmissionWorkflow.OnSubmitOperationCompleted_AnyThread().AddRaw(this, &FRemoteSubmissionListener::SubmitNowOrWaitUntilReady, NumTriesSoFar);
-		}
-	}
-
-	void FRemoteSubmissionListener::SubmitRequest()
+	void FRemoteSubmissionListener::PerformSubmission_GameThread(ISubmissionWorkflow& Workflow)
 	{
 		using namespace ConcertSyncClient::Replication;
 
@@ -92,7 +78,7 @@ namespace UE::MultiUserClient
 			return;
 		}
 		
-		const TSharedPtr<ISubmissionOperation> Operation = SubmissionWorkflow.SubmitChanges({
+		const TSharedPtr<ISubmissionOperation> Operation = Workflow.SubmitChanges({
 			FChangeStreamRequest { RequestData.StreamChangeRequest }, FAuthorityChangeRequest{ RequestData.AuthorityRequest } }
 			);
 		if (!ensure(Operation))
diff --git a/Engine/Plugins/Developer/Concert/ConcertApp/MultiUserClient/Source/MultiUserClient/Private/Replication/Submission/Remote/RemoteSubmissionListener.h b/Engine/Plugins/Developer/Concert/ConcertApp/MultiUserClient/Source/MultiUserClient/Private/Replication/Submission/Remote/RemoteSubmissionListener.h
--- a/Engine/Plugins/Developer/Concert/ConcertApp/MultiUserClient/Source/MultiUserClient/Private/Replication/Submission/Remote/RemoteSubmissionListener.h
+++ b/Engine/Plugins/Developer/Concert/ConcertApp/MultiUserClient/Source/MultiUserClient/Private/Replication/Submission/Remote/RemoteSubmissionListener.h
@@ -4,13 +4,11 @@
 
 #include "RemoteSubmissionMessages.h"
 #include "Replication/IToken.h"
+#include "Replication/Submission/Queue/DeferredSubmitter.h"
+
 #include "Templates/SharedPointer.h"
 #include "Templates/UnrealTemplate.h"
 
-namespace UE::MultiUserClient
-{
-	class IClientStreamSynchronizer;
-}
 
 class IConcertClientSession;
 
@@ -20,9 +18,9 @@ struct FConcertSessionContext;
 
 namespace UE::MultiUserClient
 {
+	class IClientStreamSynchronizer;
 	struct FSubmitAuthorityChangesResponse;
 	struct FSubmitStreamChangesResponse;
-	class ISubmissionWorkflow;
 	
 	/**
 	 * Handles stream and authority requests made by other clients towards the local client.
@@ -36,11 +34,13 @@ namespace UE::MultiUserClient
 	 *
 	 * @see FSubmissionWorkflow_RemoteClient
 	 */
-	class FRemoteSubmissionListener : public FNoncopyable
+	class FRemoteSubmissionListener
+		: public FNoncopyable
+		, FSelfUnregisteringDeferredSubmitter
 	{
 	public:
 		
-		FRemoteSubmissionListener(TSharedRef<IConcertClientSession> InConcertSession, IClientStreamSynchronizer& InStreamSynchronizer, ISubmissionWorkflow& InSubmissionWorkflow);
+		FRemoteSubmissionListener(TSharedRef<IConcertClientSession> InConcertSession, IClientStreamSynchronizer& InStreamSynchronizer, FSubmissionQueue& InSubmissionQueue);
 		~FRemoteSubmissionListener();
 
 	private:
@@ -54,7 +54,7 @@ namespace UE::MultiUserClient
 		/** Used to validate a request against the client's server state. */
 		IClientStreamSynchronizer& StreamSynchronizer;
 		/** Used to serve the incoming change requests. */
-		ISubmissionWorkflow& SubmissionWorkflow;
+		FSubmissionQueue& SubmissionQueue;
 
 		struct FOperationData
 		{
@@ -72,6 +72,11 @@ namespace UE::MultiUserClient
 		
 		/** Checks whether the local submit operation is ready. If not, we subscribe to it finishing and try again. */
 		void SubmitNowOrWaitUntilReady(int32 NumTriesSoFar = 0);
+
+		//~ Begin IDeferredSubmission Interface
+		virtual void PerformSubmission_GameThread(ISubmissionWorkflow& Workflow) override;
+		//~ End IDeferredSubmission Interface
+
 		/** The submission workflow is ready for submission. Do the submission. */
 		void SubmitRequest();
 
