Add ISubmissionOperation::OnCompleteOperation

diff --git a/Engine/Plugins/Developer/Concert/ConcertApp/MultiUserClient/Source/MultiUserClient/Private/Replication/Submission/ISubmissionOperation.h b/Engine/Plugins/Developer/Concert/ConcertApp/MultiUserClient/Source/MultiUserClient/Private/Replication/Submission/ISubmissionOperation.h
--- a/Engine/Plugins/Developer/Concert/ConcertApp/MultiUserClient/Source/MultiUserClient/Private/Replication/Submission/ISubmissionOperation.h
+++ b/Engine/Plugins/Developer/Concert/ConcertApp/MultiUserClient/Source/MultiUserClient/Private/Replication/Submission/ISubmissionOperation.h
@@ -8,6 +8,14 @@
 
 namespace UE::MultiUserClient
 {
+	enum class ESubmissionOperationCompletedCode
+	{
+		/** The operation was processed normally */
+		Processed,
+		/** Cancelled, e.g. because the workflow was destroyed */
+		Cancelled
+	};
+	
 	/**
 	 * Exposes the stages of submitting to the server.
 	 * Every time ISubmissionWorkflow::SubmitChanges is called a new instance is created.
@@ -23,19 +31,25 @@ namespace UE::MultiUserClient
 		 * Completes when the operation of changing streams has completed.
 		 * @note This can be called at most once; subsequent calls result in an unset future.
 		 */
-		virtual TFuture<FSubmitStreamChangesResponse> OnStreamChangesSubmittedFuture() = 0;
+		virtual TFuture<FSubmitStreamChangesResponse> OnCompleteStreamChangesFuture() = 0;
 
 		/**
 		 * Completes when the authority change request has been sent to the server.
 		 * @note This can be called at most once; subsequent calls result in an unset future.
 		 */
-		virtual TFuture<FSubmitAuthorityChangesRequest> OnAuthorityChangeRequestedFuture() = 0;
+		virtual TFuture<FSubmitAuthorityChangesRequest> OnRequestAuthorityChangeFuture() = 0;
 		
 		/**
 		 * Completes when the operation of changing authority has completed.
 		 * @note This can be called at most once; subsequent calls result in an unset future.
 		 */
-		virtual TFuture<FSubmitAuthorityChangesResponse> OnAuthorityChangeResponseReceivedFuture() = 0;
+		virtual TFuture<FSubmitAuthorityChangesResponse> OnCompleteAuthorityChangeFuture() = 0;
+
+		/**
+		 * Completes when the operation is done. No further work will be performed.
+		 * @note This can be called at most once; subsequent calls result in an unset future.
+		 */
+		virtual TFuture<ESubmissionOperationCompletedCode> OnCompletedOperation() = 0;
 		
 		virtual ~ISubmissionOperation() = default;
 	};
diff --git a/Engine/Plugins/Developer/Concert/ConcertApp/MultiUserClient/Source/MultiUserClient/Private/Replication/Submission/Remote/RemoteSubmissionListener.cpp b/Engine/Plugins/Developer/Concert/ConcertApp/MultiUserClient/Source/MultiUserClient/Private/Replication/Submission/Remote/RemoteSubmissionListener.cpp
--- a/Engine/Plugins/Developer/Concert/ConcertApp/MultiUserClient/Source/MultiUserClient/Private/Replication/Submission/Remote/RemoteSubmissionListener.cpp
+++ b/Engine/Plugins/Developer/Concert/ConcertApp/MultiUserClient/Source/MultiUserClient/Private/Replication/Submission/Remote/RemoteSubmissionListener.cpp
@@ -89,7 +89,7 @@ namespace UE::MultiUserClient
 
 		if (!RequestData.StreamChangeRequest.IsEmpty())
 		{
-			Operation->OnStreamChangesSubmittedFuture().Next([this, Token = LifetimeToken.ToWeakPtr()](FSubmitStreamChangesResponse&& Response)
+			Operation->OnCompleteStreamChangesFuture().Next([this, Token = LifetimeToken.ToWeakPtr()](FSubmitStreamChangesResponse&& Response)
 			{
 				if (const TSharedPtr<FToken> TokenPin = Token.Pin())
 				{
@@ -99,7 +99,7 @@ namespace UE::MultiUserClient
 		}
 		if (!RequestData.AuthorityRequest.IsEmpty())
 		{
-			Operation->OnAuthorityChangeResponseReceivedFuture().Next([this, Token = LifetimeToken.ToWeakPtr()](FSubmitAuthorityChangesResponse&& Response)
+			Operation->OnCompleteAuthorityChangeFuture().Next([this, Token = LifetimeToken.ToWeakPtr()](FSubmitAuthorityChangesResponse&& Response)
 			{
 				if (const TSharedPtr<FToken> TokenPin = Token.Pin())
 				{
diff --git a/Engine/Plugins/Developer/Concert/ConcertApp/MultiUserClient/Source/MultiUserClient/Private/Replication/Submission/Remote/SubmissionWorkflow_RemoteClient.cpp b/Engine/Plugins/Developer/Concert/ConcertApp/MultiUserClient/Source/MultiUserClient/Private/Replication/Submission/Remote/SubmissionWorkflow_RemoteClient.cpp
--- a/Engine/Plugins/Developer/Concert/ConcertApp/MultiUserClient/Source/MultiUserClient/Private/Replication/Submission/Remote/SubmissionWorkflow_RemoteClient.cpp
+++ b/Engine/Plugins/Developer/Concert/ConcertApp/MultiUserClient/Source/MultiUserClient/Private/Replication/Submission/Remote/SubmissionWorkflow_RemoteClient.cpp
@@ -138,7 +138,10 @@ namespace UE::MultiUserClient
 
 	void FSubmissionWorkflow_RemoteClient::CleanUpSubmissionProcess()
 	{
+		const TSharedRef<FSingleClientSubmissionOperation> ExposedOperation = InProgressOperation->ExposedOperation;
 		InProgressOperation.Reset();
+		
+		ExposedOperation->EmplaceCompleteOperationPromise(ESubmissionOperationCompletedCode::Processed);
 		OnSubmitOperationCompletedDelegate.Broadcast();
 	}
 
diff --git a/Engine/Plugins/Developer/Concert/ConcertApp/MultiUserClient/Source/MultiUserClient/Private/Replication/Submission/SingleClientSubmissionOperation.h b/Engine/Plugins/Developer/Concert/ConcertApp/MultiUserClient/Source/MultiUserClient/Private/Replication/Submission/SingleClientSubmissionOperation.h
--- a/Engine/Plugins/Developer/Concert/ConcertApp/MultiUserClient/Source/MultiUserClient/Private/Replication/Submission/SingleClientSubmissionOperation.h
+++ b/Engine/Plugins/Developer/Concert/ConcertApp/MultiUserClient/Source/MultiUserClient/Private/Replication/Submission/SingleClientSubmissionOperation.h
@@ -34,6 +34,10 @@ namespace UE::MultiUserClient
 			{
 				AuthorityChangeResponsePromise.EmplaceValue(FSubmitAuthorityChangesResponse{ EAuthoritySubmissionResponseErrorCode::Cancelled });
 			}
+			if (!bCompleteOperationPromiseWasSet)
+			{
+				CompleteOperationPromise.EmplaceValue(ESubmissionOperationCompletedCode::Cancelled);
+			}
 		}
 
 		void EmplaceStreamPromise(FSubmitStreamChangesResponse Result)
@@ -51,16 +55,23 @@ namespace UE::MultiUserClient
 			bAuthorityResponsePromiseWasSet = true;
 			AuthorityChangeResponsePromise.EmplaceValue(MoveTemp(Result));
 		}
+		void EmplaceCompleteOperationPromise(ESubmissionOperationCompletedCode Result)
+		{
+			bAuthorityResponsePromiseWasSet = true;
+			CompleteOperationPromise.EmplaceValue(Result);
+		}
 
 		bool HasSetStreamPromise() const { return bStreamPromiseWasSet; }
 		bool HasSetAuthorityRequestPromise() const { return bAuthorityRequestPromiseWasSet; }
 		bool HasSetAuthorityResponsePromise() const { return bAuthorityResponsePromiseWasSet; }
+		bool HasSetCompleteOperationPromise() const { return bCompleteOperationPromiseWasSet; }
 		
 		//~ Begin ISubmissionOperation Interface
 		virtual bool IsModifyingStreams() const override { return bModifiesStreams; }
-		virtual TFuture<FSubmitStreamChangesResponse> OnStreamChangesSubmittedFuture() override { return StreamChangesPromise.GetFuture();  }
-		virtual TFuture<FSubmitAuthorityChangesRequest> OnAuthorityChangeRequestedFuture() override { return AuthorityChangeRequestPromise.GetFuture(); }
-		virtual TFuture<FSubmitAuthorityChangesResponse> OnAuthorityChangeResponseReceivedFuture() override { return AuthorityChangeResponsePromise.GetFuture(); }
+		virtual TFuture<FSubmitStreamChangesResponse> OnCompleteStreamChangesFuture() override { return StreamChangesPromise.GetFuture();  }
+		virtual TFuture<FSubmitAuthorityChangesRequest> OnRequestAuthorityChangeFuture() override { return AuthorityChangeRequestPromise.GetFuture(); }
+		virtual TFuture<FSubmitAuthorityChangesResponse> OnCompleteAuthorityChangeFuture() override { return AuthorityChangeResponsePromise.GetFuture(); }
+		virtual TFuture<ESubmissionOperationCompletedCode> OnCompletedOperation() override { return CompleteOperationPromise.GetFuture(); }
 		//~ End ISubmissionOperation Interface
 
 	private:
@@ -70,10 +81,12 @@ namespace UE::MultiUserClient
 		bool bStreamPromiseWasSet = false;
 		bool bAuthorityRequestPromiseWasSet = false;
 		bool bAuthorityResponsePromiseWasSet = false;
+		bool bCompleteOperationPromiseWasSet = false;
 		
 		// All fulfilled by the owning FSingleClientSubmissionWorkflow
 		TPromise<FSubmitStreamChangesResponse> StreamChangesPromise;
 		TPromise<FSubmitAuthorityChangesRequest> AuthorityChangeRequestPromise;
 		TPromise<FSubmitAuthorityChangesResponse> AuthorityChangeResponsePromise;
+		TPromise<ESubmissionOperationCompletedCode> CompleteOperationPromise;
 	};
 }
diff --git a/Engine/Plugins/Developer/Concert/ConcertApp/MultiUserClient/Source/MultiUserClient/Private/Replication/Submission/SubmissionWorkflow_LocalClient.cpp b/Engine/Plugins/Developer/Concert/ConcertApp/MultiUserClient/Source/MultiUserClient/Private/Replication/Submission/SubmissionWorkflow_LocalClient.cpp
--- a/Engine/Plugins/Developer/Concert/ConcertApp/MultiUserClient/Source/MultiUserClient/Private/Replication/Submission/SubmissionWorkflow_LocalClient.cpp
+++ b/Engine/Plugins/Developer/Concert/ConcertApp/MultiUserClient/Source/MultiUserClient/Private/Replication/Submission/SubmissionWorkflow_LocalClient.cpp
@@ -147,7 +147,10 @@ namespace UE::MultiUserClient
 
 	void FSubmissionWorkflow_LocalClient::CleanUpSubmissionOperation()
 	{
+		const TSharedRef<FSingleClientSubmissionOperation> Operation = InProgressOperation->Operation;
 		InProgressOperation.Reset();
+		
+		Operation->EmplaceCompleteOperationPromise(ESubmissionOperationCompletedCode::Processed);
 		OnSubmitOperationCompletedDelegate.Broadcast();
 	}
 }
