Ownership reassignment combo box

diff --git a/Engine/Plugins/Developer/Concert/ConcertApp/MultiUserClient/Source/MultiUserClient/Private/MultiUserClientModule.cpp b/Engine/Plugins/Developer/Concert/ConcertApp/MultiUserClient/Source/MultiUserClient/Private/MultiUserClientModule.cpp
--- a/Engine/Plugins/Developer/Concert/ConcertApp/MultiUserClient/Source/MultiUserClient/Private/MultiUserClientModule.cpp
+++ b/Engine/Plugins/Developer/Concert/ConcertApp/MultiUserClient/Source/MultiUserClient/Private/MultiUserClientModule.cpp
@@ -51,6 +51,7 @@
 #include "DetailCategoryBuilder.h"
 #include "DetailWidgetRow.h"
 #include "MultiUserReplicationSettings.h"
+#include "MultiUserReplicationStyle.h"
 
 
 #include "Widgets/SConcertBrowser.h"
@@ -797,6 +798,7 @@ private:
 
 		// Initialize Style
 		FConcertFrontendStyle::Initialize();
+		UE::MultiUserClient::FMultiUserReplicationStyle::Initialize();
 
 		// Multi-User front end currently relies on EditorStyle being loaded
 		FModuleManager::LoadModuleChecked<IEditorStyleModule>("EditorStyle");
@@ -826,6 +828,9 @@ private:
 		UnregisterTabSpawner();
 
 #if WITH_EDITOR
+		FConcertFrontendStyle::Shutdown();
+		UE::MultiUserClient::FMultiUserReplicationStyle::Shutdown();
+		
 		UnregisterWorkspaceUI();
 		UnregisterSettings();
 		RemoveEditorToolbarButton();
diff --git a/Engine/Plugins/Developer/Concert/ConcertApp/MultiUserClient/Source/MultiUserClient/Private/MultiUserReplicationStyle.cpp b/Engine/Plugins/Developer/Concert/ConcertApp/MultiUserClient/Source/MultiUserClient/Private/MultiUserReplicationStyle.cpp
new file mode 100644
--- /dev/null
+++ b/Engine/Plugins/Developer/Concert/ConcertApp/MultiUserClient/Source/MultiUserClient/Private/MultiUserReplicationStyle.cpp
@@ -0,0 +1,71 @@
+// Copyright Epic Games, Inc. All Rights Reserved.
+
+#include "MultiUserReplicationStyle.h"
+
+#include "Brushes/SlateBoxBrush.h"
+#include "Brushes/SlateImageBrush.h"
+#include "Interfaces/IPluginManager.h"
+#include "Misc/Paths.h"
+
+#include "Styling/SlateTypes.h"
+#include "Styling/SlateStyleRegistry.h"
+#include "Styling/StarshipCoreStyle.h"
+#include "Styling/SlateStyle.h"
+
+namespace UE::MultiUserClient
+{
+	FString FMultiUserReplicationStyle::InContent(const FString& RelativePath, const ANSICHAR* Extension)
+	{
+		static FString ContentDir = IPluginManager::Get().FindPlugin(TEXT("ConcertSharedSlate"))->GetContentDir();
+		return (ContentDir / RelativePath) + Extension;
+	}
+
+	TSharedPtr< class FSlateStyleSet > FMultiUserReplicationStyle::StyleSet;
+
+	FName FMultiUserReplicationStyle::GetStyleSetName()
+	{
+		return FName(TEXT("MultiUserReplicationStyle"));
+	}
+
+	void FMultiUserReplicationStyle::Initialize()
+	{
+		// Only register once
+		if (StyleSet.IsValid())
+		{
+			return;
+		}
+
+		StyleSet = MakeShared<FSlateStyleSet>(GetStyleSetName());
+		StyleSet->SetContentRoot(FPaths::EngineContentDir() / TEXT("Editor/Slate"));
+		StyleSet->SetCoreContentRoot(FPaths::EngineContentDir() / TEXT("Slate"));
+
+		// Column widths
+		StyleSet->Set("AllClients.Object.OwnerColumnWidth", 200.f);
+		StyleSet->Set("AllClients.Object.ReplicationToggle", 50.f);
+		StyleSet->Set("AllClients.Property.OwnerColumnWidth", 200.f);
+		StyleSet->Set("SingleClient.Object.OwnerColumnWidth", 200.f);
+		StyleSet->Set("SingleClient.Property.OwnerColumnWidth", 200.f);
+
+		// Timing
+		StyleSet->Set("AllClients.Reassignment.DisplayThrobberAfterSeconds", 0.2f);
+		
+		FSlateStyleRegistry::RegisterSlateStyle(*StyleSet.Get());
+	};
+
+	void FMultiUserReplicationStyle::Shutdown()
+	{
+		if (StyleSet.IsValid())
+		{
+			FSlateStyleRegistry::UnRegisterSlateStyle(*StyleSet.Get());
+			ensure(StyleSet.IsUnique());
+			StyleSet.Reset();
+		}
+	}
+
+	TSharedPtr<class ISlateStyle> FMultiUserReplicationStyle::Get()
+	{
+		return StyleSet;
+	}
+}
+
+
diff --git a/Engine/Plugins/Developer/Concert/ConcertApp/MultiUserClient/Source/MultiUserClient/Private/MultiUserReplicationStyle.h b/Engine/Plugins/Developer/Concert/ConcertApp/MultiUserClient/Source/MultiUserClient/Private/MultiUserReplicationStyle.h
new file mode 100644
--- /dev/null
+++ b/Engine/Plugins/Developer/Concert/ConcertApp/MultiUserClient/Source/MultiUserClient/Private/MultiUserReplicationStyle.h
@@ -0,0 +1,33 @@
+// Copyright Epic Games, Inc. All Rights Reserved.
+
+#pragma once
+
+#include "Styling/ISlateStyle.h"
+#include "Styling/SlateStyle.h"
+
+class FSlateStyleSet;
+class ISlateStyle;
+
+namespace UE::MultiUserClient
+{
+	/**
+	 * Style settings specifically for all replication widgets.
+	 * Separate class to make it more straight forward to move to a separate module (there are no specific plans for that only future proofing).
+	 */
+	class FMultiUserReplicationStyle
+	{
+	public:
+	
+		static void Initialize();
+		static void Shutdown();
+
+		static TSharedPtr<ISlateStyle> Get();
+
+		static FName GetStyleSetName();
+	
+	private:
+	
+		static FString InContent(const FString& RelativePath, const ANSICHAR* Extension);
+		static TSharedPtr<FSlateStyleSet> StyleSet;
+	};
+}
diff --git a/Engine/Plugins/Developer/Concert/ConcertApp/MultiUserClient/Source/MultiUserClient/Private/Replication/Client/ReplicationClientManager.cpp b/Engine/Plugins/Developer/Concert/ConcertApp/MultiUserClient/Source/MultiUserClient/Private/Replication/Client/ReplicationClientManager.cpp
--- a/Engine/Plugins/Developer/Concert/ConcertApp/MultiUserClient/Source/MultiUserClient/Private/Replication/Client/ReplicationClientManager.cpp
+++ b/Engine/Plugins/Developer/Concert/ConcertApp/MultiUserClient/Source/MultiUserClient/Private/Replication/Client/ReplicationClientManager.cpp
@@ -28,6 +28,7 @@ namespace UE::MultiUserClient
 			return FLocalReplicationClient(AuthorityCache, *ClientPreset, MakeUnique<FStreamSynchronizer_LocalClient>(InClient, ClientPreset->Stream->StreamId), InClient);
 		}())
 		, SubmissionNotifier(*this)
+		, ReassignmentLogic(*this)
 	{
 		AuthorityCache.RegisterEvents();
 		InSession->OnSessionClientChanged().AddRaw(this, &FReplicationClientManager::OnSessionClientChanged);
@@ -46,7 +47,17 @@ namespace UE::MultiUserClient
 		}
 	}
 
-	TArray<TNonNullPtr<FRemoteReplicationClient>> FReplicationClientManager::GetRemoteClients() const
+	TArray<TNonNullPtr<const FRemoteReplicationClient>> FReplicationClientManager::GetRemoteClients() const
+	{
+		TArray<TNonNullPtr<const FRemoteReplicationClient>> Result;
+		Algo::Transform(RemoteClients, Result, [](const TUniquePtr<FRemoteReplicationClient>& Client) -> TNonNullPtr<const FRemoteReplicationClient>
+		{
+			return Client.Get();
+		});
+		return Result;
+	}
+
+	TArray<TNonNullPtr<FRemoteReplicationClient>> FReplicationClientManager::GetRemoteClients()
 	{
 		TArray<TNonNullPtr<FRemoteReplicationClient>> Result;
 		Algo::Transform(RemoteClients, Result, [](const TUniquePtr<FRemoteReplicationClient>& Client) -> TNonNullPtr<FRemoteReplicationClient>
@@ -56,6 +67,20 @@ namespace UE::MultiUserClient
 		return Result;
 	}
 
+	TArray<TNonNullPtr<const FReplicationClient>> FReplicationClientManager::GetClients(TFunctionRef<bool(const FReplicationClient& Client)> Predicate) const
+	{
+		TArray<TNonNullPtr<const FReplicationClient>> Result;
+		ForEachClient([&Predicate, &Result](const FReplicationClient& Client)
+		{
+			if (Predicate(Client))
+			{
+				Result.Add(&Client);
+			}
+			return EBreakBehavior::Continue;
+		});
+		return Result;
+	}
+
 	const FRemoteReplicationClient* FReplicationClientManager::FindRemoteClient(const FGuid& EndpointId) const
 	{
 		const TUniquePtr<FRemoteReplicationClient>* Client = RemoteClients.FindByPredicate([&EndpointId](const TUniquePtr<FRemoteReplicationClient>& Client)
@@ -71,7 +96,7 @@ namespace UE::MultiUserClient
 		{
 			return;
 		}
-		for (const TNonNullPtr<FRemoteReplicationClient>& RemoteClient : GetRemoteClients())
+		for (const TNonNullPtr<const FRemoteReplicationClient>& RemoteClient : GetRemoteClients())
 		{
 			if (ProcessClient(*RemoteClient) == EBreakBehavior::Break)
 			{
diff --git a/Engine/Plugins/Developer/Concert/ConcertApp/MultiUserClient/Source/MultiUserClient/Private/Replication/Client/ReplicationClientManager.h b/Engine/Plugins/Developer/Concert/ConcertApp/MultiUserClient/Source/MultiUserClient/Private/Replication/Client/ReplicationClientManager.h
--- a/Engine/Plugins/Developer/Concert/ConcertApp/MultiUserClient/Source/MultiUserClient/Private/Replication/Client/ReplicationClientManager.h
+++ b/Engine/Plugins/Developer/Concert/ConcertApp/MultiUserClient/Source/MultiUserClient/Private/Replication/Client/ReplicationClientManager.h
@@ -6,6 +6,7 @@
 #include "ConcertMessageData.h"
 #include "LocalReplicationClient.h"
 #include "ReplicationClient.h"
+#include "Replication/Submission/MultiEdit/ReassignObjectPropertiesLogic.h"
 #include "Replication/Submission/Notification/SubmissionNotifier.h"
 #include "Replication/Util/GlobalAuthorityCache.h"
 #include "Replication/Util/RegularQueryService.h"
@@ -51,10 +52,17 @@ namespace UE::MultiUserClient
 
 		const FLocalReplicationClient& GetLocalClient() const { return LocalClient; }
 		FLocalReplicationClient& GetLocalClient() { return LocalClient; }
-		TArray<TNonNullPtr<FRemoteReplicationClient>> GetRemoteClients() const;
+		TArray<TNonNullPtr<const FRemoteReplicationClient>> GetRemoteClients() const;
+		TArray<TNonNullPtr<FRemoteReplicationClient>> GetRemoteClients();
+
+		/** @return Gets all clients given a predicate */
+		TArray<TNonNullPtr<const FReplicationClient>> GetClients(TFunctionRef<bool(const FReplicationClient& Client)> Predicate = [](const auto&){ return true; }) const;
 		
 		const FGlobalAuthorityCache& GetAuthorityCache() const { return AuthorityCache; }
 		FGlobalAuthorityCache& GetAuthorityCache() { return AuthorityCache; }
+		
+		const FReassignObjectPropertiesLogic& GetReassignmentLogic() const { return ReassignmentLogic; }
+		FReassignObjectPropertiesLogic& GetReassignmentLogic() { return ReassignmentLogic; }
 
 		/** Util for finding a remote client by its EndpointId. */
 		const FRemoteReplicationClient* FindRemoteClient(const FGuid& EndpointId) const;
@@ -123,7 +131,6 @@ namespace UE::MultiUserClient
 		
 		/** Manages the local client */
 		FLocalReplicationClient LocalClient;
-		
 		/**
 		 * Manages remote clients. Updated when client connects or disconnects to the active session.
 		 * UI keeps references to systems inside the client so it is TSharedRef in case TArray is reallocated.
@@ -137,8 +144,14 @@ namespace UE::MultiUserClient
 		/** Called just before a remote client is about to be removed from RemoteClients. */
 		FRemoteClientDelegate OnPreRemoteClientRemovedDelegate; 
 		
-		/** Manages SNotificationItems when submission to the server fails. */
+		/**
+		 * Manages SNotificationItems when submission to the server fails.
+		 * TODO UE-200925: This should be moved to the client so we can customize messages depending on whether it is a remote or local client
+		 */
 		FSubmissionNotifier SubmissionNotifier;
+
+		/** Used for transferring ownership from multiple clients to one. Used by multi client view. */
+		FReassignObjectPropertiesLogic ReassignmentLogic;
 		
 		/** Updates RemoteClients depending on the change. */
 		void OnSessionClientChanged(IConcertClientSession&, EConcertClientStatus NewStatus, const FConcertSessionClientInfo& ClientInfo);
diff --git a/Engine/Plugins/Developer/Concert/ConcertApp/MultiUserClient/Source/MultiUserClient/Private/Replication/Stream/StreamChangeTracker.h b/Engine/Plugins/Developer/Concert/ConcertApp/MultiUserClient/Source/MultiUserClient/Private/Replication/Stream/StreamChangeTracker.h
--- a/Engine/Plugins/Developer/Concert/ConcertApp/MultiUserClient/Source/MultiUserClient/Private/Replication/Stream/StreamChangeTracker.h
+++ b/Engine/Plugins/Developer/Concert/ConcertApp/MultiUserClient/Source/MultiUserClient/Private/Replication/Stream/StreamChangeTracker.h
@@ -3,9 +3,9 @@
 #pragma once
 
 #include "IClientStreamSynchronizer.h"
-#include "Replication/Data/ObjectIds.h"
 #include "Replication/Data/ObjectReplicationMap.h"
 #include "Replication/Messages/ChangeStream.h"
+#include "Replication/Util/StreamRequestUtils.h"
 
 #include "Async/Future.h"
 #include "Containers/Set.h"
@@ -30,13 +30,6 @@ namespace UE::MultiUserClient
 	};
 	ENUM_CLASS_FLAGS(EObjectWarningFlags);
 	
-	/** Describes changes that MU client makes. */
-	struct FStreamChangelist
-	{
-		TSet<FObjectInStreamID> ObjectsToRemove;
-		TMap<FObjectInStreamID, FConcertReplication_ChangeStream_PutObject> ObjectsToPut;
-	};
-	
 	/**
 	 * Knows of the local client's registered replication streams and builds a changelist. The changelist tracks the
 	 * unconfirmed changes to the client's streams and is updates when the server confirms the change.
diff --git a/Engine/Plugins/Developer/Concert/ConcertApp/MultiUserClient/Source/MultiUserClient/Private/Replication/Submission/AutoSubmissionPolicy.h b/Engine/Plugins/Developer/Concert/ConcertApp/MultiUserClient/Source/MultiUserClient/Private/Replication/Submission/AutoSubmissionPolicy.h
--- a/Engine/Plugins/Developer/Concert/ConcertApp/MultiUserClient/Source/MultiUserClient/Private/Replication/Submission/AutoSubmissionPolicy.h
+++ b/Engine/Plugins/Developer/Concert/ConcertApp/MultiUserClient/Source/MultiUserClient/Private/Replication/Submission/AutoSubmissionPolicy.h
@@ -6,11 +6,6 @@
 #include "Replication/Submission/Queue/DeferredSubmitter.h"
 #include "Templates/UnrealTemplate.h"
 
-namespace UE::MultiUserClient
-{
-	class FChangeRequestBuilder;
-}
-
 namespace UE::ConcertClientSharedSlate
 {
 	class IEditableReplicationStreamModel;
@@ -18,6 +13,7 @@ namespace UE::ConcertClientSharedSlate
 
 namespace UE::MultiUserClient
 {
+	class FChangeRequestBuilder;
 	class IClientStreamSynchronizer;
 	class ISubmissionWorkflow;
 	class FAuthorityChangeTracker;
diff --git a/Engine/Plugins/Developer/Concert/ConcertApp/MultiUserClient/Source/MultiUserClient/Private/Replication/Submission/Data/AuthoritySubmission.h b/Engine/Plugins/Developer/Concert/ConcertApp/MultiUserClient/Source/MultiUserClient/Private/Replication/Submission/Data/AuthoritySubmission.h
--- a/Engine/Plugins/Developer/Concert/ConcertApp/MultiUserClient/Source/MultiUserClient/Private/Replication/Submission/Data/AuthoritySubmission.h
+++ b/Engine/Plugins/Developer/Concert/ConcertApp/MultiUserClient/Source/MultiUserClient/Private/Replication/Submission/Data/AuthoritySubmission.h
@@ -53,7 +53,7 @@ namespace UE::MultiUserClient
 	
 	struct FSubmitAuthorityChangesResponse
 	{
-		EAuthoritySubmissionResponseErrorCode ErrorCode;
+		EAuthoritySubmissionResponseErrorCode ErrorCode = EAuthoritySubmissionResponseErrorCode::Cancelled;
 		
 		/** Only valid if ErrorCode == EAuthoritySubmissionErrorCode::Success */
 		TOptional<ConcertSyncClient::Replication::FAuthorityChangeResponse> Response;
diff --git a/Engine/Plugins/Developer/Concert/ConcertApp/MultiUserClient/Source/MultiUserClient/Private/Replication/Submission/Data/StreamSubmission.h b/Engine/Plugins/Developer/Concert/ConcertApp/MultiUserClient/Source/MultiUserClient/Private/Replication/Submission/Data/StreamSubmission.h
--- a/Engine/Plugins/Developer/Concert/ConcertApp/MultiUserClient/Source/MultiUserClient/Private/Replication/Submission/Data/StreamSubmission.h
+++ b/Engine/Plugins/Developer/Concert/ConcertApp/MultiUserClient/Source/MultiUserClient/Private/Replication/Submission/Data/StreamSubmission.h
@@ -34,7 +34,7 @@ namespace UE::MultiUserClient
 	struct FSubmitStreamChangesResponse
 	{
 		/** Error code of the submission. Determines whether Response is valid. */
-		EStreamSubmissionErrorCode ErrorCode;
+		EStreamSubmissionErrorCode ErrorCode = EStreamSubmissionErrorCode::Cancelled;
 
 		/** Valid if ErrorCode == ESubmitChangesErrorCode::Success. */
 		TOptional<FCompletedChangeSubmission> SubmissionInfo;
diff --git a/Engine/Plugins/Developer/Concert/ConcertApp/MultiUserClient/Source/MultiUserClient/Private/Replication/Submission/MultiEdit/ReassignObjectPropertiesLogic.cpp b/Engine/Plugins/Developer/Concert/ConcertApp/MultiUserClient/Source/MultiUserClient/Private/Replication/Submission/MultiEdit/ReassignObjectPropertiesLogic.cpp
new file mode 100644
--- /dev/null
+++ b/Engine/Plugins/Developer/Concert/ConcertApp/MultiUserClient/Source/MultiUserClient/Private/Replication/Submission/MultiEdit/ReassignObjectPropertiesLogic.cpp
@@ -0,0 +1,498 @@
+// Copyright Epic Games, Inc. All Rights Reserved.
+
+#include "ReassignObjectPropertiesLogic.h"
+
+#include "ConcertLogGlobal.h"
+#include "Replication/Client/ReplicationClientManager.h"
+#include "Replication/Submission/ISubmissionOperation.h"
+#include "Replication/Util/SynchronizedRequestUtils.h"
+
+#include "Algo/AllOf.h"
+#include "Algo/AnyOf.h"
+#include "Async/Async.h"
+
+#define LOCTEXT_NAMESPACE "FReassignObjectPropertiesLogic"
+
+namespace UE::MultiUserClient
+{
+	namespace ReassignObjectProperties::Private
+	{
+		struct FObjectReassignment
+		{
+			/** Used to build the final changelist to apply to the target client. Maps ClientIds to the object data they had registered. */
+			TMap<FGuid, FObjectReplicationMap> OldRegisteredObjects;
+			/** Changes to make to all the other clients */
+			TMap<FGuid, ConcertSyncClient::Replication::FChangeStreamRequest> ReassignedClientRequests;
+			/** Authority that the clients had before. This is used to determine whether to give the assigned to client authority once the stream changes have been made. */
+			TMap<FGuid, TArray<FSoftObjectPath>> ReassignedAuthority;
+		};
+		
+		/** Builds all the changes that must be made to the clients we are transferring from.  */
+		static FObjectReassignment BuildChangesForTransferal(const FReplicationClientManager& ClientManager, TConstArrayView<FSoftObjectPath> ObjectsToReassign, const FGuid& StreamId)
+		{
+			FObjectReassignment Result;
+			TMap<FGuid, FObjectReplicationMap>& OldRegisteredObjects = Result.OldRegisteredObjects;
+			TMap<FGuid, ConcertSyncClient::Replication::FChangeStreamRequest>& ReassignedClientRequests = Result.ReassignedClientRequests;
+			TMap<FGuid, TArray<FSoftObjectPath>>& ReassignedAuthority = Result.ReassignedAuthority;
+		
+			for (const FSoftObjectPath& ObjectPath : ObjectsToReassign)
+			{
+				const FObjectInStreamID ObjectId{ StreamId, ObjectPath };
+
+				const FGlobalAuthorityCache& AuthorityCache = ClientManager.GetAuthorityCache();
+				AuthorityCache.ForEachClientWithObjectInStream(ObjectPath, [&ClientManager, &OldRegisteredObjects, &ReassignedClientRequests, &ReassignedAuthority, &ObjectPath, &ObjectId, &AuthorityCache](const FGuid& ClientId)
+				{
+					// Get the property assignments of the client being reassigned from
+					const FReplicationClient* ClientToReassignFrom = ClientManager.FindClient(ClientId);
+					if (!ensureMsgf(ClientToReassignFrom, TEXT("Authority cache out of sync"))
+						|| !ClientToReassignFrom->AllowsEditing())
+					{
+						return EBreakBehavior::Continue;
+					}
+					const FReplicatedObjectInfo* ObjectInfo = ClientToReassignFrom->GetStreamSynchronizer().GetServerState().ReplicatedObjects.Find(ObjectId.Object);
+					if (!ensureMsgf(ObjectInfo, TEXT("ForEachClientWithObjectInStream lied")))
+					{
+						return EBreakBehavior::Continue;
+					}
+
+					// Transfer from client being reassigned from to the client being assigned to
+					OldRegisteredObjects.FindOrAdd(ClientId).ReplicatedObjects.Add(ObjectPath, *ObjectInfo);
+					// Remove the object completely from the reassigned client.
+					ReassignedClientRequests.FindOrAdd(ClientId).ObjectsToRemove.Add(ObjectId);
+					// Later we can determine whether to give the assigned to client authority
+					if (AuthorityCache.HasAuthorityOverObject(ObjectPath, ClientId))
+					{
+						ReassignedAuthority.FindOrAdd(ClientId).Add(ObjectPath);
+					}
+				
+					return EBreakBehavior::Continue;
+				});
+			}
+
+			return Result;
+		}
+
+		/** Builds a changelist based on which remote clients we managed to change successfully. */
+		static FStreamChangelist MakeChangelistFromAppliedChanges(
+			const TMap<FGuid, FObjectReplicationMap>& OldRegisteredObjects,
+			const FParallelExecutionResult& ParallelExecutionResult,
+			const FObjectReplicationMap& CurrentTargetClientState,
+			const FGuid& TargetStreamId
+			)
+		{
+			FStreamChangelist Result;
+			for (const TPair<FGuid, FSubmitStreamChangesResponse>& ChangesRequestedOnRemote : ParallelExecutionResult.StreamResponses)
+			{
+				const TOptional<FCompletedChangeSubmission>& SubmissionInfo = ChangesRequestedOnRemote.Value.SubmissionInfo;
+				if (!SubmissionInfo || SubmissionInfo->Response.IsFailure())
+				{
+					continue;
+				}
+				
+				const FObjectReplicationMap& ObjectReplicationMap = OldRegisteredObjects[ChangesRequestedOnRemote.Key];
+				for (const TPair<FSoftObjectPath, FReplicatedObjectInfo>& ChangesToApply : ObjectReplicationMap.ReplicatedObjects)
+				{
+					const FObjectInStreamID ObjectId { TargetStreamId, ChangesToApply.Key };
+					TOptional<FConcertReplication_ChangeStream_PutObject> PutRequest = FConcertReplication_ChangeStream_PutObject::MakeFromInfo(ChangesToApply.Value);
+					if (!ensure(PutRequest))
+					{
+						continue;
+					}
+
+					// We want to append the other client's properties to the ones the target client already has.
+					// TODO UE-201166: This step would not be necessary if we had an append operation in FConcertReplication_ChangeStream_PutObject
+					const FReplicatedObjectInfo* CurrentInfo = CurrentTargetClientState.ReplicatedObjects.Find(ObjectId.Object);
+					if (CurrentInfo)
+					{
+						for (const FConcertPropertyChain& PropertyChain : CurrentInfo->PropertySelection.ReplicatedProperties)
+						{
+							PutRequest->Properties.ReplicatedProperties.AddUnique(PropertyChain);
+						}
+					}
+					
+					Result.ObjectsToPut.Add(ObjectId, MoveTemp(*PutRequest));
+				}
+			}
+			return Result;
+		}
+
+		/** Builds list of authority changes based on which clients we managed to change successfully. */
+		static ConcertSyncClient::Replication::FAuthorityChangeRequest MakeAuthorityRequestFrom(
+			const TMap<FGuid, TArray<FSoftObjectPath>>& ReassignedAuthority, 
+			const FParallelExecutionResult& ParallelExecutionResult,
+			const FGuid& TargetStreamId
+			)
+		{
+			ConcertSyncClient::Replication::FAuthorityChangeRequest ResultRequest;
+			for (const TPair<FGuid, FSubmitStreamChangesResponse>& ChangesRequestedOnRemote : ParallelExecutionResult.StreamResponses)
+			{
+				const TOptional<FCompletedChangeSubmission>& SubmissionInfo = ChangesRequestedOnRemote.Value.SubmissionInfo;
+				if (!SubmissionInfo || SubmissionInfo->Response.IsFailure())
+				{
+					continue;
+				}
+
+				const TArray<FSoftObjectPath>* PreviousAuthority = ReassignedAuthority.Find(ChangesRequestedOnRemote.Key);
+				if (!PreviousAuthority)
+				{
+					continue;
+				}
+
+				for (const FSoftObjectPath& ObjectPath : *PreviousAuthority)
+				{
+					ResultRequest.TakeAuthority.Add(ObjectPath, FConcertStreamArray{{ TargetStreamId }});
+				}
+			}
+			return ResultRequest;
+		}
+
+		/** @return Whether TargetClientObjects includes all properties ClientsToConsider have registered for ObjectToCheck. */
+		static bool DoesTargetIncludeOthers_SingleObject(
+			const FObjectReplicationMap& TargetClientObjects,
+			const TArray<TNonNullPtr<const FReplicationClient>> ClientsToConsider,
+			const FSoftObjectPath& ObjectToCheck)
+		{
+			const FReplicatedObjectInfo* ObjectInfo = TargetClientObjects.ReplicatedObjects.Find(ObjectToCheck);
+				
+			for (const FReplicationClient* OtherClient : ClientsToConsider)
+			{
+				const FObjectReplicationMap& OtherClientObjects = OtherClient->GetStreamSynchronizer().GetServerState();
+				const FReplicatedObjectInfo* OtherObjectInfo = OtherClientObjects.ReplicatedObjects.Find(ObjectToCheck);
+					
+				// TargetClient has at least as much if OtherClient has nothing
+				if (!OtherObjectInfo)
+				{
+					continue;
+				}
+
+				// TargetClient has nothing but TargetClient has something?
+				if (!ObjectInfo && !OtherObjectInfo->PropertySelection.ReplicatedProperties.IsEmpty())
+				{
+					return false;
+				}
+
+				if (!ObjectInfo->PropertySelection.Includes(OtherObjectInfo->PropertySelection))
+				{
+					return false;
+				}
+			}
+
+			return true;
+		}
+		
+		/** @return Whether ClientId includes all properties assigned to ObjectsToCheck as all other clients (passing the predicate) do. */
+		static bool DoesTargetIncludeOthers(
+			const FGuid& TargetClientId,
+			TConstArrayView<FSoftObjectPath> ObjectsToCheck,
+			const FReplicationClientManager& ClientManager,
+			TFunctionRef<bool(const FReplicationClient& Client)> ShouldConsiderClientPredicate
+			)
+		{
+			const FReplicationClient* TargetClient = ClientManager.FindClient(TargetClientId);
+			if (!ensure(TargetClient))
+			{
+				return false;
+			}
+
+			const FObjectReplicationMap& TargetClientObjects = TargetClient->GetStreamSynchronizer().GetServerState();
+			const TArray<TNonNullPtr<const FReplicationClient>> ClientsToConsider = ClientManager.GetClients(ShouldConsiderClientPredicate);
+			return Algo::AllOf(ObjectsToCheck, [&TargetClientObjects, &ClientsToConsider](const FSoftObjectPath& ObjectToCheck)
+			{
+				return DoesTargetIncludeOthers_SingleObject(TargetClientObjects, ClientsToConsider, ObjectToCheck);
+			});
+		}
+
+		/** @return Whether ClientId has any properties assigned to any of the objects in ObjectsToCheck. */
+		static bool HasAnyProperties(const FGuid& ClientId, TConstArrayView<FSoftObjectPath> ObjectsToCheck, const FReplicationClientManager& ClientManager)
+		{
+			const FReplicationClient* TargetClient = ClientManager.FindClient(ClientId);
+			if (!ensure(TargetClient))
+			{
+				return false;
+			}
+
+			const FObjectReplicationMap& TargetClientObjects = TargetClient->GetStreamSynchronizer().GetServerState();
+			return Algo::AnyOf(ObjectsToCheck, [&TargetClientObjects](const FSoftObjectPath& ObjectPath)
+			{
+				const FReplicatedObjectInfo* ObjectInfo = TargetClientObjects.ReplicatedObjects.Find(ObjectPath);
+				return ObjectInfo && !ObjectInfo->PropertySelection.ReplicatedProperties.IsEmpty();
+			});
+		}
+	}
+	
+	FReassignObjectPropertiesLogic::FReassignObjectPropertiesLogic(FReplicationClientManager& InClientManager)
+		: ClientManager(InClientManager)
+	{
+		ClientManager.OnRemoteClientsChanged().AddRaw(this, &FReassignObjectPropertiesLogic::OnRemoteClientsChanged);
+		ClientManager.GetAuthorityCache().OnCacheChanged().AddRaw(this, &FReassignObjectPropertiesLogic::OnClientCacheChanged);
+	}
+
+	FReassignObjectPropertiesLogic::~FReassignObjectPropertiesLogic()
+	{
+		// Note that all of this clean up is not REALLY needed because the FReassignObjectPropertiesLogic is a member of ClientManager.
+		// However, we'll follow good RAII here in case that should ever change.
+		ClientManager.OnRemoteClientsChanged().RemoveAll(this);
+		ClientManager.GetAuthorityCache().OnCacheChanged().RemoveAll(this);
+
+		if (FSubmissionQueue* SubmissionQueue = FindSubmissionQueueForTargetClient())
+		{
+			SubmissionQueue->Dequeue_GameThread(*this);
+		}
+	}
+
+	void FReassignObjectPropertiesLogic::EnumerateClientOwnershipState(const FSoftObjectPath& ObjectPath, FProcessClientOwnership Callback) const
+	{
+		ClientManager.ForEachClient([&ObjectPath, &Callback](const FReplicationClient& ReplicationClient)
+		{
+			const bool bHasProperties = ReplicationClient.GetStreamSynchronizer().GetServerState().HasProperties(ObjectPath);
+			const EOwnershipState Ownership = bHasProperties ? EOwnershipState::HasObjectRegistered : EOwnershipState::NoOwnership;
+			return Callback(ReplicationClient.GetEndpointId(), Ownership);
+		});
+	}
+
+	bool FReassignObjectPropertiesLogic::OwnsAnyOf(TConstArrayView<FSoftObjectPath> Objects, const FGuid& TargetClientId) const
+	{
+		bool bOwnsAny = false;
+		for (int32 i = 0; i < Objects.Num() && !bOwnsAny; ++i)
+		{
+			EnumerateClientOwnershipState(Objects[i], [&bOwnsAny, &TargetClientId](const FGuid& ClientId, EOwnershipState Ownership)
+			{
+				bOwnsAny = TargetClientId == ClientId && Ownership == EOwnershipState::HasObjectRegistered;
+				return bOwnsAny ? EBreakBehavior::Break : EBreakBehavior::Continue;
+			});
+		}
+		return bOwnsAny;
+	}
+
+	bool FReassignObjectPropertiesLogic::IsAnyObjectOwned(TConstArrayView<FSoftObjectPath> Objects) const
+	{
+		bool bOwnsAny = false;
+		for (int32 i = 0; i < Objects.Num() && !bOwnsAny; ++i)
+		{
+			EnumerateClientOwnershipState(Objects[i], [&bOwnsAny](const FGuid& ClientId, EOwnershipState Ownership)
+			{
+				bOwnsAny = Ownership == EOwnershipState::HasObjectRegistered;
+				return bOwnsAny ? EBreakBehavior::Break : EBreakBehavior::Continue;
+			});
+		}
+		return bOwnsAny;
+	}
+
+#define SET_REASON(Text) if (Reason) { *Reason = Text; }
+	bool FReassignObjectPropertiesLogic::CanReassignAnyTo(TConstArrayView<FSoftObjectPath> ObjectsToReassign, const FGuid& ClientId, FText* Reason) const
+	{
+		// TODO UE-201309: Implement queue for reassignment operations for better UX
+		if (InProgressOperation)
+		{
+			SET_REASON(LOCTEXT("OperationInProgress", "Wait for the previous reassignment operation to finish."));
+			return false;
+		}
+		
+		const FReplicationClient* Client = ClientManager.FindClient(ClientId);
+		if (!Client)
+		{
+			SET_REASON(LOCTEXT("Reason.ClientDisconnected", "Client disconnected"));
+			return false;
+		}
+		
+		if (!Client->AllowsEditing())
+		{
+			SET_REASON(LOCTEXT("Reason.ClientNoEditing", "Client does not allow remote editing"));
+			return false;
+		}
+
+		// Relevant after user has re-assigned all properties to a given client
+		auto ShouldConsiderClient = [this](const FReplicationClient& Client) { return Client.AllowsEditing(); };
+		if (ReassignObjectProperties::Private::DoesTargetIncludeOthers(ClientId, ObjectsToReassign, ClientManager, ShouldConsiderClient))
+		{
+			const bool bHasAnyProperties = ReassignObjectProperties::Private::HasAnyProperties(ClientId, ObjectsToReassign, ClientManager);
+			SET_REASON(
+				FText::Format(
+					LOCTEXT("Reason.NothingToAssignFmt", "Nothing to assign: {0}"),
+					bHasAnyProperties ? LOCTEXT("OwnsAll", "client already owns all properties.") : LOCTEXT("HasNone", "assign some properties first.")
+					)
+				);
+			return false;
+		}
+
+		return true;
+	}
+#undef SET_REASON
+	
+	void FReassignObjectPropertiesLogic::ReassignAllTo(TArray<FSoftObjectPath> ObjectsToReassign, const FGuid& ClientId)
+	{
+		using namespace ConcertSyncClient::Replication;
+		
+		const FReplicationClient* ClientToAssignTo = ClientManager.FindClient(ClientId);
+		if (!ensure(ClientToAssignTo && ClientToAssignTo->AllowsEditing()))
+		{
+			UE_LOG(LogConcert, Error, TEXT("Property Reassignment: The target client is not editable."));
+			return;
+		}
+
+		const IClientStreamSynchronizer& StreamSynchronizer = ClientToAssignTo->GetStreamSynchronizer();
+		const FGuid& TargetStreamId = StreamSynchronizer.GetStreamId();
+		auto[OldRegisteredObjects, ReassignedClientRequests, ReassignedAuthority]
+			= ReassignObjectProperties::Private::BuildChangesForTransferal(
+				ClientManager,
+				ObjectsToReassign,
+				TargetStreamId
+			);
+
+		// ParallelSubmissionOperation's RAII cleans up on destruction, e.g. dequeues submission tasks, so InProgressOperation keeps a reference to it.
+		const TSharedPtr<IParallelSubmissionOperation> ParallelSubmissionOperation = ExecuteParallelStreamChanges(ClientManager, ReassignedClientRequests);
+		if (!ensureMsgf(ParallelSubmissionOperation, TEXT("Investigate why starting task failed")))
+		{
+			return;
+		}
+		
+		InProgressOperation = { MoveTemp(ObjectsToReassign), ClientId, MoveTemp(OldRegisteredObjects), MoveTemp(ReassignedAuthority), ParallelSubmissionOperation.ToSharedRef() };
+		ParallelSubmissionOperation->GetFuture()
+			.Next([this](FParallelExecutionResult&& Result)
+			{
+				// When "this" is being destroyed, there is no sense in further processing.
+				// IParallelSubmissionOperation's RAII sets bWasCancelled = true when it is destroyed.
+				// The IParallelSubmissionOperation is destroyed just before "this" due to InProgressOperation keeping the only strong reference.
+				if (!Result.bWasCancelled)
+				{
+					OnAssignedFromClientsCompleted(Result);
+				}
+			});
+	}
+
+	bool FReassignObjectPropertiesLogic::IsReassigning(const FSoftObjectPath& ObjectPath) const
+	{
+		if (!InProgressOperation)
+		{
+			return false;
+		}
+
+		return Algo::AnyOf(InProgressOperation->OldRegisteredObjects, [&ObjectPath](const TPair<FGuid, FObjectReplicationMap>& Change)
+		{
+			const FReplicatedObjectInfo* ObjectInfo = Change.Value.ReplicatedObjects.Find(ObjectPath);
+			return ObjectInfo && !ObjectInfo->PropertySelection.ReplicatedProperties.IsEmpty();
+		});
+	}
+
+	TOptional<FDateTime> FReassignObjectPropertiesLogic::GetTimeReassignmentWasStarted() const
+	{
+		return InProgressOperation ? InProgressOperation->StartTime : TOptional<FDateTime>{};
+	}
+
+	void FReassignObjectPropertiesLogic::OnRemoteClientsChanged()
+	{
+		if (InProgressOperation && !ClientManager.FindClient(InProgressOperation->AssignedToClient))
+		{
+			UE_LOG(LogConcert, Warning,
+				TEXT("Client %s disconnected while reassigning objects to it. Cancelling the operation however object may already have been removed on other clients."),
+				*InProgressOperation->AssignedToClient.ToString()
+				);
+			InProgressOperation.Reset();
+		}
+
+		BroadcastOwnershipChanged();
+	}
+
+	void FReassignObjectPropertiesLogic::OnAssignedFromClientsCompleted(FParallelExecutionResult ParallelResult)
+	{
+		InProgressOperation->ParallelIntermediateResult = MoveTemp(ParallelResult);
+		GuardedExecuteOnGameThread([this]()
+		{
+			FSubmissionQueue* SubmissionQueue = FindSubmissionQueueForTargetClient();
+			if (ensureMsgf(SubmissionQueue, TEXT("Target client (likely) disconnected during reassignment but FReassignObjectPropertiesLogic should have handled canelling the reassignment.")))
+			{
+				// It would be better to check ParallelResult for failed sub-actions and possibly take them out of the request we're about to send to the reassigned to client.
+				// For now, we'll just send and worst case the change is rejected by the server.
+				SubmissionQueue->SubmitNowOrEnqueue_GameThread(*this);
+			}
+		});
+	}
+
+	FSubmissionQueue* FReassignObjectPropertiesLogic::FindSubmissionQueueForTargetClient() const
+	{
+		FReplicationClient* Client = InProgressOperation.IsSet()
+			? ClientManager.FindClient(InProgressOperation->AssignedToClient)
+			: nullptr;
+		return Client ? &Client->GetSubmissionQueue() : nullptr;
+	}
+
+	void FReassignObjectPropertiesLogic::PerformSubmission_GameThread(ISubmissionWorkflow& Workflow)
+	{
+		// Was the pending change cancelled, e.g. due to the target client disconnecting?
+		if (!InProgressOperation)
+		{
+			return;
+		}
+		
+		const FReplicationClient* ClientToAssignTo = ClientManager.FindClient(InProgressOperation->AssignedToClient);
+		if (!ensureMsgf(ClientToAssignTo, TEXT("FReassignObjectPropertiesLogic should have cancelled the change upon disconnect"))
+			|| !ClientToAssignTo->AllowsEditing())
+		{
+			UE_LOG(LogConcert, Error, TEXT("Property Reassignment: The target client is no longer editable."));
+			return;
+		}
+		
+		// TODO UE-201340: If there were errors, add a dedicated notification
+		const FParallelExecutionResult& ParallelResult = *InProgressOperation->ParallelIntermediateResult;
+		const bool bAtLeastOneSucceeded = Algo::AnyOf(ParallelResult.StreamResponses, [](const TPair<FGuid, FSubmitStreamChangesResponse>& Response)
+		{
+			return Response.Value.SubmissionInfo && Response.Value.SubmissionInfo->Response.IsSuccess();
+		});
+		if (!bAtLeastOneSucceeded)
+		{
+			UE_LOG(LogConcert, Warning, TEXT("Property Reassignment: No reassignment will take place because all remote clients failed."));
+			InProgressOperation.Reset();
+			return;
+		}
+		
+		const IClientStreamSynchronizer& StreamSynchronizer = ClientToAssignTo->GetStreamSynchronizer();
+		const FGuid& StreamId = StreamSynchronizer.GetStreamId();
+		const bool bIsStreamRegistered = !StreamSynchronizer.GetServerState().IsEmpty();
+		FStreamChangelist Changelist = ReassignObjectProperties::Private::MakeChangelistFromAppliedChanges(
+			InProgressOperation->OldRegisteredObjects,
+			ParallelResult,
+			StreamSynchronizer.GetServerState(),
+			StreamId
+			);
+		ConcertSyncClient::Replication::FChangeStreamRequest AssignedToClientRequest = bIsStreamRegistered
+			? StreamRequestUtils::BuildChangeRequest_UpdateExistingStream(Changelist)
+			: StreamRequestUtils::BuildChangeRequest_CreateNewStream(StreamId, Changelist);
+		
+		const TSharedPtr<ISubmissionOperation> SubmitOperation = Workflow.SubmitChanges({
+			MoveTemp(AssignedToClientRequest),
+			ReassignObjectProperties::Private::MakeAuthorityRequestFrom(InProgressOperation->OldAuthority, ParallelResult, StreamId)
+		});
+		if (!ensure(SubmitOperation))
+		{
+			UE_LOG(LogConcert, Error, TEXT("Property Reassignment: Failed to give target client properties."));
+			return;
+		}
+		
+		SubmitOperation->OnCompletedOperation().Next([this](ESubmissionOperationCompletedCode)
+		{
+			// We're not on game thread if OnCompletedOperation completes due to a timeout.
+			// By contract, IParallelSubmissionOperation must be destroyed on the game thread!
+			GuardedExecuteOnGameThread([this]() { InProgressOperation.Reset(); }); 
+		});
+	}
+	
+	void FReassignObjectPropertiesLogic::GuardedExecuteOnGameThread(TUniqueFunction<void()> Callback) const
+	{
+		if (IsInGameThread())
+		{
+			Callback();
+		}
+		else
+		{
+			AsyncTask(ENamedThreads::GameThread, [Callback = MoveTemp(Callback), WeakToken = DestructionToken.ToWeakPtr()]()
+			{
+				if (WeakToken.IsValid())
+				{
+					Callback();
+				}
+			});
+		}
+	}
+}
+
+#undef LOCTEXT_NAMESPACE
diff --git a/Engine/Plugins/Developer/Concert/ConcertApp/MultiUserClient/Source/MultiUserClient/Private/Replication/Submission/MultiEdit/ReassignObjectPropertiesLogic.h b/Engine/Plugins/Developer/Concert/ConcertApp/MultiUserClient/Source/MultiUserClient/Private/Replication/Submission/MultiEdit/ReassignObjectPropertiesLogic.h
new file mode 100644
--- /dev/null
+++ b/Engine/Plugins/Developer/Concert/ConcertApp/MultiUserClient/Source/MultiUserClient/Private/Replication/Submission/MultiEdit/ReassignObjectPropertiesLogic.h
@@ -0,0 +1,122 @@
+// Copyright Epic Games, Inc. All Rights Reserved.
+
+#pragma once
+
+#include "Replication/IToken.h"
+#include "Replication/Submission/Queue/DeferredSubmitter.h"
+#include "Replication/Util/StreamRequestUtils.h"
+#include "Replication/Util/SynchronizedRequestUtils.h"
+
+#include "Misc/OptionalFwd.h"
+#include "Templates/Function.h"
+#include "Templates/UnrealTemplate.h"
+
+enum class EBreakBehavior : uint8;
+struct FConcertPropertyChain;
+
+namespace UE::MultiUserClient
+{
+	class FReplicationClientManager;
+	class IParallelSubmissionOperation;
+	
+	/**
+	 * Knows how to transfer object to properties mappings from clients to other clients.
+	 * Acts as a model for the UI.
+	 *
+	 * Currently transfering is implemented using a cooperative peer-to-peer environment. This makes the operation difficult to sync and impossible to make atomic:
+	 * that is if one client fails, we cannot rollback the changes already done on other clients.
+	 * TODO UE-201136: In the future UE-201136 is supposed to address this problem by introducing a dedicated server operation.
+	 * 
+	 * @see SMultiClientView, MultiStreamColumns::ReassignOwnership
+	 */
+	class FReassignObjectPropertiesLogic
+		: public FNoncopyable
+		, IDeferredSubmitter
+	{
+	public:
+
+		enum class EOwnershipState
+		{
+			HasObjectRegistered,
+			NoOwnership
+		};
+		using FProcessClientOwnership = TFunctionRef<EBreakBehavior(const FGuid& ClientId, EOwnershipState Ownership)>;
+		
+		FReassignObjectPropertiesLogic(FReplicationClientManager& InClientManager);
+		virtual ~FReassignObjectPropertiesLogic() override;
+
+		/** Enumerates every client that has registered properties to this object. */
+		void EnumerateClientOwnershipState(const FSoftObjectPath& ObjectPath, FProcessClientOwnership Callback) const;
+		/** @return Whether ClientId has ownership over at least one */
+		bool OwnsAnyOf(TConstArrayView<FSoftObjectPath> Objects, const FGuid& TargetClientId) const;
+		/** @return Whether any of the following objects has at least one owner */
+		bool IsAnyObjectOwned(TConstArrayView<FSoftObjectPath> Objects) const;
+		
+		/** @return Whether content can be assigned to this client (not possible if client does not allow remote changing of streams). */
+		bool CanReassignAnyTo(TConstArrayView<FSoftObjectPath> ObjectsToReassign, const FGuid& ClientId, FText* Reason = nullptr) const;
+		/** Reassigns all properties mapped to the objects from other clients to ClientId.*/
+		void ReassignAllTo(TArray<FSoftObjectPath> ObjectsToReassign, const FGuid& ClientId);
+
+		/** @return Whether ObjectPath is currently being reassigned by this logic. */
+		bool IsReassigning(const FSoftObjectPath& ObjectPath) const;
+		/** @return The local time at which reassignment operation was started. Useful for UI showing throbber but only if operation takes long.*/
+		TOptional<FDateTime> GetTimeReassignmentWasStarted() const;
+
+		/** Broadcast when the result of EnumerateClientOwnershipState may have changed. */
+		FSimpleMulticastDelegate& OnOwnershipChanged() { return OnOwnershipChangedDelegate; }
+
+	private:
+
+		/** Detects when "this" is destroyed to avoid dereferencing it in latent actions. */
+		const TSharedRef<FToken> DestructionToken = FToken::Make();
+
+		/** Used to get info about clients and register to client list changing. */
+		FReplicationClientManager& ClientManager;
+
+		struct FOperationData
+		{
+			const TArray<FSoftObjectPath> ObjectsBeingReassigned;
+			const FGuid AssignedToClient;
+			/** The objects we tried to remove from all the other clients. If the removal request worked, the (previously) assigned properties will be transferred to the target client. */
+			const TMap<FGuid, FObjectReplicationMap> OldRegisteredObjects;
+			/** Clients to the objects they had authority over. Used to later to determine which objects to give the target client authority over (do not transfer authority from failed clients). */
+			const TMap<FGuid, TArray<FSoftObjectPath>> OldAuthority;
+			
+			/**
+			 * The first step of reassignment is to remove all object bindings from all other clients.
+			 * 
+			 * The operation does RAII for cleaning up latent resources, such as dequeuing submission tasks, and handling underlying Concert messages
+			 * futures completing at some later time.
+			 *
+			 * @note This must be destroyed on the game thread (see ExecuteParallelStreamChanges)
+			 */
+			TSharedRef<IParallelSubmissionOperation> RemoveObjectsFromClientsStep;
+			/** Result from completing RemoveObjectsFromClientsStep once it is completed. Unset until RemoveObjectsFromClientsStep completes. */
+			TOptional<FParallelExecutionResult> ParallelIntermediateResult;
+
+			/** Only exists for GetTimeReassignmentWasStarted. */
+			const FDateTime StartTime = FDateTime::Now(); 
+		};
+		/** Set whilst an operation is in progress. */
+		TOptional<FOperationData> InProgressOperation;
+		
+		/** Broadcast when the result of EnumerateClientOwnershipState may have changed. */
+		FSimpleMulticastDelegate OnOwnershipChangedDelegate;
+
+		void OnRemoteClientsChanged();
+		void OnClientCacheChanged(const FGuid&) const { BroadcastOwnershipChanged(); }
+		void BroadcastOwnershipChanged() const { OnOwnershipChangedDelegate.Broadcast(); }
+
+		/** Called once all clients being reassigned have remove their object bindings. The target client will now receive all the previous bindings. */
+		void OnAssignedFromClientsCompleted(FParallelExecutionResult ParallelResult);
+		/** @return Finds the submission queue for the client being reassigned to. */
+		FSubmissionQueue* FindSubmissionQueueForTargetClient() const;
+
+		//~ Begin IDeferredSubmitter Interface
+		virtual void PerformSubmission_GameThread(ISubmissionWorkflow& Workflow) override;
+		//~ End IDeferredSubmitter Interface
+
+		/** Executes an action straight away if on game thread or schedules for it to be executed. Will ensure that "this" is not destroyed if scheduled latently. */
+		void GuardedExecuteOnGameThread(TUniqueFunction<void()> Callback) const;
+	};
+}
diff --git a/Engine/Plugins/Developer/Concert/ConcertApp/MultiUserClient/Source/MultiUserClient/Private/Replication/Submission/Queue/DeferredSubmitter.h b/Engine/Plugins/Developer/Concert/ConcertApp/MultiUserClient/Source/MultiUserClient/Private/Replication/Submission/Queue/DeferredSubmitter.h
--- a/Engine/Plugins/Developer/Concert/ConcertApp/MultiUserClient/Source/MultiUserClient/Private/Replication/Submission/Queue/DeferredSubmitter.h
+++ b/Engine/Plugins/Developer/Concert/ConcertApp/MultiUserClient/Source/MultiUserClient/Private/Replication/Submission/Queue/DeferredSubmitter.h
@@ -31,6 +31,7 @@ namespace UE::MultiUserClient
 	/** Util for systems to auto-unregister when they are destroyed. */
 	class FSelfUnregisteringDeferredSubmitter : public IDeferredSubmitter
 	{
+	protected:
 		FSubmissionQueue& SubmissionQueue;
 	public:
 
diff --git a/Engine/Plugins/Developer/Concert/ConcertApp/MultiUserClient/Source/MultiUserClient/Private/Replication/Submission/Queue/SubmissionQueue.cpp b/Engine/Plugins/Developer/Concert/ConcertApp/MultiUserClient/Source/MultiUserClient/Private/Replication/Submission/Queue/SubmissionQueue.cpp
--- a/Engine/Plugins/Developer/Concert/ConcertApp/MultiUserClient/Source/MultiUserClient/Private/Replication/Submission/Queue/SubmissionQueue.cpp
+++ b/Engine/Plugins/Developer/Concert/ConcertApp/MultiUserClient/Source/MultiUserClient/Private/Replication/Submission/Queue/SubmissionQueue.cpp
@@ -22,7 +22,10 @@ namespace UE::MultiUserClient
 
 	void FSubmissionQueue::SubmitNowOrEnqueue_GameThread(IDeferredSubmitter& DeferredSubmission)
 	{
-		check(IsInGameThread());
+		if (!ensure(IsInGameThread()))
+		{
+			return;
+		}
 
 		if (Workflow.CanSubmit())
 		{
@@ -36,7 +39,15 @@ namespace UE::MultiUserClient
 
 	void FSubmissionQueue::Dequeue_GameThread(IDeferredSubmitter& DeferredSubmission)
 	{
-		check(IsInGameThread());
+		if (!IsInGameThread())
+		{
+			// It is imaginable that some async system (like a request) pushes a task to the game thread which is supposed to call Dequeue_GameThread.
+			// Since the engine is shutting down, never the task would never actually be executed.
+			// However, the underlying system is still destroyed, even if living on another thread, since the engine is shutting down.
+			ensureMsgf(IsEngineExitRequested(), TEXT("Called outside of game thread but not a problem if engine is shutting down. Check your destruction logic."));
+			return;
+		}
+		
 		Queue.RemoveNode(&DeferredSubmission);
 	}
 
diff --git a/Engine/Plugins/Developer/Concert/ConcertApp/MultiUserClient/Source/MultiUserClient/Private/Replication/Submission/Remote/SubmissionWorkflow_RemoteClient.cpp b/Engine/Plugins/Developer/Concert/ConcertApp/MultiUserClient/Source/MultiUserClient/Private/Replication/Submission/Remote/SubmissionWorkflow_RemoteClient.cpp
--- a/Engine/Plugins/Developer/Concert/ConcertApp/MultiUserClient/Source/MultiUserClient/Private/Replication/Submission/Remote/SubmissionWorkflow_RemoteClient.cpp
+++ b/Engine/Plugins/Developer/Concert/ConcertApp/MultiUserClient/Source/MultiUserClient/Private/Replication/Submission/Remote/SubmissionWorkflow_RemoteClient.cpp
@@ -138,11 +138,14 @@ namespace UE::MultiUserClient
 
 	void FSubmissionWorkflow_RemoteClient::CleanUpSubmissionProcess()
 	{
-		const TSharedRef<FSingleClientSubmissionOperation> ExposedOperation = InProgressOperation->ExposedOperation;
-		InProgressOperation.Reset();
+		if (InProgressOperation)
+		{
+			const TSharedPtr<FSingleClientSubmissionOperation> ExposedOperation = InProgressOperation->ExposedOperation;
+			InProgressOperation.Reset();
 		
-		ExposedOperation->EmplaceCompleteOperationPromise(ESubmissionOperationCompletedCode::Processed);
-		OnSubmitOperationCompletedDelegate.Broadcast();
+			ExposedOperation->EmplaceCompleteOperationPromise(ESubmissionOperationCompletedCode::Processed);
+			OnSubmitOperationCompletedDelegate.Broadcast();
+		}
 	}
 
 	void FSubmissionWorkflow_RemoteClient::OnStreamRemoteChangeEvent(const FConcertSessionContext& Context, const FMultiUser_ChangeRemote_StreamUpdatedEvent& EventData)
diff --git a/Engine/Plugins/Developer/Concert/ConcertApp/MultiUserClient/Source/MultiUserClient/Private/Replication/Submission/SingleClientSubmissionOperation.h b/Engine/Plugins/Developer/Concert/ConcertApp/MultiUserClient/Source/MultiUserClient/Private/Replication/Submission/SingleClientSubmissionOperation.h
--- a/Engine/Plugins/Developer/Concert/ConcertApp/MultiUserClient/Source/MultiUserClient/Private/Replication/Submission/SingleClientSubmissionOperation.h
+++ b/Engine/Plugins/Developer/Concert/ConcertApp/MultiUserClient/Source/MultiUserClient/Private/Replication/Submission/SingleClientSubmissionOperation.h
@@ -57,7 +57,7 @@ namespace UE::MultiUserClient
 		}
 		void EmplaceCompleteOperationPromise(ESubmissionOperationCompletedCode Result)
 		{
-			bAuthorityResponsePromiseWasSet = true;
+			bCompleteOperationPromiseWasSet = true;
 			CompleteOperationPromise.EmplaceValue(Result);
 		}
 
diff --git a/Engine/Plugins/Developer/Concert/ConcertApp/MultiUserClient/Source/MultiUserClient/Private/Replication/Util/GlobalAuthorityCache.cpp b/Engine/Plugins/Developer/Concert/ConcertApp/MultiUserClient/Source/MultiUserClient/Private/Replication/Util/GlobalAuthorityCache.cpp
--- a/Engine/Plugins/Developer/Concert/ConcertApp/MultiUserClient/Source/MultiUserClient/Private/Replication/Util/GlobalAuthorityCache.cpp
+++ b/Engine/Plugins/Developer/Concert/ConcertApp/MultiUserClient/Source/MultiUserClient/Private/Replication/Util/GlobalAuthorityCache.cpp
@@ -66,6 +66,17 @@ namespace UE::MultiUserClient
 		return Result;
 	}
 
+	bool FGlobalAuthorityCache::HasAuthorityOverObject(const FSoftObjectPath& Object, const FGuid& ClientId) const
+	{
+		bool bHasAuthority = false;
+		ForEachClientWithAuthorityOverObject(Object, [&ClientId, &bHasAuthority](const FGuid& OtherClientId)
+		{
+			bHasAuthority = ClientId == OtherClientId;
+			return bHasAuthority ? EBreakBehavior::Break : EBreakBehavior::Continue;
+		});
+		return bHasAuthority;
+	}
+
 	EAuthorityMutability FGlobalAuthorityCache::CanClientTakeAuthorityAfterSubmission(const FSoftObjectPath& Object, const FGuid& ClientId, FProcessPropertyConflict ProcessConflict) const
 	{
 		const FReplicationClient* Client = ClientManager.FindClient(ClientId);
diff --git a/Engine/Plugins/Developer/Concert/ConcertApp/MultiUserClient/Source/MultiUserClient/Private/Replication/Util/GlobalAuthorityCache.h b/Engine/Plugins/Developer/Concert/ConcertApp/MultiUserClient/Source/MultiUserClient/Private/Replication/Util/GlobalAuthorityCache.h
--- a/Engine/Plugins/Developer/Concert/ConcertApp/MultiUserClient/Source/MultiUserClient/Private/Replication/Util/GlobalAuthorityCache.h
+++ b/Engine/Plugins/Developer/Concert/ConcertApp/MultiUserClient/Source/MultiUserClient/Private/Replication/Util/GlobalAuthorityCache.h
@@ -26,7 +26,7 @@ namespace UE::MultiUserClient
 	 */
 	class FGlobalAuthorityCache
 		: public FNoncopyable
-		, private ConcertSyncCore::Replication::AuthorityConflictUtils::IReplicationGroundTruth
+		, ConcertSyncCore::Replication::AuthorityConflictUtils::IReplicationGroundTruth
 	{
 	public:
 
@@ -43,6 +43,8 @@ namespace UE::MultiUserClient
 		void ForEachClientWithAuthorityOverObject(const FSoftObjectPath& Object, TFunctionRef<EBreakBehavior(const FGuid& ClientId)> Callback) const;
 		/** Util that uses ForEachClientWithAuthorityOverObject to make an array. */
 		TArray<FGuid> GetClientsWithAuthorityOverObject(const FSoftObjectPath& Object) const;
+		/** @return Returns whether the client has (partial) authority over the object. */
+		bool HasAuthorityOverObject(const FSoftObjectPath& Object, const FGuid& ClientId) const;
 		
 		/** @return Whether the given client can take authority over the object without causing any conflicts. This also considers the changes made to the stream after submission. */
 		EAuthorityMutability CanClientTakeAuthorityAfterSubmission(const FSoftObjectPath& Object, const FGuid& ClientId, FProcessPropertyConflict ProcessConflict = [](auto&, auto&){ return EBreakBehavior::Break; }) const;
@@ -93,8 +95,16 @@ namespace UE::MultiUserClient
 		void RemoveClient(const FGuid& ClientId);
 
 		// Respond to remote client registration
-		void OnPostRemoteClientAdded(FRemoteReplicationClient& RemoteClient) { RegisterForClientEvents(RemoteClient); }
-		void OnPreRemoteClientRemoved(FRemoteReplicationClient& RemoteClient) const { UnregisterFromClientEvents(RemoteClient); }
+		void OnPostRemoteClientAdded(FRemoteReplicationClient& RemoteClient)
+		{
+			RebuildClient(RemoteClient.GetEndpointId());
+			RegisterForClientEvents(RemoteClient);
+		}
+		void OnPreRemoteClientRemoved(FRemoteReplicationClient& RemoteClient)
+		{
+			RemoveClient(RemoteClient.GetEndpointId());
+			UnregisterFromClientEvents(RemoteClient);
+		}
 
 		// Rebuild client when their authority changes
 		void OnPostAuthorityChanged(const FGuid ClientId) { RebuildClient(ClientId); }
diff --git a/Engine/Plugins/Developer/Concert/ConcertApp/MultiUserClient/Source/MultiUserClient/Private/Replication/Util/StreamRequestUtils.h b/Engine/Plugins/Developer/Concert/ConcertApp/MultiUserClient/Source/MultiUserClient/Private/Replication/Util/StreamRequestUtils.h
--- a/Engine/Plugins/Developer/Concert/ConcertApp/MultiUserClient/Source/MultiUserClient/Private/Replication/Util/StreamRequestUtils.h
+++ b/Engine/Plugins/Developer/Concert/ConcertApp/MultiUserClient/Source/MultiUserClient/Private/Replication/Util/StreamRequestUtils.h
@@ -11,7 +11,12 @@ namespace UE::ConcertSyncClient::Replication
 
 namespace UE::MultiUserClient
 {
-	struct FStreamChangelist;
+	/** Describes changes that MU client makes. */
+	struct FStreamChangelist
+	{
+		TSet<FObjectInStreamID> ObjectsToRemove;
+		TMap<FObjectInStreamID, FConcertReplication_ChangeStream_PutObject> ObjectsToPut;
+	};
 }
 
 namespace UE::MultiUserClient::StreamRequestUtils
diff --git a/Engine/Plugins/Developer/Concert/ConcertApp/MultiUserClient/Source/MultiUserClient/Private/Replication/Util/SynchronizedRequestUtils.cpp b/Engine/Plugins/Developer/Concert/ConcertApp/MultiUserClient/Source/MultiUserClient/Private/Replication/Util/SynchronizedRequestUtils.cpp
new file mode 100644
--- /dev/null
+++ b/Engine/Plugins/Developer/Concert/ConcertApp/MultiUserClient/Source/MultiUserClient/Private/Replication/Util/SynchronizedRequestUtils.cpp
@@ -0,0 +1,287 @@
+// Copyright Epic Games, Inc. All Rights Reserved.
+
+#include "SynchronizedRequestUtils.h"
+
+#include "ConcertLogGlobal.h"
+#include "Replication/Client/ReplicationClientManager.h"
+#include "Replication/Submission/Data/AuthoritySubmission.h"
+#include "Replication/Submission/Data/StreamSubmission.h"
+#include "Replication/Submission/ISubmissionOperation.h"
+
+#include "Algo/AllOf.h"
+#include "Containers/UnrealString.h"
+
+namespace UE::MultiUserClient
+{
+	namespace ParallelSubmission::Private
+	{
+		static void RemoveInvalidClientEntries(const FReplicationClientManager& ClientManager, TMap<FGuid, FSubmissionParams> ParallelOperations)
+		{
+			for (auto It = ParallelOperations.CreateIterator(); It; ++It)
+			{
+				const TPair<FGuid, FSubmissionParams>& Operation =*It;
+				if (!ensure(ClientManager.FindClient(Operation.Key)))
+				{
+					It.RemoveCurrent();
+				}
+			}
+		}
+		
+		static void LogStreamChangeError(const FGuid& ClientId, const FSubmitStreamChangesResponse& Response)
+		{
+			if (Response.ErrorCode == EStreamSubmissionErrorCode::Timeout)
+			{
+				UE_LOG(LogConcert, Warning, TEXT("Remote stream change to client %s timed out."), *ClientId.ToString());
+			}
+			else if (Response.SubmissionInfo && Response.SubmissionInfo->Response.IsFailure())
+			{
+				FStringOutputDevice OutputDevice;
+				Response.SubmissionInfo->Response.LogErrors(OutputDevice);
+				UE_LOG(LogConcert, Warning, TEXT("Remote stream change to client %s failed: %s"), *ClientId.ToString(), *OutputDevice);
+			}
+		}
+
+		class FSyncOperation;
+
+		/** Adapter for FSubmissionQueue. Handles submitting for one client and notifying FSyncOperation if completions. */
+		class FDeferredSubmitter
+			: public FNoncopyable
+			, public IDeferredSubmitter
+		{
+			FSyncOperation& Owner;
+			FReplicationClientManager& ClientManager;
+			FSubmissionQueue& SubmissionQueue;
+			const FGuid ClientId;
+			const FSubmissionParams SubmissionParams;
+			
+		public:
+
+			FDeferredSubmitter(
+				FSyncOperation& InOwner,
+				FReplicationClientManager& InClientManager,
+				FSubmissionQueue& InSubmissionQueue,
+				const FGuid& InClientId,
+				FSubmissionParams&& InSubmissionParams
+				)
+				: Owner(InOwner)
+				, ClientManager(InClientManager)
+				, SubmissionQueue(InSubmissionQueue)
+				, ClientId(InClientId)
+				, SubmissionParams(MoveTemp(InSubmissionParams))
+			{}
+
+			virtual ~FDeferredSubmitter() override
+			{
+				checkf(IsInGameThread(), TEXT("You must destroy IParallelSubmissionOperation on the game thread"));
+				if (FReplicationClient* Client = ClientManager.FindClient(ClientId))
+				{
+					Client->GetSubmissionQueue().Dequeue_GameThread(*this);
+				}
+			}
+
+			void EnqueueOrSubmit()
+			{
+				SubmissionQueue.SubmitNowOrEnqueue_GameThread(*this);
+			}
+			
+			const FGuid& GetClientId() const { return ClientId; }
+
+			virtual void PerformSubmission_GameThread(ISubmissionWorkflow& Workflow) override;
+		};
+
+		/** Starts all submissions and synchronizes their concurrent completion. */
+		class FSyncOperation : public IParallelSubmissionOperation, public TSharedFromThis<FSyncOperation>
+		{
+		public:
+			
+			FSyncOperation(FReplicationClientManager& ClientManager, TMap<FGuid, FSubmissionParams>&& Operations)
+			{
+				const int32 NumOperations = Operations.Num();
+				IntermediateResults.StreamResponses.Reserve(NumOperations);
+				IntermediateResults.AuthorityResponses.Reserve(NumOperations);
+				SubOperations.Reserve(NumOperations);
+
+				// Prepare the intermediate slots in advance to handle concurrent writing
+				for (TPair<FGuid, FSubmissionParams>& Pair : Operations)
+				{
+					// Allocate slots
+					const FGuid& ClientId = Pair.Key;
+					IntermediateResults.StreamResponses.Add(ClientId);
+					IntermediateResults.AuthorityResponses.Add(ClientId);
+
+					FReplicationClient* Client = ClientManager.FindClient(ClientId);
+					checkf(Client, TEXT("Was supposed to have been validated ExecuteParallelStreamChanges."))
+					SubOperations.Emplace(*this, ClientManager, ClientId, MoveTemp(Pair.Value), Client->GetSubmissionQueue());
+				}
+			}
+
+			virtual ~FSyncOperation() override
+			{
+				if (!bPromiseWasSet)
+				{
+					bPromiseWasSet = true;
+					constexpr bool bWasCancelled = true;
+					Promise.EmplaceValue(FParallelExecutionResult{ bWasCancelled });
+				}
+			}
+
+			void StartSubmissions()
+			{
+				// FDeferredSubmitter needs access to SharedThis(), so it must be called outside of constructor.
+				for (FPendingOperation& PendingOperation : SubOperations)
+				{
+					FDeferredSubmitter& Submitter = PendingOperation.Submitter;
+					Submitter.EnqueueOrSubmit();
+				}
+			}
+
+			//~ Begin IParallelSubmissionOperation Interface
+			virtual TFuture<FParallelExecutionResult> GetFuture() override { return Promise.GetFuture(); }
+			//~ End IParallelSubmissionOperation Interface
+
+			void OnCompleteStream(const FGuid& ClientId, FSubmitStreamChangesResponse&& Response)
+			{
+				LogStreamChangeError(ClientId, Response);
+				
+				--FindOperationByClient(ClientId)->NumOperationsLeft;
+				IntermediateResults.StreamResponses[ClientId] = MoveTemp(Response);
+
+				FinishIfDone();
+			}
+
+			void OnCompleteAuthority(const FGuid& ClientId, FSubmitAuthorityChangesResponse&& Response)
+			{
+				--FindOperationByClient(ClientId)->NumOperationsLeft;
+				IntermediateResults.AuthorityResponses[ClientId] = MoveTemp(Response);
+				
+				const bool bIsFailure = Response.ErrorCode == EAuthoritySubmissionResponseErrorCode::Timeout || (Response.Response && !Response.Response->RejectedObjects.IsEmpty());
+				UE_CLOG(bIsFailure, LogConcert, Warning, TEXT("Remote authority change to client %s failed"), *ClientId.ToString());
+				FinishIfDone();
+			}
+
+			void OnSubmissionFailure(const FGuid& ClientId)
+			{
+				FindOperationByClient(ClientId)->NumOperationsLeft = 0;
+				FinishIfDone();
+			}
+
+		private:
+			
+			TPromise<FParallelExecutionResult> Promise;
+			bool bPromiseWasSet = false;
+
+			struct FPendingOperation
+			{
+				/** Handles the submission. Unregisters from the client queue (if we're destroyed early). */
+				FDeferredSubmitter Submitter;
+			
+				/** Used to determine when the operation is done. */
+				int32 NumOperationsLeft = 2; // Stream + Authority = 2
+				
+				FPendingOperation(
+					FSyncOperation& InOwner,
+					FReplicationClientManager& InClientManager,
+					const FGuid& InClientId,
+					FSubmissionParams&& InSubmissionParams,
+					FSubmissionQueue& InSubmissionQueue
+					)
+					: Submitter(InOwner, InClientManager, InSubmissionQueue, InClientId, MoveTemp(InSubmissionParams))
+				{}
+
+				bool IsDone() const { return NumOperationsLeft == 0; }
+			};
+			/**
+			 * Each latent tasks has its own, pre-allocated slot which effectively handles potentially concurrent writes.
+			 *
+			 * Initially every task value is initialized to 2: one for each operation (stream, authority).
+			 * Every time a sub-operation finishes, the value decrements.
+			 * Once all values have reached 0, that means Operations can be Emplace()'ed with IntermediateResults.
+			 */
+			TArray<FPendingOperation> SubOperations;
+
+			/** Holds intermediate results as they come in. */
+			FParallelExecutionResult IntermediateResults;
+
+			void FinishIfDone()
+			{
+				const bool bAllDone = Algo::AllOf(SubOperations, [](const FPendingOperation& OperationInfo)
+				{
+					return OperationInfo.IsDone();
+				});
+
+				if (bAllDone && ensure(!bPromiseWasSet))
+				{
+					bPromiseWasSet = true;
+					Promise.EmplaceValue(MoveTemp(IntermediateResults));
+				}
+			}
+
+			FPendingOperation* FindOperationByClient(const FGuid& ClientId)
+			{
+				return SubOperations.FindByPredicate([&ClientId](const FPendingOperation& Operation)
+				{
+					return Operation.Submitter.GetClientId() == ClientId;
+				});
+			}
+		};
+
+		void FDeferredSubmitter::PerformSubmission_GameThread(ISubmissionWorkflow& Workflow)
+		{
+			const TSharedPtr<ISubmissionOperation> SubmissionOperation = Workflow.SubmitChanges(SubmissionParams);
+			if (ensure(SubmissionOperation))
+			{
+				SubmissionOperation->OnCompleteStreamChangesFuture()
+					.Next([this, WeakOwner = Owner.AsWeak()](FSubmitStreamChangesResponse&& Response)
+					{
+						if (const TSharedPtr<FSyncOperation> SyncOperation = WeakOwner.Pin())
+						{
+							SyncOperation->OnCompleteStream(ClientId, MoveTemp(Response));
+						}
+					});
+				SubmissionOperation->OnCompleteAuthorityChangeFuture()
+					.Next([this, WeakOwner = Owner.AsWeak()](FSubmitAuthorityChangesResponse&& Response)
+					{
+						if (const TSharedPtr<FSyncOperation> SyncOperation = WeakOwner.Pin())
+						{
+							SyncOperation->OnCompleteAuthority(ClientId, MoveTemp(Response));
+						}
+					});
+			}
+			else
+			{
+				Owner.OnSubmissionFailure(ClientId);
+			}
+		}
+	}
+	
+	TSharedPtr<IParallelSubmissionOperation> ExecuteParallelStreamChanges(FReplicationClientManager& ClientManager, TMap<FGuid, FSubmissionParams> ParallelOperations)
+	{
+		ParallelSubmission::Private::RemoveInvalidClientEntries(ClientManager, ParallelOperations);
+		if (!ensure(!ParallelOperations.IsEmpty()))
+		{
+			return nullptr;
+		}
+
+		using namespace ParallelSubmission::Private;
+		TSharedRef<FSyncOperation> Operation = MakeShared<FSyncOperation>(ClientManager, MoveTemp(ParallelOperations));
+		Operation->StartSubmissions();
+		return Operation;
+	}
+
+	TSharedPtr<IParallelSubmissionOperation> ExecuteParallelStreamChanges(FReplicationClientManager& ClientManager, TMap<FGuid, ConcertSyncClient::Replication::FChangeStreamRequest> ParallelOperations)
+	{
+		if (!ensure(!ParallelOperations.IsEmpty()))
+		{
+			return nullptr;
+		}
+		
+		TMap<FGuid, FSubmissionParams> Transformed;
+		Transformed.Reserve(ParallelOperations.Num()); // check()s if Num() == 0
+		for (TPair<FGuid, ConcertSyncClient::Replication::FChangeStreamRequest>& RequestPair : ParallelOperations)
+		{
+			Transformed.Emplace(RequestPair.Key, FSubmissionParams{ MoveTemp(RequestPair.Value) });
+		}
+		
+		return ExecuteParallelStreamChanges(ClientManager, MoveTemp(Transformed));
+	}
+}
diff --git a/Engine/Plugins/Developer/Concert/ConcertApp/MultiUserClient/Source/MultiUserClient/Private/Replication/Util/SynchronizedRequestUtils.h b/Engine/Plugins/Developer/Concert/ConcertApp/MultiUserClient/Source/MultiUserClient/Private/Replication/Util/SynchronizedRequestUtils.h
new file mode 100644
--- /dev/null
+++ b/Engine/Plugins/Developer/Concert/ConcertApp/MultiUserClient/Source/MultiUserClient/Private/Replication/Util/SynchronizedRequestUtils.h
@@ -0,0 +1,60 @@
+// Copyright Epic Games, Inc. All Rights Reserved.
+
+#pragma once
+
+#include "Replication/Submission/Data/AuthoritySubmission.h"
+#include "Replication/Submission/Data/StreamSubmission.h"
+
+#include "Async/Future.h"
+#include "Templates/UnrealTemplate.h"
+
+struct FConcertReplication_ChangeStream_Request;
+
+namespace UE::ConcertSyncClient::Replication
+{
+	struct FChangeStreamRequest;
+}
+
+namespace UE::MultiUserClient
+{
+	class FReplicationClientManager;
+	struct FSubmissionParams;
+
+	struct FParallelExecutionResult
+	{
+		/** True if the TSharedPtr<IParallelSubmissionOperation> was destroyed before completion of submissions. */
+		bool bWasCancelled = false;
+
+		/** Maps client IDs to responses from the submissions. Empty if bWasCancelled. */
+		TMap<FGuid, FSubmitStreamChangesResponse> StreamResponses;
+		/** Maps client IDs to responses from the submissions. Empty if bWasCancelled. */
+		TMap<FGuid, FSubmitAuthorityChangesResponse> AuthorityResponses;
+	};
+
+	/** Result of ExecuteParallelStreamChanges. Handles correctly unregistering when the object is destroyed. */
+	class IParallelSubmissionOperation
+	{
+	public:
+
+		/** @return Future completed when all parallel submissions have been completed. */
+		virtual TFuture<FParallelExecutionResult> GetFuture() = 0;
+
+		virtual ~IParallelSubmissionOperation() = default;
+	};
+
+	// TODO UE-201136: This code can be removed once a dedicated server operation for changing client streams has been implemented.
+	
+	/**
+	 * Executes the stream changes in parallel on the remote clients.
+	 *
+	 * @param ClientManager Used to get the clients to send to
+	 * @param ParallelOperations The operations to complete
+	 * 
+	 * @return Object you can get future through. This object also handles correct unregistering from other systems when you destroy it.
+	 * @note The returned object MUST be destroyed on the game thread.
+	 */
+	TSharedPtr<IParallelSubmissionOperation> ExecuteParallelStreamChanges(FReplicationClientManager& ClientManager, TMap<FGuid, FSubmissionParams> ParallelOperations);
+	/** Util for transforming into FSubmissionParams */
+	TSharedPtr<IParallelSubmissionOperation> ExecuteParallelStreamChanges(FReplicationClientManager& ClientManager, TMap<FGuid, ConcertSyncClient::Replication::FChangeStreamRequest> ParallelOperations);
+}
+
diff --git a/Engine/Plugins/Developer/Concert/ConcertApp/MultiUserClient/Source/MultiUserClient/Private/Widgets/ActiveSession/Replication/Client/ClientUtils.cpp b/Engine/Plugins/Developer/Concert/ConcertApp/MultiUserClient/Source/MultiUserClient/Private/Widgets/ActiveSession/Replication/Client/ClientUtils.cpp
--- a/Engine/Plugins/Developer/Concert/ConcertApp/MultiUserClient/Source/MultiUserClient/Private/Widgets/ActiveSession/Replication/Client/ClientUtils.cpp
+++ b/Engine/Plugins/Developer/Concert/ConcertApp/MultiUserClient/Source/MultiUserClient/Private/Widgets/ActiveSession/Replication/Client/ClientUtils.cpp
@@ -36,15 +36,20 @@ namespace UE::MultiUserClient::ClientUtils
 	bool GetClientDisplayInfo(const IConcertClient& InLocalClientInstance, const FGuid& InClientEndpointId, FConcertClientInfo& OutClientInfo)
 	{
 		const TSharedPtr<IConcertClientSession> Session = InLocalClientInstance.GetCurrentSession();
-		const bool bIsLocalClient = ensure(Session) && Session->GetSessionClientEndpointId() == InClientEndpointId;
+		return GetClientDisplayInfo(*Session, InClientEndpointId, OutClientInfo);
+	}
+	
+	bool GetClientDisplayInfo(const IConcertClientSession& InSession, const FGuid& InClientEndpointId, FConcertClientInfo& OutClientInfo)
+	{
+		const bool bIsLocalClient = InSession.GetSessionClientEndpointId() == InClientEndpointId;
 		if (bIsLocalClient)
 		{
-			OutClientInfo = InLocalClientInstance.GetClientInfo();
+			OutClientInfo = InSession.GetLocalClientInfo();
 			return true;
 		}
 
 		FConcertSessionClientInfo ClientInfo;
-		if (ensure(Session) && Session->FindSessionClient(InClientEndpointId, ClientInfo))
+		if (InSession.FindSessionClient(InClientEndpointId, ClientInfo))
 		{
 			OutClientInfo = MoveTemp(ClientInfo.ClientInfo);
 			return true;
@@ -54,13 +59,18 @@ namespace UE::MultiUserClient::ClientUtils
 	}
 	
 	TArray<const FReplicationClient*> GetSortedClientList(const IConcertClient& InLocalClientInstance, const FReplicationClientManager& InReplicationManager)
+	{
+		return GetSortedClientList(*InLocalClientInstance.GetCurrentSession(), InReplicationManager);
+	}
+
+	TArray<const FReplicationClient*> GetSortedClientList(const IConcertClientSession& InSession, const FReplicationClientManager& InReplicationManager)
 	{
 		TArray<const FReplicationClient*> Result;
 		TMap<const FReplicationClient*, FConcertClientInfo> ClientToDisplayInfo;
-		for (const TNonNullPtr<FRemoteReplicationClient> RemoteClient : InReplicationManager.GetRemoteClients())
+		for (const TNonNullPtr<const FRemoteReplicationClient> RemoteClient : InReplicationManager.GetRemoteClients())
 		{
 			FConcertClientInfo Info;
-			if (GetClientDisplayInfo(InLocalClientInstance, RemoteClient->GetEndpointId(), Info))
+			if (GetClientDisplayInfo(InSession, RemoteClient->GetEndpointId(), Info))
 			{
 				Result.Add(RemoteClient);
 				ClientToDisplayInfo.Add(RemoteClient, MoveTemp(Info));
diff --git a/Engine/Plugins/Developer/Concert/ConcertApp/MultiUserClient/Source/MultiUserClient/Private/Widgets/ActiveSession/Replication/Client/ClientUtils.h b/Engine/Plugins/Developer/Concert/ConcertApp/MultiUserClient/Source/MultiUserClient/Private/Widgets/ActiveSession/Replication/Client/ClientUtils.h
--- a/Engine/Plugins/Developer/Concert/ConcertApp/MultiUserClient/Source/MultiUserClient/Private/Widgets/ActiveSession/Replication/Client/ClientUtils.h
+++ b/Engine/Plugins/Developer/Concert/ConcertApp/MultiUserClient/Source/MultiUserClient/Private/Widgets/ActiveSession/Replication/Client/ClientUtils.h
@@ -38,6 +38,7 @@ namespace UE::MultiUserClient::ClientUtils
 	 * @return Whether OutClientInfo holds a value
 	 */
 	bool GetClientDisplayInfo(const IConcertClient& InLocalClientInstance, const FGuid& InClientEndpointId, FConcertClientInfo& OutClientInfo);
+	bool GetClientDisplayInfo(const IConcertClientSession& InSession, const FGuid& InClientEndpointId, FConcertClientInfo& OutClientInfo);
 
 	/**
 	 * Gets all replication clients in a sorted array. The local client will always be first and then come all remote clients sorted alphabetically.
@@ -46,4 +47,5 @@ namespace UE::MultiUserClient::ClientUtils
 	 * @return Sorted client array
 	 */
 	TArray<const FReplicationClient*> GetSortedClientList(const IConcertClient& InLocalClientInstance, const FReplicationClientManager& InReplicationManager);
+	TArray<const FReplicationClient*> GetSortedClientList(const IConcertClientSession& InSession, const FReplicationClientManager& InReplicationManager);
 }
diff --git a/Engine/Plugins/Developer/Concert/ConcertApp/MultiUserClient/Source/MultiUserClient/Private/Widgets/ActiveSession/Replication/Client/Multi/Columns/MultiStreamColumns.h b/Engine/Plugins/Developer/Concert/ConcertApp/MultiUserClient/Source/MultiUserClient/Private/Widgets/ActiveSession/Replication/Client/Multi/Columns/MultiStreamColumns.h
--- a/Engine/Plugins/Developer/Concert/ConcertApp/MultiUserClient/Source/MultiUserClient/Private/Widgets/ActiveSession/Replication/Client/Multi/Columns/MultiStreamColumns.h
+++ b/Engine/Plugins/Developer/Concert/ConcertApp/MultiUserClient/Source/MultiUserClient/Private/Widgets/ActiveSession/Replication/Client/Multi/Columns/MultiStreamColumns.h
@@ -11,6 +11,7 @@ class IConcertClient;
 
 namespace UE::MultiUserClient
 {
+	class FReassignObjectPropertiesLogic;
 	class FReplicationClientManager;
 }
 
@@ -24,21 +25,23 @@ namespace UE::ConcertClientSharedSlate
 namespace UE::MultiUserClient::MultiStreamColumns
 {
 	const extern FName ReplicationToggleColumnId;
+	const extern FName ReassignOwnershipColumnId;
 	const extern FName AssignPropertyColumnId;
 	
-	/* @see EReplicationPropertyColumnOrder */
+	/* @see ETopLevelColumnOrder and EReplicationPropertyColumnOrder */
 	enum class EColumnSortOrder
 	{
 		ReplicationToggle = 0,
-		AssignPropertyColumn = 30
+		AssignPropertyColumn = 30,
+		ReassignOwnership = 40
 	};
 
 	/**
 	 * Toggles replication for all clients assigned to the object (and optionally all children).
 	 * 
 	 * @param ConcertClient Used to look up client names
-	 * @param ClientManager Used to access all clients for toggling authority
 	 * @param ConsolidatedStreamModelAttribute Used to get child objects
+	 * @param ClientManager Used to access all clients for toggling authority
 	 * @param ColumnsSortPriority The order relative to the other columns
 	 * 
 	 * @return A checkbox for controlling the authority of the object in the row
@@ -49,6 +52,23 @@ namespace UE::MultiUserClient::MultiStreamColumns
 		FReplicationClientManager& ClientManager,
 		const int32 ColumnsSortPriority = static_cast<int32>(EColumnSortOrder::ReplicationToggle)
 		);
+
+	/**
+	 * A combo box which displays all current owners for an object and allows bulk reassigning properties.
+	 *
+	 * @param ConcertClient Used to look up client names
+	 * @param ConsolidatedModelAttribute Used to get child objects
+	 * @param ReassignmentLogic Performs the act of reassigning
+	 * @param ClientManager Used to access all clients for display in the combo box drop-down
+	 * @param ColumnsSortPriority The order relative to the other columns
+	 */
+	ConcertClientSharedSlate::ReplicationColumns::FReplicationTopLevelObjectColumn ReassignOwnership(
+		TSharedRef<IConcertClient> ConcertClient,
+		TAttribute<ConcertClientSharedSlate::IReplicationStreamModel*> ConsolidatedModelAttribute,
+		FReassignObjectPropertiesLogic& ReassignmentLogic,
+		const FReplicationClientManager& ClientManager,
+		int32 ColumnsSortPriority = static_cast<int32>(EColumnSortOrder::ReassignOwnership)
+		);
 	
 	/**
 	 * Creates a property column which assigns the property to the stream selected in the combo box.
diff --git a/Engine/Plugins/Developer/Concert/ConcertApp/MultiUserClient/Source/MultiUserClient/Private/Widgets/ActiveSession/Replication/Client/Multi/Columns/MultiStreamColumns_AssignProperties.cpp b/Engine/Plugins/Developer/Concert/ConcertApp/MultiUserClient/Source/MultiUserClient/Private/Widgets/ActiveSession/Replication/Client/Multi/Columns/MultiStreamColumns_AssignProperties.cpp
--- a/Engine/Plugins/Developer/Concert/ConcertApp/MultiUserClient/Source/MultiUserClient/Private/Widgets/ActiveSession/Replication/Client/Multi/Columns/MultiStreamColumns_AssignProperties.cpp
+++ b/Engine/Plugins/Developer/Concert/ConcertApp/MultiUserClient/Source/MultiUserClient/Private/Widgets/ActiveSession/Replication/Client/Multi/Columns/MultiStreamColumns_AssignProperties.cpp
@@ -3,6 +3,7 @@
 #include "MultiStreamColumns.h"
 
 #include "IConcertClient.h"
+#include "MultiUserReplicationStyle.h"
 #include "Replication/Client/ReplicationClient.h"
 #include "Replication/Client/ReplicationClientManager.h"
 #include "Replication/Editor/Model/IEditableMultiReplicationStreamModel.h"
@@ -48,7 +49,7 @@ namespace UE::MultiUserClient::MultiStreamColumns
 				return &ClientManager.GetLocalClient();
 			}
 
-			for (const TNonNullPtr<FRemoteReplicationClient> Client : ClientManager.GetRemoteClients())
+			for (const TNonNullPtr<const FRemoteReplicationClient> Client : ClientManager.GetRemoteClients())
 			{
 				if (&Client->GetClientEditModel().Get() == &StreamModel)
 				{
@@ -106,7 +107,7 @@ namespace UE::MultiUserClient::MultiStreamColumns
 				.ColumnSortOrder(ColumnsSortPriority),
 			SHeaderRow::Column(AssignPropertyColumnId)
 				.DefaultLabel(LOCTEXT("Owner", "Owner"))
-				.FillSized(200.f)
+				.FillSized(FMultiUserReplicationStyle::Get()->GetFloat(TEXT("AllClients.Property.OwnerColumnWidth")))
 			);
 	}
 }
diff --git a/Engine/Plugins/Developer/Concert/ConcertApp/MultiUserClient/Source/MultiUserClient/Private/Widgets/ActiveSession/Replication/Client/Multi/Columns/MultiStreamColumns_ReassignOwnership.cpp b/Engine/Plugins/Developer/Concert/ConcertApp/MultiUserClient/Source/MultiUserClient/Private/Widgets/ActiveSession/Replication/Client/Multi/Columns/MultiStreamColumns_ReassignOwnership.cpp
new file mode 100644
--- /dev/null
+++ b/Engine/Plugins/Developer/Concert/ConcertApp/MultiUserClient/Source/MultiUserClient/Private/Widgets/ActiveSession/Replication/Client/Multi/Columns/MultiStreamColumns_ReassignOwnership.cpp
@@ -0,0 +1,50 @@
+// Copyright Epic Games, Inc. All Rights Reserved.
+
+#include "IConcertClient.h"
+#include "MultiStreamColumns.h"
+
+#include "MultiUserReplicationStyle.h"
+#include "SReassignObjectComboBox.h"
+
+#define LOCTEXT_NAMESPACE "ReassignOwnershipColumn"
+
+namespace UE::MultiUserClient::MultiStreamColumns
+{
+	const FName ReassignOwnershipColumnId(TEXT("ReassignOwnershipColumn"));
+	
+	ConcertClientSharedSlate::ReplicationColumns::FReplicationTopLevelObjectColumn ReassignOwnership(
+		TSharedRef<IConcertClient> ConcertClient,
+		TAttribute<ConcertClientSharedSlate::IReplicationStreamModel*> ConsolidatedModelAttribute,
+		FReassignObjectPropertiesLogic& ReassignmentLogic,
+		const FReplicationClientManager& ClientManager,
+		const int32 ColumnsSortPriority
+		)
+	{
+		using namespace ConcertClientSharedSlate::ReplicationColumns;
+
+		auto MakeWidget =
+			[ConcertClient, ConsolidatedModelAttribute = MoveTemp(ConsolidatedModelAttribute), &ReassignmentLogic, &ClientManager]
+			(const FReplicationTopLevelObjectColumn::FBuildArgs& InArgs)
+			{
+				return SNew(SReassignObjectComboBox, ConcertClient, ReassignmentLogic, ClientManager)
+					.ManagedObject(InArgs.RowData.GetObjectPath())
+					.ConsolidatedModel(ConsolidatedModelAttribute)
+					.HighlightText(InArgs.HighlightText);
+			};
+		
+		return FReplicationTopLevelObjectColumn(
+			FReplicationTopLevelObjectColumn::FArguments()
+				.GenerateWidgetColumn_Lambda(MoveTemp(MakeWidget))
+				.PopulateSearchItems_Lambda([ConcertClient, &ReassignmentLogic](const ConcertClientSharedSlate::FReplicatedObjectData& ObjectData, TArray<FString>& InOutSearchStrings)
+				{
+					SReassignObjectComboBox::PopulateSearchTerms(*ConcertClient->GetCurrentSession(), ReassignmentLogic, ObjectData.GetObjectPath(), InOutSearchStrings);
+				})
+				.ColumnSortOrder(ColumnsSortPriority),
+			SHeaderRow::Column(ReassignOwnershipColumnId)
+				.DefaultLabel(LOCTEXT("Owner", "Owner"))
+				.FillSized(FMultiUserReplicationStyle::Get()->GetFloat(TEXT("AllClients.Object.OwnerColumnWidth")))
+			);
+	}
+}
+
+#undef LOCTEXT_NAMESPACE
\ No newline at end of file
diff --git a/Engine/Plugins/Developer/Concert/ConcertApp/MultiUserClient/Source/MultiUserClient/Private/Widgets/ActiveSession/Replication/Client/Multi/Columns/MultiStreamColumns_ReplicationToggle.cpp b/Engine/Plugins/Developer/Concert/ConcertApp/MultiUserClient/Source/MultiUserClient/Private/Widgets/ActiveSession/Replication/Client/Multi/Columns/MultiStreamColumns_ReplicationToggle.cpp
--- a/Engine/Plugins/Developer/Concert/ConcertApp/MultiUserClient/Source/MultiUserClient/Private/Widgets/ActiveSession/Replication/Client/Multi/Columns/MultiStreamColumns_ReplicationToggle.cpp
+++ b/Engine/Plugins/Developer/Concert/ConcertApp/MultiUserClient/Source/MultiUserClient/Private/Widgets/ActiveSession/Replication/Client/Multi/Columns/MultiStreamColumns_ReplicationToggle.cpp
@@ -2,6 +2,7 @@
 
 #include "MultiStreamColumns.h"
 
+#include "MultiUserReplicationStyle.h"
 #include "SReplicationMultiToggleCheckbox.h"
 
 #include "Widgets/Layout/SBox.h"
@@ -37,7 +38,7 @@ namespace UE::MultiUserClient::MultiStreamColumns
 				.ColumnSortOrder(ColumnsSortPriority),
 			SHeaderRow::Column(ReplicationToggleColumnId)
 				.DefaultLabel(FText::GetEmpty())
-				.FixedWidth(70.f)
+				.FixedWidth(FMultiUserReplicationStyle::Get()->GetFloat(TEXT("AllClients.Object.ReplicationToggle")))
 			);
 	}
 }
diff --git a/Engine/Plugins/Developer/Concert/ConcertApp/MultiUserClient/Source/MultiUserClient/Private/Widgets/ActiveSession/Replication/Client/Multi/Columns/SAssignPropertyComboBox.cpp b/Engine/Plugins/Developer/Concert/ConcertApp/MultiUserClient/Source/MultiUserClient/Private/Widgets/ActiveSession/Replication/Client/Multi/Columns/SAssignPropertyComboBox.cpp
--- a/Engine/Plugins/Developer/Concert/ConcertApp/MultiUserClient/Source/MultiUserClient/Private/Widgets/ActiveSession/Replication/Client/Multi/Columns/SAssignPropertyComboBox.cpp
+++ b/Engine/Plugins/Developer/Concert/ConcertApp/MultiUserClient/Source/MultiUserClient/Private/Widgets/ActiveSession/Replication/Client/Multi/Columns/SAssignPropertyComboBox.cpp
@@ -6,6 +6,7 @@
 #include "Replication/Client/ReplicationClient.h"
 #include "Replication/Client/ReplicationClientManager.h"
 #include "Replication/Editor/Model/IEditableReplicationStreamModel.h"
+#include "Widgets/ActiveSession/Replication/Misc/SNoClients.h"
 #include "Widgets/ActiveSession/Replication/Client/ClientUtils.h"
 #include "Widgets/ClientName/SHorizontalClientList.h"
 #include "Widgets/ClientName/SLocalClientName.h"
@@ -45,11 +46,7 @@ namespace UE::MultiUserClient
 			[
 				SAssignNew(ClientListWidget, ConcertClientSharedSlate::SHorizontalClientList, ConcertClient.ToSharedRef())
 				.HighlightText_Lambda([this](){ return HighlightText ? *HighlightText : FText::GetEmpty(); })
-				.EmptyListSlot()
-				[
-					SNew(STextBlock)
-					.Text(LOCTEXT("None", " - "))
-				]
+				.EmptyListSlot() [ SNew(SNoClients) ]
 			]
 			.OnGetMenuContent(this, &SAssignPropertyComboBox::GetMenuContent)
 		];
diff --git a/Engine/Plugins/Developer/Concert/ConcertApp/MultiUserClient/Source/MultiUserClient/Private/Widgets/ActiveSession/Replication/Client/Multi/Columns/SReassignObjectComboBox.cpp b/Engine/Plugins/Developer/Concert/ConcertApp/MultiUserClient/Source/MultiUserClient/Private/Widgets/ActiveSession/Replication/Client/Multi/Columns/SReassignObjectComboBox.cpp
new file mode 100644
--- /dev/null
+++ b/Engine/Plugins/Developer/Concert/ConcertApp/MultiUserClient/Source/MultiUserClient/Private/Widgets/ActiveSession/Replication/Client/Multi/Columns/SReassignObjectComboBox.cpp
@@ -0,0 +1,187 @@
+// Copyright Epic Games, Inc. All Rights Reserved.
+
+#include "SReassignObjectComboBox.h"
+
+#include "MultiUserReplicationStyle.h"
+#include "Misc/EBreakBehavior.h"
+#include "Replication/Client/ReplicationClient.h"
+#include "Replication/Submission/MultiEdit/ReassignObjectPropertiesLogic.h"
+#include "Widgets/ActiveSession/Replication/Client/ClientUtils.h"
+#include "Widgets/ActiveSession/Replication/Misc/SNoClients.h"
+#include "Widgets/ClientName/SHorizontalClientList.h"
+
+#include "Framework/MultiBox/MultiBoxBuilder.h"
+#include "Framework/Commands/UIAction.h"
+#include "Framework/Commands/UICommandInfo.h"
+#include "Textures/SlateIcon.h"
+#include "Widgets/Images/SThrobber.h"
+#include "Widgets/Input/SComboBox.h"
+
+#define LOCTEXT_NAMESPACE "SReassignObjectComboBox"
+
+namespace UE::MultiUserClient
+{
+	void SReassignObjectComboBox::PopulateSearchTerms(
+		const IConcertClientSession& Session,
+		const FReassignObjectPropertiesLogic& ReassignmentLogic,
+		const FSoftObjectPath& ObjectPath,
+		TArray<FString>& InOutSearchTerms
+		)
+	{
+		ReassignmentLogic.EnumerateClientOwnershipState(ObjectPath, [&Session, &InOutSearchTerms](const FGuid& ClientId, FReassignObjectPropertiesLogic::EOwnershipState Ownership)
+		{
+			if (Ownership == FReassignObjectPropertiesLogic::EOwnershipState::HasObjectRegistered)
+			{
+				InOutSearchTerms.Add(ClientUtils::GetClientDisplayName(Session, ClientId));
+			}
+			return EBreakBehavior::Continue;
+		});
+	}
+
+	void SReassignObjectComboBox::Construct(
+		const FArguments& InArgs,
+		TSharedRef<IConcertClient> InConcertClient,
+		FReassignObjectPropertiesLogic& InReassignmentLogic,
+		const FReplicationClientManager& InReplicationManager
+		)
+	{
+		ConcertClient = MoveTemp(InConcertClient);
+		ReassignmentLogic = &InReassignmentLogic;
+		ReplicationManager = &InReplicationManager;
+		HighlightText = InArgs._HighlightText;
+
+		ManagedObject = InArgs._ManagedObject;
+		ConsolidatedStreamModelAttribute = InArgs._ConsolidatedModel;
+		check(ConsolidatedStreamModelAttribute.IsSet() || ConsolidatedStreamModelAttribute.IsBound());
+		
+		ChildSlot
+		[
+			SNew(SHorizontalBox)
+
+			// For indicating that a change is pending
+			+SHorizontalBox::Slot()
+			.AutoWidth()
+			[
+				SNew(SThrobber)
+				.Visibility(this, &SReassignObjectComboBox::GetThrobberVisibility)
+			]
+			
+			+SHorizontalBox::Slot()
+			[
+				SNew(SComboButton)
+				.HasDownArrow(true)
+				.ContentPadding(FMargin(2.0f, 2.0f))
+				.ButtonContent()
+				[
+					SAssignNew(ComboClientList, ConcertClientSharedSlate::SHorizontalClientList, InConcertClient)
+					.HighlightText_Lambda([this](){ return HighlightText ? *HighlightText : FText::GetEmpty(); })
+					.EmptyListSlot() [ SNew(SNoClients) ]
+				]
+				.OnGetMenuContent(this, &SReassignObjectComboBox::MakeMenuContent)
+			]
+		];
+		
+		UpdateComboButtonContent();
+		ReassignmentLogic->OnOwnershipChanged().AddSP(this, &SReassignObjectComboBox::UpdateComboButtonContent);
+	}
+
+	SReassignObjectComboBox::~SReassignObjectComboBox()
+	{
+		ReassignmentLogic->OnOwnershipChanged().RemoveAll(this);
+	}
+
+	void SReassignObjectComboBox::UpdateComboButtonContent() const
+	{
+		using EOwnership = FReassignObjectPropertiesLogic::EOwnershipState;
+		
+		TArray<FGuid> ClientsWithOwnership;
+		ReassignmentLogic->EnumerateClientOwnershipState(ManagedObject, [&ClientsWithOwnership](const FGuid& ClientId, EOwnership Ownership)
+		{
+			if (Ownership == FReassignObjectPropertiesLogic::EOwnershipState::HasObjectRegistered)
+			{
+				ClientsWithOwnership.Add(ClientId);
+			}
+			return EBreakBehavior::Continue;
+		});
+
+		ComboClientList->RefreshList(ClientsWithOwnership);
+	}
+
+	EVisibility SReassignObjectComboBox::GetThrobberVisibility() const
+	{
+		if (!ReassignmentLogic->IsReassigning(ManagedObject))
+		{
+			return EVisibility::Collapsed;
+		}
+		
+		const TOptional<FDateTime> TimeStarted = ReassignmentLogic->GetTimeReassignmentWasStarted();
+		const float WaitTime = FMultiUserReplicationStyle::Get()->GetFloat(TEXT("AllClients.Reassignment.DisplayThrobberAfterSeconds"));
+		const bool bIsTakingLongTime = TimeStarted
+			&& (FDateTime::Now() - *TimeStarted).GetTotalSeconds() >= WaitTime;
+		return bIsTakingLongTime ? EVisibility::Visible : EVisibility::Collapsed;
+	}
+
+	TSharedRef<SWidget> SReassignObjectComboBox::MakeMenuContent() const
+	{
+		FMenuBuilder MenuBuilder(false, nullptr);
+		const TArray<const FReplicationClient*> SortedClients = ClientUtils::GetSortedClientList(*ConcertClient, *ReplicationManager);
+		const FText CanAssignText;
+		
+		MenuBuilder.BeginSection(TEXT("Reassign.This"), LOCTEXT("Reassign.This", "Reassign this to"));
+		AddReassignSection(MenuBuilder, SortedClients, FInlineObjectPathArray{ ManagedObject });
+		MenuBuilder.EndSection();
+
+		// Do not distract the user with more options if children have no assigned properties
+		const bool bHasChildrenWithProperties = ReassignmentLogic->IsAnyObjectOwned(GetChildrenOfManagedObject());
+		if (bHasChildrenWithProperties)
+		{
+			MenuBuilder.BeginSection(TEXT("Reassign.Children"), LOCTEXT("Reassign.Children", "Reassign children to"));
+			AddReassignSection(MenuBuilder, SortedClients, TAttribute<FInlineObjectPathArray>::CreateLambda([this](){ return GetChildrenOfManagedObject(); }));
+			MenuBuilder.EndSection();
+		}
+		
+		return MenuBuilder.MakeWidget();
+	}
+
+	void SReassignObjectComboBox::AddReassignSection(
+		FMenuBuilder& MenuBuilder,
+		const TArray<const FReplicationClient*>& SortedClients,
+		TAttribute<FInlineObjectPathArray> ObjectsToAssign
+		) const
+	{
+		for (const FReplicationClient* Client : SortedClients)
+		{
+			const FGuid& ClientId = Client->GetEndpointId();
+			MenuBuilder.AddMenuEntry(
+				FText::FromString(ClientUtils::GetClientDisplayName(*ConcertClient, Client->GetEndpointId())),
+				TAttribute<FText>::CreateLambda([this, ClientId, ObjectsToAssign]()
+				{
+					FText CannotEditReason;
+					const bool bCanReassign = ReassignmentLogic->CanReassignAnyTo(ObjectsToAssign.Get(), ClientId, &CannotEditReason);
+					return bCanReassign ? LOCTEXT("DoReassign", "Reassign to this client") : CannotEditReason;
+				}),
+				FSlateIcon(),
+				FUIAction(
+					FExecuteAction::CreateLambda([this, ObjectsToAssign, ClientId]()
+					{
+						TArray<FSoftObjectPath> ObjectsNotInline;
+						Algo::Transform(ObjectsToAssign.Get(), ObjectsNotInline, [](const FSoftObjectPath& ObjectPath){ return ObjectPath; });
+						ReassignmentLogic->ReassignAllTo(ObjectsNotInline, ClientId);
+					}),
+					FCanExecuteAction::CreateLambda([this, ObjectsToAssign, ClientId](){ return ReassignmentLogic->CanReassignAnyTo(ObjectsToAssign.Get(), ClientId); }),
+					FIsActionChecked::CreateLambda([this, ObjectsToAssign, ClientId](){ return ReassignmentLogic->OwnsAnyOf(ObjectsToAssign.Get(), ClientId); })
+					),
+					NAME_None,
+					EUserInterfaceActionType::Check
+				);
+		}
+	}
+
+	SReassignObjectComboBox::FInlineObjectPathArray SReassignObjectComboBox::GetChildrenOfManagedObject() const
+	{
+		const ConcertClientSharedSlate::IReplicationStreamModel* Model = ConsolidatedStreamModelAttribute.Get();
+		return ensure(Model) ? Model->GetSubobjects<FInlineAllocator>(ManagedObject) : FInlineObjectPathArray{};
+	}
+}
+
+#undef LOCTEXT_NAMESPACE
\ No newline at end of file
diff --git a/Engine/Plugins/Developer/Concert/ConcertApp/MultiUserClient/Source/MultiUserClient/Private/Widgets/ActiveSession/Replication/Client/Multi/Columns/SReassignObjectComboBox.h b/Engine/Plugins/Developer/Concert/ConcertApp/MultiUserClient/Source/MultiUserClient/Private/Widgets/ActiveSession/Replication/Client/Multi/Columns/SReassignObjectComboBox.h
new file mode 100644
--- /dev/null
+++ b/Engine/Plugins/Developer/Concert/ConcertApp/MultiUserClient/Source/MultiUserClient/Private/Widgets/ActiveSession/Replication/Client/Multi/Columns/SReassignObjectComboBox.h
@@ -0,0 +1,85 @@
+// Copyright Epic Games, Inc. All Rights Reserved.
+
+#pragma once
+
+#include "Widgets/DeclarativeSyntaxSupport.h"
+#include "Widgets/SCompoundWidget.h"
+
+class FMenuBuilder;
+class IConcertClient;
+class IConcertClientSession;
+
+namespace UE::ConcertClientSharedSlate
+{
+	class IReplicationStreamModel;
+	class SHorizontalClientList;
+}
+
+namespace UE::MultiUserClient
+{
+	class FReplicationClient;
+	class FReplicationClientManager;
+	class FReassignObjectPropertiesLogic;
+
+	/** Displayed in object rows of a client view for reassigning object ownership. */
+	class SReassignObjectComboBox : public SCompoundWidget
+	{
+	public:
+
+		/** Fills in the search terms based what the a widget instance would be displaying given its state. */
+		static void PopulateSearchTerms(const IConcertClientSession& Session, const FReassignObjectPropertiesLogic& ReassignmentLogic, const FSoftObjectPath& ObjectPath, TArray<FString>& InOutSearchTerms);
+
+		SLATE_BEGIN_ARGS(SReassignObjectComboBox)
+		{}
+			/** The object being reassigned by this widget. */
+			SLATE_ARGUMENT(FSoftObjectPath, ManagedObject)
+			SLATE_ARGUMENT(TSharedPtr<FText>, HighlightText)
+			/** Used to to figure out child objects. */
+			SLATE_ATTRIBUTE(ConcertClientSharedSlate::IReplicationStreamModel*, ConsolidatedModel)
+		SLATE_END_ARGS()
+
+		void Construct(
+			const FArguments& InArgs,
+			TSharedRef<IConcertClient> InConcertClient,
+			FReassignObjectPropertiesLogic& InReassignmentLogic,
+			const FReplicationClientManager& InReplicationManager
+		);
+		virtual ~SReassignObjectComboBox() override;
+
+	private:
+
+		/** Used to get client display info */
+		TSharedPtr<IConcertClient> ConcertClient;
+		
+		/** Does the logic of reassigning */
+		FReassignObjectPropertiesLogic* ReassignmentLogic = nullptr;
+		/** Used to get all clients and sort them. */
+		const FReplicationClientManager* ReplicationManager = nullptr;
+
+		/** The object being reassigned by this widget. */
+		FSoftObjectPath ManagedObject;
+		/** Used to to figure out child objects. */
+		TAttribute<ConcertClientSharedSlate::IReplicationStreamModel*> ConsolidatedStreamModelAttribute;
+		
+		/** Displayed as content of combo button */
+		TSharedPtr<ConcertClientSharedSlate::SHorizontalClientList> ComboClientList;
+		/** Passed to client list for search highlighting */
+		TSharedPtr<FText> HighlightText;
+
+		void UpdateComboButtonContent() const;
+
+		/** @return Whether the progress throbber should be visible */
+		EVisibility GetThrobberVisibility() const;
+
+		/** @return The menu to display for the combo box */
+		TSharedRef<SWidget> MakeMenuContent() const;
+		using FInlineAllocator = TInlineAllocator<24>;
+		using FInlineObjectPathArray = TArray<FSoftObjectPath, FInlineAllocator>;
+		/** Populates a menu which for reassigning a client to ObjectsToAssign. */
+		void AddReassignSection(FMenuBuilder& MenuBuilder, const TArray<const FReplicationClient*>& SortedClients, TAttribute<FInlineObjectPathArray> ObjectsToAssign) const;
+
+		/** @return All children of ManagedObject that need reassigning. */
+		FInlineObjectPathArray GetChildrenOfManagedObject() const;
+	};
+}
+
diff --git a/Engine/Plugins/Developer/Concert/ConcertApp/MultiUserClient/Source/MultiUserClient/Private/Widgets/ActiveSession/Replication/Client/Multi/SMultiClientView.cpp b/Engine/Plugins/Developer/Concert/ConcertApp/MultiUserClient/Source/MultiUserClient/Private/Widgets/ActiveSession/Replication/Client/Multi/SMultiClientView.cpp
--- a/Engine/Plugins/Developer/Concert/ConcertApp/MultiUserClient/Source/MultiUserClient/Private/Widgets/ActiveSession/Replication/Client/Multi/SMultiClientView.cpp
+++ b/Engine/Plugins/Developer/Concert/ConcertApp/MultiUserClient/Source/MultiUserClient/Private/Widgets/ActiveSession/Replication/Client/Multi/SMultiClientView.cpp
@@ -76,7 +76,11 @@ namespace UE::MultiUserClient
 			.ViewerParams =
 			{
 				.SubobjectModel = CreateDefaultComponentHierarchySubobjectModel(), // This makes actors have children in the top view
-				.AdditionalObjectColumns = { MultiStreamColumns::ReplicationToggle(InConcertClient, MoveTemp(ConsolidatedStreamModelAttribute), InClientManager) },
+				.AdditionalObjectColumns =
+				{
+					MultiStreamColumns::ReplicationToggle(InConcertClient, ConsolidatedStreamModelAttribute, InClientManager),
+					MultiStreamColumns::ReassignOwnership(InConcertClient, ConsolidatedStreamModelAttribute, InClientManager.GetReassignmentLogic(), InClientManager)
+				},
 				.AdditionalPropertyColumns = { MultiStreamColumns::AssignPropertyColumn(MoveTemp(MultiStreamEditorAttribute), InConcertClient, InClientManager) }
 			}
 		};
diff --git a/Engine/Plugins/Developer/Concert/ConcertApp/MultiUserClient/Source/MultiUserClient/Private/Widgets/ActiveSession/Replication/Client/Single/Columns/SOwnerClientList.cpp b/Engine/Plugins/Developer/Concert/ConcertApp/MultiUserClient/Source/MultiUserClient/Private/Widgets/ActiveSession/Replication/Client/Single/Columns/SOwnerClientList.cpp
--- a/Engine/Plugins/Developer/Concert/ConcertApp/MultiUserClient/Source/MultiUserClient/Private/Widgets/ActiveSession/Replication/Client/Single/Columns/SOwnerClientList.cpp
+++ b/Engine/Plugins/Developer/Concert/ConcertApp/MultiUserClient/Source/MultiUserClient/Private/Widgets/ActiveSession/Replication/Client/Single/Columns/SOwnerClientList.cpp
@@ -3,6 +3,7 @@
 #include "SOwnerClientList.h"
 
 #include "Replication/Util/GlobalAuthorityCache.h"
+#include "Widgets/ActiveSession/Replication/Misc/SNoClients.h"
 #include "Widgets/ClientName/SHorizontalClientList.h"
 
 #include "Widgets/Layout/SBox.h"
@@ -30,8 +31,7 @@ namespace UE::MultiUserClient
 				.HAlign(HAlign_Left)
 				.Padding(5.f, 0.f, 0.f, 0.f)
 				[
-					SNew(STextBlock)
-					.Text(LOCTEXT("NotApplicable", "-"))
+					SNew(SNoClients)
 				]
 			]
 		];
diff --git a/Engine/Plugins/Developer/Concert/ConcertApp/MultiUserClient/Source/MultiUserClient/Private/Widgets/ActiveSession/Replication/Client/Single/Columns/SingleClientColumns_Owner.cpp b/Engine/Plugins/Developer/Concert/ConcertApp/MultiUserClient/Source/MultiUserClient/Private/Widgets/ActiveSession/Replication/Client/Single/Columns/SingleClientColumns_Owner.cpp
--- a/Engine/Plugins/Developer/Concert/ConcertApp/MultiUserClient/Source/MultiUserClient/Private/Widgets/ActiveSession/Replication/Client/Single/Columns/SingleClientColumns_Owner.cpp
+++ b/Engine/Plugins/Developer/Concert/ConcertApp/MultiUserClient/Source/MultiUserClient/Private/Widgets/ActiveSession/Replication/Client/Single/Columns/SingleClientColumns_Owner.cpp
@@ -3,12 +3,12 @@
 #include "SingleClientColumns.h"
 
 #include "IConcertClient.h"
-#include "SOwnerClientList.h"
+#include "MultiUserReplicationStyle.h"
 #include "Replication/Editor/View/IReplicationStreamViewer.h"
 #include "Replication/Editor/View/ReplicationColumnsUtils.h"
 #include "Replication/Util/GlobalAuthorityCache.h"
+#include "SOwnerClientList.h"
 #include "Widgets/ActiveSession/Replication/Client/ClientUtils.h"
-#include "Widgets/ClientName/SClientName.h"
 
 #define LOCTEXT_NAMESPACE "SingleClientColumns.Owner"
 
@@ -19,8 +19,6 @@ namespace UE::MultiUserClient::SingleClientColumns
 
 	namespace Private::Shared
 	{
-		constexpr float DefaultWidth = 200.f;
-		
 		static void PopulateSearchTerms(const TSharedRef<IConcertClient>& InClient, const TArray<FGuid>& Clients, TArray<FString>& InOutSearchStrings)
 		{
 			for (const FGuid& ClientId : Clients)
@@ -29,60 +27,6 @@ namespace UE::MultiUserClient::SingleClientColumns
 			}
 		}
 	}
-
-	namespace Private::TopLevel
-	{
-		static TArray<FGuid> GetOwnersOfTopLevelObject(
-			const FGlobalAuthorityCache& InAuthorityCache,
-			ConcertClientSharedSlate::IReplicationStreamModel& InObjectModel,
-			const FSoftObjectPath& InTopLevelObject
-			)
-		{
-			TArray<FGuid> Owners = InAuthorityCache.GetClientsWithAuthorityOverObject(InTopLevelObject);
-			InObjectModel.ForEachSubobject(InTopLevelObject, [&InAuthorityCache, &Owners](const FSoftObjectPath& Subobject)
-			{
-				for (const FGuid& Client : InAuthorityCache.GetClientsWithAuthorityOverObject(Subobject))
-				{
-					Owners.AddUnique(Client);
-				}
-				return EBreakBehavior::Continue;
-			});
-			return Owners;
-		}
-	}
-	
-	ConcertClientSharedSlate::ReplicationColumns::FReplicationTopLevelObjectColumn OwnerOfTopLevelObject(
-		const TSharedRef<IConcertClient>& InClient,
-		FGlobalAuthorityCache& InAuthorityCache,
-		ConcertClientSharedSlate::IReplicationStreamModel& InObjectModel
-		)
-	{
-		using FColumnType = ConcertClientSharedSlate::ReplicationColumns::FReplicationTopLevelObjectColumn;
-		return FColumnType(
-			FColumnType::FArguments()
-				.GenerateWidgetColumn_Lambda([InClient, &InObjectModel, &InAuthorityCache](const FColumnType::FBuildArgs& InArgs)
-				{
-					return SNew(SOwnerClientList, InClient, InAuthorityCache)
-						.GetClientList_Lambda([&InAuthorityCache, &InObjectModel, ObjectPath = InArgs.RowData.GetObjectPath()](const FGlobalAuthorityCache&)
-						{
-							return Private::TopLevel::GetOwnersOfTopLevelObject(InAuthorityCache, InObjectModel, ObjectPath);
-						})
-						.HighlightText_Lambda([HighlightText = InArgs.HighlightText](){ return *HighlightText.Get(); });
-				})
-				.PopulateSearchItems_Lambda([InClient, &InAuthorityCache, &InObjectModel](const ConcertClientSharedSlate::FReplicatedObjectData& InArgs, TArray<FString>& InOutSearchStrings)
-				{
-					Private::Shared::PopulateSearchTerms(
-						InClient,
-						Private::TopLevel::GetOwnersOfTopLevelObject(InAuthorityCache, InObjectModel, InArgs.GetObjectPath()),
-						InOutSearchStrings
-						);
-				})
-				.ColumnSortOrder(static_cast<int32>(ETopLevelObjectColumnOrder::Owner)),
-			SHeaderRow::Column(OwnerOfSubobjectColumnId)
-				.DefaultLabel(LOCTEXT("Subobject.Owner", "Owner"))
-				.FillSized(Private::Shared::DefaultWidth)
-			);
-	}
 	
 	ConcertClientSharedSlate::ReplicationColumns::FReplicationTopLevelObjectColumn OwnerOfObject(
 		const TSharedRef<IConcertClient>& InClient,
@@ -112,7 +56,7 @@ namespace UE::MultiUserClient::SingleClientColumns
 				.ColumnSortOrder(static_cast<int32>(ETopLevelObjectColumnOrder::Owner)),
 			SHeaderRow::Column(OwnerOfSubobjectColumnId)
 				.DefaultLabel(LOCTEXT("Subobject.Owner", "Owner"))
-				.FillSized(Private::Shared::DefaultWidth)
+				.FillSized(FMultiUserReplicationStyle::Get()->GetFloat(TEXT("SingleClient.Object.OwnerColumnWidth")))
 			);
 	}
 	
@@ -167,7 +111,7 @@ namespace UE::MultiUserClient::SingleClientColumns
 				.ColumnSortOrder(static_cast<int32>(EPropertyColumnOrder::Owner)),
 			SHeaderRow::Column(OwnerOfPropertyColumnId)
 				.DefaultLabel(LOCTEXT("Property.Owner", "Owner"))
-				.FillSized(Private::Shared::DefaultWidth)
+				.FillSized(FMultiUserReplicationStyle::Get()->GetFloat(TEXT("SingleClient.Property.OwnerColumnWidth")))
 			);
 	}
 }
diff --git a/Engine/Plugins/Developer/Concert/ConcertApp/MultiUserClient/Source/MultiUserClient/Private/Widgets/ActiveSession/Replication/Misc/SNoClients.h b/Engine/Plugins/Developer/Concert/ConcertApp/MultiUserClient/Source/MultiUserClient/Private/Widgets/ActiveSession/Replication/Misc/SNoClients.h
new file mode 100644
--- /dev/null
+++ b/Engine/Plugins/Developer/Concert/ConcertApp/MultiUserClient/Source/MultiUserClient/Private/Widgets/ActiveSession/Replication/Misc/SNoClients.h
@@ -0,0 +1,33 @@
+// Copyright Epic Games, Inc. All Rights Reserved.
+
+#pragma once
+
+#include "Widgets/DeclarativeSyntaxSupport.h"
+#include "Widgets/SCompoundWidget.h"
+#include "Widgets/Text/STextBlock.h"
+
+#define LOCTEXT_NAMESPACE "SNoClients"
+
+namespace UE::MultiUserClient
+{
+	/** Special widget that indicates no clients in widgets that list clients, like ownership reassignment combo box. */
+	class SNoClients : public SCompoundWidget
+	{
+	public:
+
+		SLATE_BEGIN_ARGS(SNoClients)
+		{}
+		SLATE_END_ARGS()
+
+		void Construct(const FArguments& InArgs)
+		{
+			ChildSlot
+			[
+				SNew(STextBlock)
+				.Text(LOCTEXT("None", " - "))
+			];
+		}
+	};
+}
+
+#undef LOCTEXT_NAMESPACE
\ No newline at end of file
diff --git a/Engine/Plugins/Developer/Concert/ConcertSync/ConcertSyncCore/Source/ConcertSyncCore/Private/Replication/Data/ConcertPropertySelection.cpp b/Engine/Plugins/Developer/Concert/ConcertSync/ConcertSyncCore/Source/ConcertSyncCore/Private/Replication/Data/ConcertPropertySelection.cpp
--- a/Engine/Plugins/Developer/Concert/ConcertSync/ConcertSyncCore/Source/ConcertSyncCore/Private/Replication/Data/ConcertPropertySelection.cpp
+++ b/Engine/Plugins/Developer/Concert/ConcertSync/ConcertSyncCore/Source/ConcertSyncCore/Private/Replication/Data/ConcertPropertySelection.cpp
@@ -5,6 +5,7 @@
 #include "ConcertLogGlobal.h"
 #include "Replication/PropertyChainUtils.h"
 
+#include "Algo/AllOf.h"
 #include "Serialization/ArchiveSerializedPropertyChain.h"
 
 const FName FConcertPropertyChain::InternalContainerPropertyValueName(TEXT("Value"));
@@ -201,6 +202,19 @@ FString FConcertPropertyChain::ToString(EToStringMethod Method) const
 	}
 }
 
+bool FConcertPropertySelection::Includes(const FConcertPropertySelection& Other) const
+{
+	if (ReplicatedProperties.Num() < Other.ReplicatedProperties.Num())
+	{
+		return false;
+	}
+
+	return Algo::AllOf(ReplicatedProperties, [&Other](const FConcertPropertyChain& Property)
+	{
+		return Other.ReplicatedProperties.Contains(Property);
+	});
+}
+
 bool FConcertPropertySelection::EnumeratePropertyOverlaps(
 	TConstArrayView<FConcertPropertyChain> First,
 	TConstArrayView<FConcertPropertyChain> Second,
diff --git a/Engine/Plugins/Developer/Concert/ConcertSync/ConcertSyncCore/Source/ConcertSyncCore/Public/Replication/Data/ConcertPropertySelection.h b/Engine/Plugins/Developer/Concert/ConcertSync/ConcertSyncCore/Source/ConcertSyncCore/Public/Replication/Data/ConcertPropertySelection.h
--- a/Engine/Plugins/Developer/Concert/ConcertSync/ConcertSyncCore/Source/ConcertSyncCore/Public/Replication/Data/ConcertPropertySelection.h
+++ b/Engine/Plugins/Developer/Concert/ConcertSync/ConcertSyncCore/Source/ConcertSyncCore/Public/Replication/Data/ConcertPropertySelection.h
@@ -173,18 +173,16 @@ struct CONCERTSYNCCORE_API FConcertPropertySelection
 {
 	GENERATED_BODY()
 
-	/**
-	 * List of replicated objects.
-	 *
-	 * Can be actors or its subobjects.
-	 * Technically, this can also be non-UWorld objects.
-	 */
+	/** List of replicated properties. */
 	UPROPERTY()
 	TArray<FConcertPropertyChain> ReplicatedProperties;
 
 	/** @return Whether this and Other contain at least one property that is the same. */
 	bool OverlapsWith(const FConcertPropertySelection& Other) const { return EnumeratePropertyOverlaps(ReplicatedProperties, Other.ReplicatedProperties); }
 
+	/** @return Whether this includes all properties of Other */
+	bool Includes(const FConcertPropertySelection& Other) const;
+	
 	/**
 	 * Determines all properties that overlap.
 	 * This algorithm is strictly O(n^2) but runs O(n) on average.
@@ -196,7 +194,7 @@ struct CONCERTSYNCCORE_API FConcertPropertySelection
 		TConstArrayView<FConcertPropertyChain> Second,
 		TFunctionRef<EBreakBehavior(const FConcertPropertyChain&)> Callback = [](const FConcertPropertyChain&){ return EBreakBehavior::Break; }
 		);
-
+	
 	friend bool operator==(const FConcertPropertySelection& Left, const FConcertPropertySelection& Right)
 	{
 		return Left.ReplicatedProperties == Right.ReplicatedProperties;
diff --git a/Engine/Plugins/Developer/Concert/ConcertSync/ConcertSyncCore/Source/ConcertSyncCore/Public/Replication/Data/ObjectReplicationMap.h b/Engine/Plugins/Developer/Concert/ConcertSync/ConcertSyncCore/Source/ConcertSyncCore/Public/Replication/Data/ObjectReplicationMap.h
--- a/Engine/Plugins/Developer/Concert/ConcertSync/ConcertSyncCore/Source/ConcertSyncCore/Public/Replication/Data/ObjectReplicationMap.h
+++ b/Engine/Plugins/Developer/Concert/ConcertSync/ConcertSyncCore/Source/ConcertSyncCore/Public/Replication/Data/ObjectReplicationMap.h
@@ -45,6 +45,19 @@ struct FObjectReplicationMap
 	UPROPERTY()
 	TMap<FSoftObjectPath, FReplicatedObjectInfo> ReplicatedObjects;
 
+	/** @return Whether ObjectPath has any properties assigned to it. */
+	bool HasProperties(const FSoftObjectPath& ObjectPath) const
+	{
+		const FReplicatedObjectInfo* ObjectInfo = ReplicatedObjects.Find(ObjectPath);
+		return ObjectInfo && !ObjectInfo->PropertySelection.ReplicatedProperties.IsEmpty();
+	}
+
+	/** @return Whether there are any properties assigned. */
+	bool IsEmpty() const
+	{
+		return ReplicatedObjects.IsEmpty();
+	}
+
 	friend bool operator==(const FObjectReplicationMap& Left, const FObjectReplicationMap& Right)
 	{
 		return Left.ReplicatedObjects.OrderIndependentCompareEqual(Right.ReplicatedObjects);
diff --git a/Engine/Plugins/Developer/Concert/ConcertUI/ConcertClientSharedSlate/Source/ConcertClientSharedSlate/Public/Replication/Editor/Model/IReplicationStreamModel.h b/Engine/Plugins/Developer/Concert/ConcertUI/ConcertClientSharedSlate/Source/ConcertClientSharedSlate/Public/Replication/Editor/Model/IReplicationStreamModel.h
--- a/Engine/Plugins/Developer/Concert/ConcertUI/ConcertClientSharedSlate/Source/ConcertClientSharedSlate/Public/Replication/Editor/Model/IReplicationStreamModel.h
+++ b/Engine/Plugins/Developer/Concert/ConcertUI/ConcertClientSharedSlate/Source/ConcertClientSharedSlate/Public/Replication/Editor/Model/IReplicationStreamModel.h
@@ -134,9 +134,10 @@ namespace UE::ConcertClientSharedSlate
 		}
 
 		/** @return Gets the subobjects as an array */
-		TArray<FSoftObjectPath> GetSubobjects(const FSoftObjectPath& Parent) const
+		template<typename TAllocator = FDefaultAllocator>
+		TArray<FSoftObjectPath, TAllocator> GetSubobjects(const FSoftObjectPath& Parent) const
 		{
-			TArray<FSoftObjectPath> Subobjects;
+			TArray<FSoftObjectPath, TAllocator> Subobjects;
 			ForEachSubobject(Parent, [&Subobjects](const FSoftObjectPath& Child)
 			{
 				Subobjects.Add(Child);
