Fix dereferencing nullptr static analysis

diff --git a/Engine/Plugins/Developer/Concert/ConcertApp/MultiUserClient/Source/MultiUserClient/Private/Replication/Submission/MultiEdit/ReassignObjectPropertiesLogic.cpp b/Engine/Plugins/Developer/Concert/ConcertApp/MultiUserClient/Source/MultiUserClient/Private/Replication/Submission/MultiEdit/ReassignObjectPropertiesLogic.cpp
--- a/Engine/Plugins/Developer/Concert/ConcertApp/MultiUserClient/Source/MultiUserClient/Private/Replication/Submission/MultiEdit/ReassignObjectPropertiesLogic.cpp
+++ b/Engine/Plugins/Developer/Concert/ConcertApp/MultiUserClient/Source/MultiUserClient/Private/Replication/Submission/MultiEdit/ReassignObjectPropertiesLogic.cpp
@@ -160,18 +160,12 @@ namespace UE::MultiUserClient
 				const FReplicatedObjectInfo* OtherObjectInfo = OtherClientObjects.ReplicatedObjects.Find(ObjectToCheck);
 					
 				// TargetClient has at least as much if OtherClient has nothing
-				if (!OtherObjectInfo)
+				if (!OtherObjectInfo || OtherObjectInfo->PropertySelection.ReplicatedProperties.IsEmpty())
 				{
 					continue;
 				}
 
-				// TargetClient has nothing but TargetClient has something?
-				if (!ObjectInfo && !OtherObjectInfo->PropertySelection.ReplicatedProperties.IsEmpty())
-				{
-					return false;
-				}
-
-				if (!ObjectInfo->PropertySelection.Includes(OtherObjectInfo->PropertySelection))
+				if (!ObjectInfo || !ObjectInfo->PropertySelection.Includes(OtherObjectInfo->PropertySelection))
 				{
 					return false;
 				}
