Auto add properties and objects in the All Clients view

diff --git a/Engine/Plugins/Developer/Concert/ConcertApp/MultiUserClient/Source/MultiUserClient/Private/Widgets/ActiveSession/Replication/Client/Multi/SMultiClientView.cpp b/Engine/Plugins/Developer/Concert/ConcertApp/MultiUserClient/Source/MultiUserClient/Private/Widgets/ActiveSession/Replication/Client/Multi/SMultiClientView.cpp
--- a/Engine/Plugins/Developer/Concert/ConcertApp/MultiUserClient/Source/MultiUserClient/Private/Widgets/ActiveSession/Replication/Client/Multi/SMultiClientView.cpp
+++ b/Engine/Plugins/Developer/Concert/ConcertApp/MultiUserClient/Source/MultiUserClient/Private/Widgets/ActiveSession/Replication/Client/Multi/SMultiClientView.cpp
@@ -21,11 +21,6 @@ namespace UE::MultiUserClient
 {
 	void SMultiClientView::Construct(const FArguments& InArgs, TSharedRef<IConcertClient> InConcertClient, FReplicationClientManager& ClientManager, IClientSelectionModel& InDisplayClientsModel)
 	{
-		TAttribute<const FConcertReplicationEditorSettings*> ReplicationSettingsAttribute =
-			TAttribute<const FConcertReplicationEditorSettings*>::CreateLambda([]()
-			{
-				return &UMultiUserReplicationSettings::Get()->ReplicationEditorSettings;
-			});
 		StreamModel = MakeShared<FMultiStreamModel>(InDisplayClientsModel, ClientManager);
 
 		ChildSlot
@@ -66,13 +61,19 @@ namespace UE::MultiUserClient
 		   {
 			   return &StreamEditor->GetConsolidatedModel();
 		   });
+		FGetAutoAssignTarget GetAutoAssignTargetDelegate = FGetAutoAssignTarget::CreateLambda([this, &InClientManager](TConstArrayView<UObject*>)
+		{
+			const TSharedRef<IEditableReplicationStreamModel>& LocalStream = InClientManager.GetLocalClient().GetClientEditModel();
+			return StreamModel->GetEditableStreams().Contains(LocalStream) ? LocalStream.ToSharedPtr() : nullptr;
+		});
 		
 		FCreateMultiStreamEditorParams Params
 		{
 			.MultiStreamModel = StreamModel.ToSharedRef(),
 			.ObjectSource = MakeShared<FActorSelectionSourceModel>(),
 			.PropertySource = MakeShared<FSelectPropertyFromUClassModel>(),
-			.ViewerParams =
+			.GetAutoAssignToStreamDelegate = MoveTemp(GetAutoAssignTargetDelegate),
+			.ViewerParams 
 			{
 				.SubobjectModel = CreateDefaultComponentHierarchySubobjectModel(), // This makes actors have children in the top view
 				.AdditionalObjectColumns =
diff --git a/Engine/Plugins/Developer/Concert/ConcertUI/ConcertClientSharedSlate/Source/ConcertClientSharedSlate/Private/Replication/Editor/View/MultiEditor/ConsolidatedMultiStreamModel.cpp b/Engine/Plugins/Developer/Concert/ConcertUI/ConcertClientSharedSlate/Source/ConcertClientSharedSlate/Private/Replication/Editor/View/MultiEditor/ConsolidatedMultiStreamModel.cpp
--- a/Engine/Plugins/Developer/Concert/ConcertUI/ConcertClientSharedSlate/Source/ConcertClientSharedSlate/Private/Replication/Editor/View/MultiEditor/ConsolidatedMultiStreamModel.cpp
+++ b/Engine/Plugins/Developer/Concert/ConcertUI/ConcertClientSharedSlate/Source/ConcertClientSharedSlate/Private/Replication/Editor/View/MultiEditor/ConsolidatedMultiStreamModel.cpp
@@ -40,11 +40,13 @@ namespace UE::ConcertClientSharedSlate
 	
 	FConsolidatedMultiStreamModel::FConsolidatedMultiStreamModel(
 		TSharedRef<IEditableMultiReplicationStreamModel> InMultiStreamModel,
-		bool bShouldSupportTransactions
+		bool bShouldSupportTransactions,
+		FGetAutoAssignTarget InGetAutoAssignTargetDelegate
 		)
 		: StreamForAddingObject(Private::MakeObject(bShouldSupportTransactions))
 		, StreamForAdding(Private::MakeModel(StreamForAddingObject, bShouldSupportTransactions))
 		, MultiStreamModel(MoveTemp(InMultiStreamModel))
+		, GetAutoAssignTargetDelegate(MoveTemp(InGetAutoAssignTargetDelegate))
 	{
 		MultiStreamModel->OnStreamExternallyChanged().AddRaw(this, &FConsolidatedMultiStreamModel::OnStreamExternallyChanged);
 		MultiStreamModel->OnStreamSetChanged().AddRaw(this, &FConsolidatedMultiStreamModel::RebuildStreamSubscriptions);
@@ -125,7 +127,17 @@ namespace UE::ConcertClientSharedSlate
 
 	void FConsolidatedMultiStreamModel::AddObjects(TConstArrayView<UObject*> Objects)
 	{
-		StreamForAdding->AddObjects(Objects);
+		const TSharedPtr<IEditableReplicationStreamModel> TargetStream = GetAutoAssignTargetDelegate.IsBound()
+			? GetAutoAssignTargetDelegate.Execute(Objects)
+			: nullptr;
+		if (TargetStream && ensure(MultiStreamModel->GetEditableStreams().Contains(TargetStream.ToSharedRef())))
+		{
+			TargetStream->AddObjects(Objects);
+		}
+		else
+		{
+			StreamForAdding->AddObjects(Objects);
+		}
 	}
 
 	void FConsolidatedMultiStreamModel::RemoveObjects(TConstArrayView<FSoftObjectPath> Objects)
diff --git a/Engine/Plugins/Developer/Concert/ConcertUI/ConcertClientSharedSlate/Source/ConcertClientSharedSlate/Private/Replication/Editor/View/MultiEditor/ConsolidatedMultiStreamModel.h b/Engine/Plugins/Developer/Concert/ConcertUI/ConcertClientSharedSlate/Source/ConcertClientSharedSlate/Private/Replication/Editor/View/MultiEditor/ConsolidatedMultiStreamModel.h
--- a/Engine/Plugins/Developer/Concert/ConcertUI/ConcertClientSharedSlate/Source/ConcertClientSharedSlate/Private/Replication/Editor/View/MultiEditor/ConsolidatedMultiStreamModel.h
+++ b/Engine/Plugins/Developer/Concert/ConcertUI/ConcertClientSharedSlate/Source/ConcertClientSharedSlate/Private/Replication/Editor/View/MultiEditor/ConsolidatedMultiStreamModel.h
@@ -2,6 +2,7 @@
 
 #pragma once
 
+#include "Replication/ReplicationWidgetFactories.h"
 #include "Replication/Editor/Model/IEditableReplicationStreamModel.h"
 #include "UObject/GCObject.h"
 
@@ -35,7 +36,8 @@ namespace UE::ConcertClientSharedSlate
 
 		FConsolidatedMultiStreamModel(
 			TSharedRef<IEditableMultiReplicationStreamModel> InMultiStreamModel,
-			bool bShouldSupportTransactions
+			bool bShouldSupportTransactions,
+			FGetAutoAssignTarget InGetAutoAssignTargetDelegate
 			);
 		virtual ~FConsolidatedMultiStreamModel();
 		
@@ -82,6 +84,8 @@ namespace UE::ConcertClientSharedSlate
 		FOnObjectsChanged OnObjectsChangedDelegate;
 		/** Never actually called, just to implement OnPropertiesChanged. */
 		FOnPropertiesChanged OnPropertiesChangedDelegate;
+		/** If set, also adds newly added objects to the target stream. */
+		FGetAutoAssignTarget GetAutoAssignTargetDelegate;
 		
 		/** All streams we've subscribed to for changes. Used so we can unsubscribe when MultiStreamModel->OnStreamsChanged triggers. */
 		TArray<TWeakPtr<IEditableReplicationStreamModel>> SubscribedToViewers;
diff --git a/Engine/Plugins/Developer/Concert/ConcertUI/ConcertClientSharedSlate/Source/ConcertClientSharedSlate/Private/Replication/Editor/View/MultiEditor/SMultiReplicationStreamEditor.cpp b/Engine/Plugins/Developer/Concert/ConcertUI/ConcertClientSharedSlate/Source/ConcertClientSharedSlate/Private/Replication/Editor/View/MultiEditor/SMultiReplicationStreamEditor.cpp
--- a/Engine/Plugins/Developer/Concert/ConcertUI/ConcertClientSharedSlate/Source/ConcertClientSharedSlate/Private/Replication/Editor/View/MultiEditor/SMultiReplicationStreamEditor.cpp
+++ b/Engine/Plugins/Developer/Concert/ConcertUI/ConcertClientSharedSlate/Source/ConcertClientSharedSlate/Private/Replication/Editor/View/MultiEditor/SMultiReplicationStreamEditor.cpp
@@ -14,7 +14,7 @@ namespace UE::ConcertClientSharedSlate
 		MultiStreamModel = MoveTemp(Params.MultiStreamModel);
 
 		const bool bConsolidatedModelShouldTransact = EnumHasAnyFlags(Params.Flags, EMultiStreamEditorFlags::Transactional);
-		ConsolidatedModel = MakeShared<FConsolidatedMultiStreamModel>(MoveTemp(Params.MultiStreamModel), bConsolidatedModelShouldTransact);
+		ConsolidatedModel = MakeShared<FConsolidatedMultiStreamModel>(MoveTemp(Params.MultiStreamModel), bConsolidatedModelShouldTransact, InArgs._GetAutoAssignStream);
 		
 		const FCreateEditorParams BaseEditorParams
 		{
@@ -24,7 +24,7 @@ namespace UE::ConcertClientSharedSlate
 			.ViewerParams = MoveTemp(Params.ViewerParams)
 		};
 		EditorView = CreateBaseStreamEditor(BaseEditorParams);
-		
+
 		ChildSlot
 		[
 			EditorView.ToSharedRef()
diff --git a/Engine/Plugins/Developer/Concert/ConcertUI/ConcertClientSharedSlate/Source/ConcertClientSharedSlate/Private/Replication/Editor/View/MultiEditor/SMultiReplicationStreamEditor.h b/Engine/Plugins/Developer/Concert/ConcertUI/ConcertClientSharedSlate/Source/ConcertClientSharedSlate/Private/Replication/Editor/View/MultiEditor/SMultiReplicationStreamEditor.h
--- a/Engine/Plugins/Developer/Concert/ConcertUI/ConcertClientSharedSlate/Source/ConcertClientSharedSlate/Private/Replication/Editor/View/MultiEditor/SMultiReplicationStreamEditor.h
+++ b/Engine/Plugins/Developer/Concert/ConcertUI/ConcertClientSharedSlate/Source/ConcertClientSharedSlate/Private/Replication/Editor/View/MultiEditor/SMultiReplicationStreamEditor.h
@@ -35,6 +35,8 @@ namespace UE::ConcertClientSharedSlate
 
 		SLATE_BEGIN_ARGS(SMultiReplicationStreamEditor)
 		{}
+			/** Optional. If set, the Add Actor button should automatically assign the added object to stream returned */
+			SLATE_EVENT(FGetAutoAssignTarget, GetAutoAssignStream)
 		SLATE_END_ARGS()
 
 		void Construct(const FArguments& InArgs, FCreateMultiStreamEditorParams Params);
diff --git a/Engine/Plugins/Developer/Concert/ConcertUI/ConcertClientSharedSlate/Source/ConcertClientSharedSlate/Private/Replication/ReplicationWidgetFactories.cpp b/Engine/Plugins/Developer/Concert/ConcertUI/ConcertClientSharedSlate/Source/ConcertClientSharedSlate/Private/Replication/ReplicationWidgetFactories.cpp
--- a/Engine/Plugins/Developer/Concert/ConcertUI/ConcertClientSharedSlate/Source/ConcertClientSharedSlate/Private/Replication/ReplicationWidgetFactories.cpp
+++ b/Engine/Plugins/Developer/Concert/ConcertUI/ConcertClientSharedSlate/Source/ConcertClientSharedSlate/Private/Replication/ReplicationWidgetFactories.cpp
@@ -58,7 +58,8 @@ namespace UE::ConcertClientSharedSlate
 
 	TSharedRef<IMultiReplicationStreamEditor> CreateBaseMultiStreamEditor(FCreateMultiStreamEditorParams Params)
 	{
-		return SNew(SMultiReplicationStreamEditor, MoveTemp(Params));
+		return SNew(SMultiReplicationStreamEditor, MoveTemp(Params))
+			.GetAutoAssignStream(MoveTemp(Params.GetAutoAssignToStreamDelegate));
 	}
 }
 
diff --git a/Engine/Plugins/Developer/Concert/ConcertUI/ConcertClientSharedSlate/Source/ConcertClientSharedSlate/Public/Replication/ReplicationWidgetFactories.h b/Engine/Plugins/Developer/Concert/ConcertUI/ConcertClientSharedSlate/Source/ConcertClientSharedSlate/Public/Replication/ReplicationWidgetFactories.h
--- a/Engine/Plugins/Developer/Concert/ConcertUI/ConcertClientSharedSlate/Source/ConcertClientSharedSlate/Public/Replication/ReplicationWidgetFactories.h
+++ b/Engine/Plugins/Developer/Concert/ConcertUI/ConcertClientSharedSlate/Source/ConcertClientSharedSlate/Public/Replication/ReplicationWidgetFactories.h
@@ -147,7 +147,8 @@ namespace UE::ConcertClientSharedSlate
 		Transactional = 1 << 0
 	};
 	ENUM_CLASS_FLAGS(EMultiStreamEditorFlags)
-	
+
+	DECLARE_DELEGATE_RetVal_OneParam(TSharedPtr<IEditableReplicationStreamModel>, FGetAutoAssignTarget, TConstArrayView<UObject*>);
 	/** Params for creating an IMultiReplicationStreamEditor */
 	struct FCreateMultiStreamEditorParams
 	{
@@ -171,6 +172,9 @@ namespace UE::ConcertClientSharedSlate
 
 		/** These flags modify the behaviour */
 		EMultiStreamEditorFlags Flags = EMultiStreamEditorFlags::Transactional;
+
+		/** Optional. If set, the Add Actor button should automatically assign the added object to stream returned */
+		FGetAutoAssignTarget GetAutoAssignToStreamDelegate;
 		
 		/** Base params for customizing the viewing part of the editor */
 		FCreateViewerParams ViewerParams;
