Submit missing file

diff --git a/Engine/Plugins/Developer/Concert/ConcertApp/MultiUserClient/Source/MultiUserClient/Private/Widgets/ActiveSession/Replication/Client/Multi/SMultiClientView.h b/Engine/Plugins/Developer/Concert/ConcertApp/MultiUserClient/Source/MultiUserClient/Private/Widgets/ActiveSession/Replication/Client/Multi/SMultiClientView.h
--- a/Engine/Plugins/Developer/Concert/ConcertApp/MultiUserClient/Source/MultiUserClient/Private/Widgets/ActiveSession/Replication/Client/Multi/SMultiClientView.h
+++ b/Engine/Plugins/Developer/Concert/ConcertApp/MultiUserClient/Source/MultiUserClient/Private/Widgets/ActiveSession/Replication/Client/Multi/SMultiClientView.h
@@ -37,11 +37,15 @@ namespace UE::MultiUserClient
 		SLATE_END_ARGS()
 
 		void Construct(const FArguments& InArgs, TSharedRef<IConcertClient> InConcertClient, FReplicationClientManager& InClientManager, IClientSelectionModel& InDisplayClientsModel);
+		virtual ~SMultiClientView() override;
 
 	private:
 
 		/** Used to access and update the replication status widget. */
 		TSharedPtr<SClientToolbar> Toolbar;
+
+		FReplicationClientManager* ClientManager = nullptr;
+		IClientSelectionModel* SelectionModel = nullptr;
 		
 		/** Combines the clients */
 		TSharedPtr<FMultiStreamModel> StreamModel;
@@ -54,5 +58,9 @@ namespace UE::MultiUserClient
 		// SClientToolbar attributes
 		TSet<FGuid> GetDisplayClientIds() const;
 		void EnumerateObjectsInStreams(TFunctionRef<void(const FSoftObjectPath&)> Consumer);
+		
+		void RebuildClientSubscriptions();
+		void CleanClientSubscriptions();
+		void OnClientChanged(FGuid Guid);
 	};
 }
