Create Blueprint API

diff --git a/Engine/Plugins/Developer/Concert/ConcertApp/MultiUserClient/MultiUserClient.uplugin b/Engine/Plugins/Developer/Concert/ConcertApp/MultiUserClient/MultiUserClient.uplugin
--- a/Engine/Plugins/Developer/Concert/ConcertApp/MultiUserClient/MultiUserClient.uplugin
+++ b/Engine/Plugins/Developer/Concert/ConcertApp/MultiUserClient/MultiUserClient.uplugin
@@ -33,17 +33,21 @@
 	],
 	"Plugins": [
 		{
-			"Name": "ConcertSyncClient",
+			"Name": "ConcertClientSharedSlate",
 			"Enabled": true
 		},
 		{
-			"Name": "ConcertClientSharedSlate",
+			"Name": "ConcertReplicationScripting",
 			"Enabled": true
 		},
         {
           "Name": "ConcertSharedSlate",
           "Enabled": true
         },
+		{
+			"Name": "ConcertSyncClient",
+			"Enabled": true
+		},
 		{
 			"Name": "MultiUserTakes",
 			"Enabled": true
diff --git a/Engine/Plugins/Developer/Concert/ConcertApp/MultiUserClient/Source/MultiUserClient/MultiUserClient.Build.cs b/Engine/Plugins/Developer/Concert/ConcertApp/MultiUserClient/Source/MultiUserClient/MultiUserClient.Build.cs
--- a/Engine/Plugins/Developer/Concert/ConcertApp/MultiUserClient/Source/MultiUserClient/MultiUserClient.Build.cs
+++ b/Engine/Plugins/Developer/Concert/ConcertApp/MultiUserClient/Source/MultiUserClient/MultiUserClient.Build.cs
@@ -37,6 +37,7 @@ namespace UnrealBuildTool.Rules
 					"Concert",
 					"ConcertClient",
 					"ConcertClientSharedSlate",
+					"ConcertReplicationScripting",
 					"ConcertSharedSlate",
 					"ConcertSyncClient",
 					"ConcertSyncCore",
diff --git a/Engine/Plugins/Developer/Concert/ConcertApp/MultiUserClient/Source/MultiUserClient/Private/MultiUserClientModule.cpp b/Engine/Plugins/Developer/Concert/ConcertApp/MultiUserClient/Source/MultiUserClient/Private/MultiUserClientModule.cpp
--- a/Engine/Plugins/Developer/Concert/ConcertApp/MultiUserClient/Source/MultiUserClient/Private/MultiUserClientModule.cpp
+++ b/Engine/Plugins/Developer/Concert/ConcertApp/MultiUserClient/Source/MultiUserClient/Private/MultiUserClientModule.cpp
@@ -648,6 +648,11 @@ public:
 		return MultiUserClient;
 	}
 
+	virtual UE::MultiUserClient::IMultiUserReplication* GetReplication() const override
+	{
+		return ReplicationManager.Get();
+	}
+
 	/**
 	 * Invokes the Multi-User browser tab
 	 */
diff --git a/Engine/Plugins/Developer/Concert/ConcertApp/MultiUserClient/Source/MultiUserClient/Private/Replication/Client/LocalReplicationClient.cpp b/Engine/Plugins/Developer/Concert/ConcertApp/MultiUserClient/Source/MultiUserClient/Private/Replication/Client/LocalReplicationClient.cpp
--- a/Engine/Plugins/Developer/Concert/ConcertApp/MultiUserClient/Source/MultiUserClient/Private/Replication/Client/LocalReplicationClient.cpp
+++ b/Engine/Plugins/Developer/Concert/ConcertApp/MultiUserClient/Source/MultiUserClient/Private/Replication/Client/LocalReplicationClient.cpp
@@ -9,6 +9,7 @@
 namespace UE::MultiUserClient
 {
 	FLocalReplicationClient::FLocalReplicationClient(
+		FReplicationDiscoveryContainer& InDiscoveryContainer,
 		FGlobalAuthorityCache& InAuthorityCache,
 		UMultiUserReplicationClientPreset& InSessionContent,
 		TUniquePtr<IClientStreamSynchronizer> InStreamSynchronizer,
@@ -16,6 +17,7 @@ namespace UE::MultiUserClient
 		)
 		: FReplicationClient(
 			InClient->GetConcertClient()->GetCurrentSession()->GetSessionClientEndpointId(),
+			InDiscoveryContainer,
 			InAuthorityCache,
 			InSessionContent,
 			MoveTemp(InStreamSynchronizer),
diff --git a/Engine/Plugins/Developer/Concert/ConcertApp/MultiUserClient/Source/MultiUserClient/Private/Replication/Client/LocalReplicationClient.h b/Engine/Plugins/Developer/Concert/ConcertApp/MultiUserClient/Source/MultiUserClient/Private/Replication/Client/LocalReplicationClient.h
--- a/Engine/Plugins/Developer/Concert/ConcertApp/MultiUserClient/Source/MultiUserClient/Private/Replication/Client/LocalReplicationClient.h
+++ b/Engine/Plugins/Developer/Concert/ConcertApp/MultiUserClient/Source/MultiUserClient/Private/Replication/Client/LocalReplicationClient.h
@@ -13,8 +13,16 @@ namespace UE::MultiUserClient
 	class FLocalReplicationClient : public FReplicationClient
 	{
 	public:
-
+		
+		/**
+		 * @param InDiscoveryContainer Used for auto-discovering properties added to this client's stream. The caller ensures it outlives the constructed instance.
+		 * @param InAuthorityCache Caches authority state of all clients. Passed to subsystems. The caller ensures it outlives the constructed instance.
+		 * @param InSessionContent Object that this client's stream changes are to be written into. The caller ensures it outlives the constructed instance.
+		 * @param InStreamSynchronizer Implementation for obtaining stream registered on server. The constructed instance takes ownership.
+		 * @param InClient The local editor's running client. Passed to subsystems so they can send messages to the current session. Keeps strong reference.
+		 */
 		FLocalReplicationClient(
+			FReplicationDiscoveryContainer& InDiscoveryContainer,
 			FGlobalAuthorityCache& InAuthorityCache,
 			UMultiUserReplicationClientPreset& InSessionContent,
 			TUniquePtr<IClientStreamSynchronizer> InStreamSynchronizer,
diff --git a/Engine/Plugins/Developer/Concert/ConcertApp/MultiUserClient/Source/MultiUserClient/Private/Replication/Client/RemoteReplicationClient.cpp b/Engine/Plugins/Developer/Concert/ConcertApp/MultiUserClient/Source/MultiUserClient/Private/Replication/Client/RemoteReplicationClient.cpp
--- a/Engine/Plugins/Developer/Concert/ConcertApp/MultiUserClient/Source/MultiUserClient/Private/Replication/Client/RemoteReplicationClient.cpp
+++ b/Engine/Plugins/Developer/Concert/ConcertApp/MultiUserClient/Source/MultiUserClient/Private/Replication/Client/RemoteReplicationClient.cpp
@@ -12,6 +12,7 @@ namespace UE::MultiUserClient
 {
 	FRemoteReplicationClient::FRemoteReplicationClient(
 		const FGuid& InConcertClientId,
+		FReplicationDiscoveryContainer& InDiscoveryContainer,
 		TSharedRef<IConcertClient> InClient,
 		FGlobalAuthorityCache& InAuthorityCache,
 		UMultiUserReplicationClientPreset& InSessionContent,
@@ -19,6 +20,7 @@ namespace UE::MultiUserClient
 		)
 		: FReplicationClient(
 			InConcertClientId,
+			InDiscoveryContainer,
 			InAuthorityCache,
 			InSessionContent,
 			MakeUnique<FStreamSynchronizer_RemoteClient>(InConcertClientId, QueryService),
diff --git a/Engine/Plugins/Developer/Concert/ConcertApp/MultiUserClient/Source/MultiUserClient/Private/Replication/Client/RemoteReplicationClient.h b/Engine/Plugins/Developer/Concert/ConcertApp/MultiUserClient/Source/MultiUserClient/Private/Replication/Client/RemoteReplicationClient.h
--- a/Engine/Plugins/Developer/Concert/ConcertApp/MultiUserClient/Source/MultiUserClient/Private/Replication/Client/RemoteReplicationClient.h
+++ b/Engine/Plugins/Developer/Concert/ConcertApp/MultiUserClient/Source/MultiUserClient/Private/Replication/Client/RemoteReplicationClient.h
@@ -14,9 +14,18 @@ namespace UE::MultiUserClient
 	class FRemoteReplicationClient : public FReplicationClient
 	{
 	public:
-
+		
+		/**
+		 * @param InConcertClientId The endpoint ID this instance corresponds to.
+		 * @param InDiscoveryContainer Used for auto-discovering properties added to this client's stream. The caller ensures it outlives the constructed instance.
+		 * @param InClient The local editor's running client. Passed to subsystems so they can send messages to the current session. Keeps strong reference.
+		 * @param InAuthorityCache Caches authority state of all clients. Passed to subsystems. The caller ensures it outlives the constructed instance.
+		 * @param InSessionContent Object that this client's stream changes are to be written into. The caller ensures it outlives the constructed instance.
+		 * @param QueryService Service shared across remote clients which periodically asks the server about this client's state. The caller ensures it outlives the constructed instance.
+		 */
 		FRemoteReplicationClient(
 			const FGuid& InConcertClientId,
+			FReplicationDiscoveryContainer& InDiscoveryContainer,
 			TSharedRef<IConcertClient> InClient,
 			FGlobalAuthorityCache& InAuthorityCache,
 			UMultiUserReplicationClientPreset& InSessionContent,
diff --git a/Engine/Plugins/Developer/Concert/ConcertApp/MultiUserClient/Source/MultiUserClient/Private/Replication/Client/ReplicationClient.cpp b/Engine/Plugins/Developer/Concert/ConcertApp/MultiUserClient/Source/MultiUserClient/Private/Replication/Client/ReplicationClient.cpp
--- a/Engine/Plugins/Developer/Concert/ConcertApp/MultiUserClient/Source/MultiUserClient/Private/Replication/Client/ReplicationClient.cpp
+++ b/Engine/Plugins/Developer/Concert/ConcertApp/MultiUserClient/Source/MultiUserClient/Private/Replication/Client/ReplicationClient.cpp
@@ -15,6 +15,7 @@ namespace UE::MultiUserClient
 {
 	FReplicationClient::FReplicationClient(
 		const FGuid& EndpointId,
+		FReplicationDiscoveryContainer& InDiscoveryContainer,
 		FGlobalAuthorityCache& InAuthorityCache,
 		UMultiUserReplicationClientPreset& InSessionContent,
 		TUniquePtr<IClientStreamSynchronizer> InStreamSynchronizer,
@@ -30,7 +31,7 @@ namespace UE::MultiUserClient
 		, LocalClientEditModel(CreatePropertySelectionModel(
 			*ClientContentStorage->Stream,
 			ClientContentStorage->Stream->MakeReplicationMapGetterAttribute(),
-			MakeShared<FMultiUserStreamExtender>())
+			MakeShared<FMultiUserStreamExtender>(EndpointId, InDiscoveryContainer))
 			)
 		, LocalClientStreamDiffer(
 			GetStreamSynchronizer(),
diff --git a/Engine/Plugins/Developer/Concert/ConcertApp/MultiUserClient/Source/MultiUserClient/Private/Replication/Client/ReplicationClient.h b/Engine/Plugins/Developer/Concert/ConcertApp/MultiUserClient/Source/MultiUserClient/Private/Replication/Client/ReplicationClient.h
--- a/Engine/Plugins/Developer/Concert/ConcertApp/MultiUserClient/Source/MultiUserClient/Private/Replication/Client/ReplicationClient.h
+++ b/Engine/Plugins/Developer/Concert/ConcertApp/MultiUserClient/Source/MultiUserClient/Private/Replication/Client/ReplicationClient.h
@@ -25,6 +25,7 @@ namespace UE::ConcertClientSharedSlate
 
 namespace UE::MultiUserClient
 {
+	class FReplicationDiscoveryContainer;
 	class ISubmissionWorkflow;
 	
 	/**
@@ -35,8 +36,18 @@ namespace UE::MultiUserClient
 	{
 	public:
 		
+		/**
+		 * @param EndpointId The endpoint ID this instance corresponds to.
+		 * @param InDiscoveryContainer Used for auto-discovering properties added to this client's stream. The caller ensures it outlives the constructed instance.
+		 * @param InAuthorityCache Caches authority state of all clients. Passed to subsystems. The caller ensures it outlives the constructed instance.
+		 * @param InSessionContent Object that this client's stream changes are to be written into. The caller ensures it outlives the constructed instance.
+		 * @param InStreamSynchronizer Implementation for obtaining stream registered on server. The constructed instance takes ownership.
+		 * @param InAuthoritySynchronizer Implementation for obtaining the client's authority state on server. The constructed instance takes ownership.
+		 * @param InSubmissionWorkflow Implementation for for changing client streams and authority on the server. The constructed instance takes ownership.
+		 */
 		FReplicationClient(
 			const FGuid& EndpointId,
+			FReplicationDiscoveryContainer& InDiscoveryContainer,
 			FGlobalAuthorityCache& InAuthorityCache,
 			UMultiUserReplicationClientPreset& InSessionContent,
 			TUniquePtr<IClientStreamSynchronizer> InStreamSynchronizer,
@@ -46,6 +57,14 @@ namespace UE::MultiUserClient
 		~FReplicationClient();
 
 		UMultiUserReplicationClientPreset* GetClientContent() const { return ClientContentStorage; }
+		/**
+		 * This is used so the UI can construct the IReplicationStreamEditor.
+		 * @see CreateBaseStreamEditor and FCreateEditorParams.
+		 * 
+		 * @note You must make sure to release this object when this FReplicationClient is destroyed.
+		 * Listen for events on the owning FReplicationClientManager::OnPreRemoteClientRemoved and FMultiUserReplicationManager::OnLeaveSession
+		 * @see FReplicationClientManager and FMultiUserReplicationManager
+		 */
 		TSharedRef<ConcertClientSharedSlate::IEditableReplicationStreamModel> GetClientEditModel() const { return LocalClientEditModel; }
 		IClientStreamSynchronizer& GetStreamSynchronizer() const { return *StreamSynchronizer.Get(); }
 		IClientAuthoritySynchronizer& GetAuthoritySynchronizer() const { return *AuthoritySynchronizer.Get(); }
@@ -99,9 +118,10 @@ namespace UE::MultiUserClient
 		 * Used to detect changes made to the client's config by the local editor.
 		 * Those changes can later be applied to the client.
 		 * 
-		 * The SessionContent UObject must be kept alive as long as this model lives. 
+		 * The SessionContent UObject must be kept alive as long as this model lives.
+		 * 
 		 * It should be noted that the UI keeps a strong reference to this model.
-		 * The UI is destroyed right after this FClientStreamRepository.
+		 * The UI is destroyed right before FReplicationClient, which is guaranteed by FMultiUserReplicationManager.
 		 * @see FMultiUserReplicationManager::OnLeaveSession
 		 */
 		TSharedRef<ConcertClientSharedSlate::IEditableReplicationStreamModel> LocalClientEditModel;
diff --git a/Engine/Plugins/Developer/Concert/ConcertApp/MultiUserClient/Source/MultiUserClient/Private/Replication/Client/ReplicationClientManager.cpp b/Engine/Plugins/Developer/Concert/ConcertApp/MultiUserClient/Source/MultiUserClient/Private/Replication/Client/ReplicationClientManager.cpp
--- a/Engine/Plugins/Developer/Concert/ConcertApp/MultiUserClient/Source/MultiUserClient/Private/Replication/Client/ReplicationClientManager.cpp
+++ b/Engine/Plugins/Developer/Concert/ConcertApp/MultiUserClient/Source/MultiUserClient/Private/Replication/Client/ReplicationClientManager.cpp
@@ -14,18 +14,20 @@
 namespace UE::MultiUserClient
 {
 	FReplicationClientManager::FReplicationClientManager(
-		TSharedRef<IConcertSyncClient> InClient,
-		TSharedRef<IConcertClientSession> InSession
+		const TSharedRef<IConcertSyncClient>& InClient,
+		const TSharedRef<IConcertClientSession>& InSession,
+		FReplicationDiscoveryContainer& InRegisteredExtenders
 		)
 		: SessionContent(NewObject<UMultiUserReplicationSessionPreset>(GetTransientPackage(), NAME_None, RF_Transient))
 		, ConcertClient(InClient)
 		, Session(InSession)
+		, RegisteredExtenders(InRegisteredExtenders)
 		, QueryService(InClient)
 		, AuthorityCache(*this)
 		, LocalClient([this, InClient]()
 		{
 			UMultiUserReplicationClientPreset* ClientPreset = SessionContent->AddClient();
-			return FLocalReplicationClient(AuthorityCache, *ClientPreset, MakeUnique<FStreamSynchronizer_LocalClient>(InClient, ClientPreset->Stream->StreamId), InClient);
+			return FLocalReplicationClient(RegisteredExtenders, AuthorityCache, *ClientPreset, MakeUnique<FStreamSynchronizer_LocalClient>(InClient, ClientPreset->Stream->StreamId), InClient);
 		}())
 		, SubmissionNotifier(*this)
 		, ReassignmentLogic(*this)
@@ -161,7 +163,14 @@ namespace UE::MultiUserClient
 	
 	void FReplicationClientManager::CreateRemoteClient(const FGuid& ClientEndpointId, bool bBroadcastDelegate)
 	{
-		TUniquePtr<FRemoteReplicationClient> RemoteClientPtr = MakeUnique<FRemoteReplicationClient>(ClientEndpointId, ConcertClient->GetConcertClient(), AuthorityCache, *SessionContent->AddClient(), QueryService);
+		TUniquePtr<FRemoteReplicationClient> RemoteClientPtr = MakeUnique<FRemoteReplicationClient>(
+			ClientEndpointId,
+			RegisteredExtenders,
+			ConcertClient->GetConcertClient(),
+			AuthorityCache,
+			*SessionContent->AddClient(),
+			QueryService
+			);
 		FRemoteReplicationClient& RemoteClient = *RemoteClientPtr;
 		RemoteClients.Emplace(
 			MoveTemp(RemoteClientPtr)
diff --git a/Engine/Plugins/Developer/Concert/ConcertApp/MultiUserClient/Source/MultiUserClient/Private/Replication/Client/ReplicationClientManager.h b/Engine/Plugins/Developer/Concert/ConcertApp/MultiUserClient/Source/MultiUserClient/Private/Replication/Client/ReplicationClientManager.h
--- a/Engine/Plugins/Developer/Concert/ConcertApp/MultiUserClient/Source/MultiUserClient/Private/Replication/Client/ReplicationClientManager.h
+++ b/Engine/Plugins/Developer/Concert/ConcertApp/MultiUserClient/Source/MultiUserClient/Private/Replication/Client/ReplicationClientManager.h
@@ -30,12 +30,15 @@ namespace UE::ConcertClientSharedSlate
 namespace UE::MultiUserClient
 {
 	class FRemoteReplicationClient;
+	class FReplicationDiscoveryContainer;
 	class FReplicationClient;
 	class FStreamChangeTracker;
 
 	/**
 	 * Keeps track of connected clients synchronizing their stream data in a UMultiUserReplicationSessionPreset.
-	 * This object only exists for as long as the local client is in a Concert session.
+	 * 
+	 * This class is instantiated only for as long as the local client is in a Concert session and is owned by FMultiUserReplicationManager,
+	 * which drives the lifetime.
 	 */
 	class FReplicationClientManager
 		: public FGCObject
@@ -44,10 +47,15 @@ namespace UE::MultiUserClient
 	public:
 		
 		/**
-		 * @param InClient The local client. Outlives this objects.
-		 * @param InSession The session to observe. Outlives this objects.
+		 * @param InClient The local client. The owning FMultiUserReplicationManager ensures it outlives the constructed instance.
+		 * @param InSession The session to observe. The owning FMultiUserReplicationManager ensures it outlives the constructed instance.
+		 * @param InRegisteredExtenders Used for auto-discovering properties added to this client's stream. Passed to the clients. The owning FMultiUserReplicationManager ensures it outlives the constructed instance.
 		 */
-		FReplicationClientManager(TSharedRef<IConcertSyncClient> InClient, TSharedRef<IConcertClientSession> InSession);
+		FReplicationClientManager(
+			const TSharedRef<IConcertSyncClient>& InClient,
+			const TSharedRef<IConcertClientSession>& InSession,
+			FReplicationDiscoveryContainer& InRegisteredExtenders
+			);
 		virtual ~FReplicationClientManager() override;
 
 		const FLocalReplicationClient& GetLocalClient() const { return LocalClient; }
@@ -107,12 +115,11 @@ namespace UE::MultiUserClient
 		
 	private:
 		
-		/** The state of the server is synched up with this object and displayed in the UI. */
+		/** The state of the server is synchronized up with this object and displayed in the UI. */
 		TObjectPtr<UMultiUserReplicationSessionPreset> SessionContent;
 
 		/** The local Concert client */
 		TSharedRef<IConcertSyncClient> ConcertClient;
-
 		/**
 		 * The session the local client is in.
 		 * 
@@ -120,6 +127,12 @@ namespace UE::MultiUserClient
 		 * when the session shuts down.
 		 */
 		const TWeakPtr<IConcertClientSession> Session;
+
+		/**
+		 * Passed to clients for constructing for the purposes of auto-discovering properties and additional objects when adding properties to the stream.
+		 * The owning FMultiUserReplicationManager ensures it outlives this FReplicationClientManager instance.
+		 */
+		FReplicationDiscoveryContainer& RegisteredExtenders;
 		
 		/**
 		 * Sends FConcertReplication_QueryReplicationInfo_Request in regular intervals.
diff --git a/Engine/Plugins/Developer/Concert/ConcertApp/MultiUserClient/Source/MultiUserClient/Private/Replication/MultiUserReplicationManager.cpp b/Engine/Plugins/Developer/Concert/ConcertApp/MultiUserClient/Source/MultiUserClient/Private/Replication/MultiUserReplicationManager.cpp
--- a/Engine/Plugins/Developer/Concert/ConcertApp/MultiUserClient/Source/MultiUserClient/Private/Replication/MultiUserReplicationManager.cpp
+++ b/Engine/Plugins/Developer/Concert/ConcertApp/MultiUserClient/Source/MultiUserClient/Private/Replication/MultiUserReplicationManager.cpp
@@ -84,7 +84,7 @@ namespace UE::MultiUserClient
 		const bool bSuccess = JoinSessionResult.ErrorCode == EJoinReplicationErrorCode::Success;
 		if (bSuccess)
 		{
-			ConnectedState.Emplace(Client);
+			ConnectedState.Emplace(Client, DiscoveryContainer);
 			SetConnectionStateAndBroadcast(EMultiUserReplicationConnectionState::Connected);
 		}
 		else
@@ -99,7 +99,44 @@ namespace UE::MultiUserClient
 		OnReplicationConnectionStateChangedDelegate.Broadcast(ConnectionState);
 	}
 
-	FMultiUserReplicationManager::FConnectedState::FConnectedState(TSharedRef<IConcertSyncClient> InClient)
-		: ClientManager(InClient, InClient->GetConcertClient()->GetCurrentSession().ToSharedRef())
+	FGuid FMultiUserReplicationManager::GetMultiUserStreamId() const
+	{
+		return UMultiUserReplicationClientPreset::MultiUserStreamID;
+	}
+
+	const FObjectReplicationMap* FMultiUserReplicationManager::FindReplicationMapForClient(const FGuid& ClientId) const
+	{
+		if (ConnectedState)
+		{
+			const FReplicationClient* ReplicationClient = ConnectedState->ClientManager.FindClient(ClientId);
+			return ReplicationClient
+				? &ReplicationClient->GetStreamSynchronizer().GetServerState()
+				: nullptr;
+		}
+		return nullptr;
+	}
+
+	bool FMultiUserReplicationManager::IsReplicatingObject(const FGuid& ClientId, const FSoftObjectPath& ObjectPath) const 
+	{
+		if (ConnectedState)
+		{
+			const FReplicationClient* ReplicationClient = ConnectedState->ClientManager.FindClient(ClientId);
+			return ReplicationClient && ReplicationClient->GetAuthoritySynchronizer().HasAuthorityOver(ObjectPath);
+		}
+		return false;
+	}
+
+	void FMultiUserReplicationManager::RegisterReplicationDiscoverer(TSharedRef<IReplicationDiscoverer> Discoverer)
+	{
+		DiscoveryContainer.AddDiscoverer(Discoverer);
+	}
+
+	void FMultiUserReplicationManager::RemoveReplicationDiscoverer(const TSharedRef<IReplicationDiscoverer>& Discoverer)
+	{
+		DiscoveryContainer.RemoveDiscoverer(Discoverer);
+	}
+
+	FMultiUserReplicationManager::FConnectedState::FConnectedState(TSharedRef<IConcertSyncClient> InClient, FReplicationDiscoveryContainer& InDiscoveryContainer)
+		: ClientManager(InClient, InClient->GetConcertClient()->GetCurrentSession().ToSharedRef(), InDiscoveryContainer)
 	{}
 }
diff --git a/Engine/Plugins/Developer/Concert/ConcertApp/MultiUserClient/Source/MultiUserClient/Private/Replication/MultiUserReplicationManager.h b/Engine/Plugins/Developer/Concert/ConcertApp/MultiUserClient/Source/MultiUserClient/Private/Replication/MultiUserReplicationManager.h
--- a/Engine/Plugins/Developer/Concert/ConcertApp/MultiUserClient/Source/MultiUserClient/Private/Replication/MultiUserReplicationManager.h
+++ b/Engine/Plugins/Developer/Concert/ConcertApp/MultiUserClient/Source/MultiUserClient/Private/Replication/MultiUserReplicationManager.h
@@ -5,6 +5,9 @@
 #include "Client/ReplicationClientManager.h"
 
 #include "IConcertSession.h"
+#include "Replication/IMultiUserReplication.h"
+#include "Replication/Stream/Discovery/ReplicationDiscoveryContainer.h"
+
 #include "Misc/Optional.h"
 #include "Templates/SharedPointer.h"
 #include "Templates/UnrealTemplate.h"
@@ -46,11 +49,12 @@ namespace UE::MultiUserClient
 	class FMultiUserReplicationManager
 		: public TSharedFromThis<FMultiUserReplicationManager>
 		, public FNoncopyable
+		, public IMultiUserReplication
 	{
 	public:
 		
 		FMultiUserReplicationManager(TSharedRef<IConcertSyncClient> InClient);
-		~FMultiUserReplicationManager();
+		virtual ~FMultiUserReplicationManager() override;
 
 		/**
 		 * Joins the replication session.
@@ -70,6 +74,14 @@ namespace UE::MultiUserClient
 		FOnReplicationConnectionStateChanged& OnReplicationConnectionStateChanged() { return OnReplicationConnectionStateChangedDelegate; }
 		EMultiUserReplicationConnectionState GetConnectionState() const { return ConnectionState; }
 
+		//~ Begin IMultiUserReplication Interface
+		virtual FGuid GetMultiUserStreamId() const override;
+		virtual const FObjectReplicationMap* FindReplicationMapForClient(const FGuid& ClientId) const override;
+		virtual bool IsReplicatingObject(const FGuid& ClientId, const FSoftObjectPath& ObjectPath) const override;
+		virtual void RegisterReplicationDiscoverer(TSharedRef<IReplicationDiscoverer> Discoverer) override;
+		virtual void RemoveReplicationDiscoverer(const TSharedRef<IReplicationDiscoverer>& Discoverer) override;
+		//~ End IMultiUserReplication Interface
+
 	private:
 
 		/** Client through which the replication bridge is accessed. */
@@ -91,11 +103,14 @@ namespace UE::MultiUserClient
 			 */
 			FReplicationClientManager ClientManager;
 			
-			FConnectedState(TSharedRef<IConcertSyncClient> InClient);
+			FConnectedState(TSharedRef<IConcertSyncClient> InClient, FReplicationDiscoveryContainer& InDiscoveryContainer);
 		};
 		/** Set when connected to a replication session. */
 		TOptional<FConnectedState> ConnectedState;
 
+		/** Allows external modules to register discoverers for adding properties, etc. */
+		FReplicationDiscoveryContainer DiscoveryContainer;
+
 		/** Called when ConnectionState changes. */
 		FOnReplicationConnectionStateChanged OnReplicationConnectionStateChangedDelegate;
 
diff --git a/Engine/Plugins/Developer/Concert/ConcertApp/MultiUserClient/Source/MultiUserClient/Private/Replication/Stream/Discovery/MultiUserStreamExtender.cpp b/Engine/Plugins/Developer/Concert/ConcertApp/MultiUserClient/Source/MultiUserClient/Private/Replication/Stream/Discovery/MultiUserStreamExtender.cpp
--- a/Engine/Plugins/Developer/Concert/ConcertApp/MultiUserClient/Source/MultiUserClient/Private/Replication/Stream/Discovery/MultiUserStreamExtender.cpp
+++ b/Engine/Plugins/Developer/Concert/ConcertApp/MultiUserClient/Source/MultiUserClient/Private/Replication/Stream/Discovery/MultiUserStreamExtender.cpp
@@ -3,20 +3,52 @@
 #include "MultiUserStreamExtender.h"
 
 #include "MultiUserReplicationSettings.h"
+#include "ReplicationDiscoveryContainer.h"
+#include "Replication/IReplicationDiscoveryContext.h"
+#include "Replication/Data/ConcertPropertySelection.h"
+#include "Replication/Editor/Model/Extension/IStreamExtensionContext.h"
 
 namespace UE::MultiUserClient
 {
-	FMultiUserStreamExtender::FMultiUserStreamExtender()
-		: ExtendBySettings(
+	FMultiUserStreamExtender::FMultiUserStreamExtender(const FGuid& InClientId, FReplicationDiscoveryContainer& InRegisteredExtenders)
+		: ClientId(InClientId)
+		, ExtendBySettings(
 			// Use MU settings for auto adding properties & objects
 			TAttribute<const FConcertReplicationEditorSettings*>::CreateLambda([]()
 			{
 				return &UMultiUserReplicationSettings::Get()->ReplicationEditorSettings;
 			}))
+		, RegisteredExtenders(InRegisteredExtenders)
 	{}
 
 	void FMultiUserStreamExtender::ExtendStream(UObject& ExtendedObject, ConcertClientSharedSlate::IStreamExtensionContext& Context)
 	{
 		ExtendBySettings.ExtendStream(ExtendedObject, Context);
+		ExtendStreamWithRegisteredDiscoverers(ExtendedObject, Context);
+	}
+	
+	void FMultiUserStreamExtender::ExtendStreamWithRegisteredDiscoverers(UObject& ExtendedObject, ConcertClientSharedSlate::IStreamExtensionContext& Context) const
+	{
+		class FDiscoveryContext : public IReplicationDiscoveryContext
+		{
+			ConcertClientSharedSlate::IStreamExtensionContext& Context;
+		public:
+			
+			explicit FDiscoveryContext(ConcertClientSharedSlate::IStreamExtensionContext& InContext)
+				: Context(InContext)
+			{}
+
+			virtual void AddPropertyTo(UObject& Object, FConcertPropertyChain PropertyChain) override
+			{
+				Context.AddPropertyTo(Object, MoveTemp(PropertyChain));
+			}
+			virtual void AddAdditionalObject(UObject& Object) override
+			{
+				Context.AddAdditionalObject(Object);
+			}
+		};
+		
+		FDiscoveryContext MultiUserContext(Context);
+		RegisteredExtenders.DiscoverReplicationSettings({ ExtendedObject, ClientId, MultiUserContext });
 	}
 }
diff --git a/Engine/Plugins/Developer/Concert/ConcertApp/MultiUserClient/Source/MultiUserClient/Private/Replication/Stream/Discovery/MultiUserStreamExtender.h b/Engine/Plugins/Developer/Concert/ConcertApp/MultiUserClient/Source/MultiUserClient/Private/Replication/Stream/Discovery/MultiUserStreamExtender.h
--- a/Engine/Plugins/Developer/Concert/ConcertApp/MultiUserClient/Source/MultiUserClient/Private/Replication/Stream/Discovery/MultiUserStreamExtender.h
+++ b/Engine/Plugins/Developer/Concert/ConcertApp/MultiUserClient/Source/MultiUserClient/Private/Replication/Stream/Discovery/MultiUserStreamExtender.h
@@ -5,6 +5,8 @@
 #include "Replication/Editor/Model/Extension/IStreamExtender.h"
 #include "Replication/Editor/Model/StreamExtenderBySettings.h"
 
+#include "Misc/Guid.h"
+
 namespace UE::MultiUserClient
 {
 	class FReplicationDiscoveryContainer;
@@ -13,13 +15,19 @@ namespace UE::MultiUserClient
 	 * When the user adds an object, this object handles auto selecting properties and adding additional objects from context.
 	 * The following sources exist:
 	 * - Static Settings: user can specify properties & objects in the MU project settings
-	 * - Dynamic API: by using the IMultiUserReplication, external modules can register dynamic rules for auto-discovery.
+	 * - Dynamic API: by using the IMultiUserReplication, external modules can register dynamic rules for auto-discovery using IReplicationDiscoverer.
+	 *
+	 * Every FReplicationClient owns an instance of this and executes it when an object is added for the client by the local editor.
 	 */
 	class FMultiUserStreamExtender : public ConcertClientSharedSlate::IStreamExtender
 	{
 	public:
-
-		FMultiUserStreamExtender();
+		
+		/**
+		 * @param InClientId The ID of the owning FReplicationClient instance.
+		 * @param InRegisteredExtenders Holds all IReplicationDiscoverer that were registered to Multi User. The caller ensures it outlives the constructed instance.
+		 */
+		FMultiUserStreamExtender(const FGuid& InClientId, FReplicationDiscoveryContainer& InRegisteredExtenders);
 
 		//~ Begin IStreamExtender Interface
 		virtual void ExtendStream(UObject& ExtendedObject, ConcertClientSharedSlate::IStreamExtensionContext& Context) override;
@@ -27,8 +35,16 @@ namespace UE::MultiUserClient
 
 	private:
 
+		/** The client ID to which this extender is adds objects. */
+		const FGuid ClientId;
+
 		/** Handles properties from the MU settings. */
 		ConcertClientSharedSlate::FStreamExtenderBySettings ExtendBySettings;
+
+		/** Allows generic extenders to extend added objects. These can be added via the IMultiUserReplication interface. */
+		FReplicationDiscoveryContainer& RegisteredExtenders;
+		
+		void ExtendStreamWithRegisteredDiscoverers(UObject& ExtendedObject, ConcertClientSharedSlate::IStreamExtensionContext& Context) const;
 	};
 }
 
diff --git a/Engine/Plugins/Developer/Concert/ConcertApp/MultiUserClient/Source/MultiUserClient/Private/Replication/Stream/Discovery/ReplicationDiscoveryContainer.cpp b/Engine/Plugins/Developer/Concert/ConcertApp/MultiUserClient/Source/MultiUserClient/Private/Replication/Stream/Discovery/ReplicationDiscoveryContainer.cpp
new file mode 100644
--- /dev/null
+++ b/Engine/Plugins/Developer/Concert/ConcertApp/MultiUserClient/Source/MultiUserClient/Private/Replication/Stream/Discovery/ReplicationDiscoveryContainer.cpp
@@ -0,0 +1,27 @@
+// Copyright Epic Games, Inc. All Rights Reserved.
+
+#include "ReplicationDiscoveryContainer.h"
+
+namespace UE::MultiUserClient
+{
+	void FReplicationDiscoveryContainer::AddDiscoverer(TSharedRef<IReplicationDiscoverer> Discoverer)
+	{
+		Discoverers.AddUnique(MoveTemp(Discoverer));
+	}
+
+	void FReplicationDiscoveryContainer::RemoveDiscoverer(const TSharedRef<IReplicationDiscoverer>& Discoverer)
+	{
+		Discoverers.RemoveSingle(Discoverer);
+	}
+
+	void FReplicationDiscoveryContainer::DiscoverReplicationSettings(const FReplicationDiscoveryParams& Params)
+	{
+		// Iterate in reverse order so we ignore any newly added ones
+		for (int32 Index = Discoverers.Num() - 1; Discoverers.IsValidIndex(Index); --Index)
+		{
+			const TSharedRef<IReplicationDiscoverer>& Discoverer = Discoverers[Index];
+			Discoverer->DiscoverReplicationSettings(Params);
+			// if something was removed before i, we will visit them twice... which does not matter because adding properties is a no-op if they were already added.
+		}
+	}
+}
diff --git a/Engine/Plugins/Developer/Concert/ConcertApp/MultiUserClient/Source/MultiUserClient/Private/Replication/Stream/Discovery/ReplicationDiscoveryContainer.h b/Engine/Plugins/Developer/Concert/ConcertApp/MultiUserClient/Source/MultiUserClient/Private/Replication/Stream/Discovery/ReplicationDiscoveryContainer.h
new file mode 100644
--- /dev/null
+++ b/Engine/Plugins/Developer/Concert/ConcertApp/MultiUserClient/Source/MultiUserClient/Private/Replication/Stream/Discovery/ReplicationDiscoveryContainer.h
@@ -0,0 +1,36 @@
+// Copyright Epic Games, Inc. All Rights Reserved.
+
+#pragma once
+
+#include "Replication/IReplicationDiscoverer.h"
+
+#include "Containers/Array.h"
+#include "Templates/Function.h"
+#include "Templates/SharedPointer.h"
+
+class UObject;
+struct FConcertPropertyChain;
+
+namespace UE::MultiUserClient
+{
+	/**
+	 * Just holds a bunch of IReplicationDiscoverer and calls forward the function calls to them.
+	 * @see IReplicationDiscovery for the purpose of IReplicationDiscoverer.
+	 */
+	class FReplicationDiscoveryContainer : public IReplicationDiscoverer
+	{
+	public:
+
+		void AddDiscoverer(TSharedRef<IReplicationDiscoverer> Discoverer);
+		void RemoveDiscoverer(const TSharedRef<IReplicationDiscoverer>& Discoverer);
+		
+		//~ Begin IReplicationDiscovery Interface
+		virtual void DiscoverReplicationSettings(const FReplicationDiscoveryParams& Params) override;
+		//~ End IReplicationDiscovery Interface
+
+	private:
+
+		/** The discoverers to run through */
+		TArray<TSharedRef<IReplicationDiscoverer>> Discoverers;
+	};
+}
diff --git a/Engine/Plugins/Developer/Concert/ConcertApp/MultiUserClient/Source/MultiUserClient/Public/IMultiUserClientModule.h b/Engine/Plugins/Developer/Concert/ConcertApp/MultiUserClient/Source/MultiUserClient/Public/IMultiUserClientModule.h
--- a/Engine/Plugins/Developer/Concert/ConcertApp/MultiUserClient/Source/MultiUserClient/Public/IMultiUserClientModule.h
+++ b/Engine/Plugins/Developer/Concert/ConcertApp/MultiUserClient/Source/MultiUserClient/Public/IMultiUserClientModule.h
@@ -7,6 +7,8 @@
 #include "Modules/ModuleInterface.h"
 #include "Modules/ModuleManager.h"
 
+
+
 class IConcertSyncClient;
 
 struct FProcHandle;
@@ -17,6 +19,11 @@ struct FServerLaunchOverrides
 	FString ServerName;
 };
 
+namespace UE::MultiUserClient
+{
+	class IMultiUserReplication;
+}
+
 /**
  * Interface for the Multi-User module.
  */
@@ -53,6 +60,12 @@ public:
 	 */
 	virtual TSharedPtr<IConcertSyncClient> GetClient() const = 0;
 
+	/**
+	 * Gets the interface for interacting with the replication system that Multi User sets-up.
+	 * Ownership remains in MultiUser. You cannot keep any reference to it: Call this function instead.
+	 */
+	virtual UE::MultiUserClient::IMultiUserReplication* GetReplication() const = 0;
+
 	/**
 	 * Invokes the Multi-User browser tab
 	 */
diff --git a/Engine/Plugins/Developer/Concert/ConcertApp/MultiUserClient/Source/MultiUserClient/Public/Replication/IMultiUserReplication.h b/Engine/Plugins/Developer/Concert/ConcertApp/MultiUserClient/Source/MultiUserClient/Public/Replication/IMultiUserReplication.h
new file mode 100644
--- /dev/null
+++ b/Engine/Plugins/Developer/Concert/ConcertApp/MultiUserClient/Source/MultiUserClient/Public/Replication/IMultiUserReplication.h
@@ -0,0 +1,48 @@
+// Copyright Epic Games, Inc. All Rights Reserved.
+
+#pragma once
+
+#include "Templates/SharedPointer.h"
+
+struct FGuid;
+struct FObjectReplicationMap;
+
+namespace UE::MultiUserClient
+{
+	class IReplicationDiscoverer;
+	
+	/** Interface for interacting with Multi-User replication, which uses the Concert replication system. */
+	class MULTIUSERCLIENT_API IMultiUserReplication
+	{
+	public:
+
+		/**
+		 * The server keeps track of multiple streams per client and the ID is used to identify streams by client.
+		 * MU only uses one stream per client and this is its ID.
+		 * 
+		 * @return The stream ID that every MU client uses to register its replication stream with the server.
+		 */
+		virtual FGuid GetMultiUserStreamId() const = 0;
+
+		/**
+		 * @return Gets the last known server map of objects registered for replication for a given client.
+		 * This server state is regularly polled whilst the local client state should always be in synch.*/
+		virtual const FObjectReplicationMap* FindReplicationMapForClient(const FGuid& ClientId) const = 0;
+		/** @return Whether the local editor instance thinks the client has authority over the properties it has registered to ObjectPath. */
+		virtual bool IsReplicatingObject(const FGuid& ClientId, const FSoftObjectPath& ObjectPath) const = 0;
+
+		/**
+		 * Register a discoverer.
+		 *
+		 * It is used to automatically configure UObjects for replication when appropriate:
+		 * - When an user adds an object via Add Actor button
+		 * - When an UObject is added to the world via a transaction (run on the client machine that adds the UObject)
+		 */
+		virtual void RegisterReplicationDiscoverer(TSharedRef<IReplicationDiscoverer> Discoverer) = 0;
+		/** Unregisters a previously registered discoverer */
+		virtual void RemoveReplicationDiscoverer(const TSharedRef<IReplicationDiscoverer>& Discoverer) = 0;
+
+		virtual ~IMultiUserReplication() = default;
+	};
+}
+
diff --git a/Engine/Plugins/Developer/Concert/ConcertApp/MultiUserClient/Source/MultiUserClient/Public/Replication/IReplicationDiscoverer.h b/Engine/Plugins/Developer/Concert/ConcertApp/MultiUserClient/Source/MultiUserClient/Public/Replication/IReplicationDiscoverer.h
new file mode 100644
--- /dev/null
+++ b/Engine/Plugins/Developer/Concert/ConcertApp/MultiUserClient/Source/MultiUserClient/Public/Replication/IReplicationDiscoverer.h
@@ -0,0 +1,51 @@
+// Copyright Epic Games, Inc. All Rights Reserved.
+
+#pragma once
+
+#include "Templates/Function.h"
+
+class UObject;
+struct FConcertPropertyChain;
+
+namespace UE::MultiUserClient
+{
+	class IReplicationDiscoveryContext;
+
+	struct FReplicationDiscoveryParams
+	{
+		/** The object that was added to the stream and for which to discover additional settings. */
+		UObject& ExtendedObject;
+		/** The endpoint ID of the client to whose stream the object is registered */
+		const FGuid& EndpointId;
+		/** Mutate replication settings using this context object. */
+		IReplicationDiscoveryContext& Context;
+	};
+	
+	/**
+	 * An interface for setting up default replication settings for an object.
+	 *
+	 * This is used when Multi-User determines that an object should be automatically owned by a client and e.g. wants to set up a common workflow.
+	 * For example, if you have a specific actor type for which a user usually wants to replicate the position & rotation, you can use this to set up
+	 * the workflow automatically for the user.
+	 *
+	 * Situations where this is executed:
+	 * - When a user adds the object to a replication stream using the "Add Actor" button
+	 * - When the user adds an actor to the world while in a session, they might automatically want to replicate it (TODO UE-202194)
+	 *
+	 * This interface is similar to IStreamExtensionContext but the interfaces serve different purposes.
+	 * IStreamExtensionContext concerns itself with extending general streams. IReplicationDiscoveryContext is specific to Multi User.
+	 * For example, this interface exposes endpoint IDs for which does not make sense for IStreamExtensionContext.
+	 */
+	class MULTIUSERCLIENT_API IReplicationDiscoverer
+	{
+	public:
+
+		/**
+		 * Finds replication settings given the ExtendedObject.
+		 * Using Context, you can assign properties or add additional objects.
+		 */
+		virtual void DiscoverReplicationSettings(const FReplicationDiscoveryParams& Params) = 0;
+
+		virtual ~IReplicationDiscoverer() = default;
+	};
+}
diff --git a/Engine/Plugins/Developer/Concert/ConcertApp/MultiUserClient/Source/MultiUserClient/Public/Replication/IReplicationDiscoveryContext.h b/Engine/Plugins/Developer/Concert/ConcertApp/MultiUserClient/Source/MultiUserClient/Public/Replication/IReplicationDiscoveryContext.h
new file mode 100644
--- /dev/null
+++ b/Engine/Plugins/Developer/Concert/ConcertApp/MultiUserClient/Source/MultiUserClient/Public/Replication/IReplicationDiscoveryContext.h
@@ -0,0 +1,44 @@
+// Copyright Epic Games, Inc. All Rights Reserved.
+
+#pragma once
+
+#include "Templates/Function.h"
+
+class UObject;
+struct FConcertPropertyChain;
+
+namespace UE::MultiUserClient
+{
+	/**
+	 * An interface for discovering replication settings for an object.
+	 *
+	 * This interface is very similar to IStreamExtensionContext but the interfaces serve different purposes.
+	 * IStreamExtensionContext concerns itself with extending general streams. IReplicationDiscoveryContext is specific to Multi User.
+	 * 
+	 * @see IReplicationDiscoverer
+	 */
+	class MULTIUSERCLIENT_API IReplicationDiscoveryContext
+	{
+	public:
+
+		/**
+		 * Extends the stream with the given object, if not yet added, and adds the property to it.
+		 * If Object was not yet in the stream, IStreamExtender::ExtendStream will be called for Object.
+		 * 
+		 * @param Object Object to add, if not yet present, and assign a property to
+		 * @param PropertyChain Property chain to associate with the object
+		 */
+		virtual void AddPropertyTo(UObject& Object, FConcertPropertyChain PropertyChain) = 0;
+
+		/**
+		 * Adds an additional object.
+		 * 
+		 * If Object was not yet in the stream, IStreamExtender::ExtendStream will be called for Object.
+		 * 
+		 * @param Object The additional object to add
+		 */
+		virtual void AddAdditionalObject(UObject& Object) = 0;
+
+		virtual ~IReplicationDiscoveryContext() = default;
+	};
+}
diff --git a/Engine/Plugins/Developer/Concert/ConcertApp/MultiUserClient/Source/MultiUserClientLibrary/MultiUserClientLibrary.Build.cs b/Engine/Plugins/Developer/Concert/ConcertApp/MultiUserClient/Source/MultiUserClientLibrary/MultiUserClientLibrary.Build.cs
--- a/Engine/Plugins/Developer/Concert/ConcertApp/MultiUserClient/Source/MultiUserClientLibrary/MultiUserClientLibrary.Build.cs
+++ b/Engine/Plugins/Developer/Concert/ConcertApp/MultiUserClient/Source/MultiUserClientLibrary/MultiUserClientLibrary.Build.cs
@@ -9,9 +9,13 @@ namespace UnrealBuildTool.Rules
 			PublicDependencyModuleNames.AddRange(
 				new string[]
 				{
+					// Engine
 					"Core",
 					"CoreUObject",
-					"Engine"
+					"Engine",
+					
+					// Concert
+					"ConcertReplicationScripting",
 				}
 			);
 
@@ -25,7 +29,8 @@ namespace UnrealBuildTool.Rules
 						"Concert",
 						"ConcertClient",
 						"ConcertSyncCore",
-						"ConcertTransport"
+						"ConcertTransport",
+						"MultiUserClient"
 					}
 				);
 
@@ -34,15 +39,7 @@ namespace UnrealBuildTool.Rules
 					{
 						"Concert",
 						"ConcertSyncCore",
-						"ConcertSyncClient",
-						"MultiUserClient",
-					}
-				);
-
-				DynamicallyLoadedModuleNames.AddRange(
-					new string[]
-					{
-						"MultiUserClient",
+						"ConcertSyncClient"
 					}
 				);
 			}
diff --git a/Engine/Plugins/Developer/Concert/ConcertApp/MultiUserClient/Source/MultiUserClientLibrary/Private/MultiUserSubsystem.cpp b/Engine/Plugins/Developer/Concert/ConcertApp/MultiUserClient/Source/MultiUserClientLibrary/Private/MultiUserSubsystem.cpp
--- a/Engine/Plugins/Developer/Concert/ConcertApp/MultiUserClient/Source/MultiUserClientLibrary/Private/MultiUserSubsystem.cpp
+++ b/Engine/Plugins/Developer/Concert/ConcertApp/MultiUserClient/Source/MultiUserClientLibrary/Private/MultiUserSubsystem.cpp
@@ -235,6 +235,35 @@ bool UMultiUserSubsystem::IsConnectedToSession() const
 	return false;
 }
 
+bool UMultiUserSubsystem::GetLocalClientId(FGuid& OutClientId) const
+{
+#if WITH_CONCERT
+	using namespace UE::MultiUserSubsystem::Private;
+	if (ConcertManager.IsValid() && ConcertManager->WeakSession.IsValid())
+	{
+		const TSharedPtr<IConcertClientSession> SessionPin = ConcertManager->WeakSession.Pin();
+		OutClientId = SessionPin->GetSessionClientEndpointId();
+		return true;
+	}
+#endif
+	return false;
+}
+
+bool UMultiUserSubsystem::GetRemoteClientIds(TArray<FGuid>& OutRemoteClientIds)
+{
+#if WITH_CONCERT
+	using namespace UE::MultiUserSubsystem::Private;
+	if (ConcertManager.IsValid() && ConcertManager->WeakSession.IsValid())
+	{
+		const TSharedPtr<IConcertClientSession> SessionPin = ConcertManager->WeakSession.Pin();
+		OutRemoteClientIds = SessionPin->GetSessionClientEndpointIds();
+		OutRemoteClientIds.RemoveSingle(SessionPin->GetSessionClientEndpointId());
+		return true;
+	}
+#endif
+	return false;
+}
+
 void UMultiUserSubsystem::SendCustomEvent(const FConcertBlueprintEvent& EventData)
 {
 #if WITH_CONCERT
diff --git a/Engine/Plugins/Developer/Concert/ConcertApp/MultiUserClient/Source/MultiUserClientLibrary/Private/Replication/MultiUserReplicationRegistrationContextImpl.cpp b/Engine/Plugins/Developer/Concert/ConcertApp/MultiUserClient/Source/MultiUserClientLibrary/Private/Replication/MultiUserReplicationRegistrationContextImpl.cpp
new file mode 100644
--- /dev/null
+++ b/Engine/Plugins/Developer/Concert/ConcertApp/MultiUserClient/Source/MultiUserClientLibrary/Private/Replication/MultiUserReplicationRegistrationContextImpl.cpp
@@ -0,0 +1,33 @@
+// Copyright Epic Games, Inc. All Rights Reserved.
+
+#include "MultiUserReplicationRegistrationContextImpl.h"
+
+#if WITH_CONCERT
+#include "ConcertPropertyChainWrapper.h"
+#include "Replication/IReplicationDiscoveryContext.h"
+#include "Replication/Data/ConcertPropertySelection.h"
+#endif
+
+void UMultiUserReplicationRegistrationContextImpl::AddPropertiesToObject(UObject* Object, const TArray<FConcertPropertyChainWrapper>& PropertiesToAdd)
+{
+#if WITH_CONCERT
+	if (Object && ensure(NativeContext))
+	{
+		for (const FConcertPropertyChainWrapper& Wrapper : PropertiesToAdd)
+		{
+			FConcertPropertyChain Property = Wrapper.PropertyChain;
+			NativeContext->AddPropertyTo(*Object, MoveTemp(Property));
+		}
+	}
+#endif
+}
+
+void UMultiUserReplicationRegistrationContextImpl::AddAdditionalObject(UObject* Object)
+{
+#if WITH_CONCERT
+	if (Object && ensure(NativeContext))
+	{
+		NativeContext->AddAdditionalObject(*Object);
+	}
+#endif
+}
diff --git a/Engine/Plugins/Developer/Concert/ConcertApp/MultiUserClient/Source/MultiUserClientLibrary/Private/Replication/MultiUserReplicationRegistrationContextImpl.h b/Engine/Plugins/Developer/Concert/ConcertApp/MultiUserClient/Source/MultiUserClientLibrary/Private/Replication/MultiUserReplicationRegistrationContextImpl.h
new file mode 100644
--- /dev/null
+++ b/Engine/Plugins/Developer/Concert/ConcertApp/MultiUserClient/Source/MultiUserClientLibrary/Private/Replication/MultiUserReplicationRegistrationContextImpl.h
@@ -0,0 +1,37 @@
+// Copyright Epic Games, Inc. All Rights Reserved.
+
+#pragma once
+
+#include "Replication/IMultiUserReplicationRegistrationContext.h"
+#include "UObject/Object.h"
+#include "MultiUserReplicationRegistrationContextImpl.generated.h"
+
+namespace UE::MultiUserClient
+{
+	class IReplicationDiscoveryContext;
+}
+
+/**
+ * Instantiated once and passed to IMultiUserReplicationRegistration.
+ * Needed because TScriptInterface, which is passed via FMultiUserReplicationRegistrationParams, requires an UObject to bind to.
+ * 
+ * @see FUObjectAdapterReplicationDiscoverer
+ */
+UCLASS()
+class UMultiUserReplicationRegistrationContextImpl
+	: public UObject
+	, public IMultiUserReplicationRegistrationContext
+{
+	GENERATED_BODY()
+public:
+
+	/** Set by FUObjectAdapterReplicationDiscoverer for the duration of the registation */
+	UE::MultiUserClient::IReplicationDiscoveryContext* NativeContext = nullptr;
+
+	//~ Begin IMultiUserReplicationRegistrationContext Interface
+	UFUNCTION(BlueprintCallable, Category = "Multi-user")
+	virtual void AddPropertiesToObject(UObject* Object, const TArray<FConcertPropertyChainWrapper>& PropertiesToAdd) override;
+	UFUNCTION(BlueprintCallable, Category = "Multi-user")
+	virtual void AddAdditionalObject(UObject* Object) override;
+	//~ End IMultiUserReplicationRegistrationContext Interface
+};
diff --git a/Engine/Plugins/Developer/Concert/ConcertApp/MultiUserClient/Source/MultiUserClientLibrary/Private/Replication/MultiUserReplicationSubsystem.cpp b/Engine/Plugins/Developer/Concert/ConcertApp/MultiUserClient/Source/MultiUserClientLibrary/Private/Replication/MultiUserReplicationSubsystem.cpp
new file mode 100644
--- /dev/null
+++ b/Engine/Plugins/Developer/Concert/ConcertApp/MultiUserClient/Source/MultiUserClientLibrary/Private/Replication/MultiUserReplicationSubsystem.cpp
@@ -0,0 +1,130 @@
+// Copyright Epic Games, Inc. All Rights Reserved.
+
+#include "Replication/MultiUserReplicationSubsystem.h"
+
+#if WITH_CONCERT
+#include "UObjectAdapterReplicationDiscoverer.h"
+#include "IMultiUserClientModule.h"
+#include "Replication/Data/ObjectReplicationMap.h"
+#include "Replication/IMultiUserReplication.h"
+
+#include "Algo/Transform.h"
+#endif
+
+bool UMultiUserReplicationSubsystem::IsReplicatingObject(const FGuid& ClientId, const FSoftObjectPath& ObjectPath) const
+{
+#if WITH_CONCERT
+	const UE::MultiUserClient::IMultiUserReplication* ReplicationInterface = IMultiUserClientModule::Get().GetReplication();
+	if (ensureMsgf(ReplicationInterface, TEXT("We expected it to always be valid.")))
+	{
+		return ReplicationInterface->IsReplicatingObject(ClientId, ObjectPath);
+	}
+#endif
+	return false;
+}
+
+TArray<FConcertPropertyChainWrapper> UMultiUserReplicationSubsystem::GetPropertiesRegisteredToObject(const FGuid& ClientId, const FSoftObjectPath& ObjectPath) const
+{
+#if WITH_CONCERT  
+	const UE::MultiUserClient::IMultiUserReplication* ReplicationInterface = IMultiUserClientModule::Get().GetReplication();
+	if (ensureMsgf(ReplicationInterface, TEXT("We expected it to always be valid.")))
+	{
+		const FObjectReplicationMap* ObjectMap = ReplicationInterface->FindReplicationMapForClient(ClientId);
+		if (!ObjectMap)
+		{
+			return {};
+		}
+
+		const FReplicatedObjectInfo* ObjectInfo = ObjectMap->ReplicatedObjects.Find(ObjectPath);
+		if (!ObjectInfo)
+		{
+			return {};
+		}
+		
+		TArray<FConcertPropertyChainWrapper> Result;
+		Algo::Transform(ObjectInfo->PropertySelection.ReplicatedProperties, Result, [](const FConcertPropertyChain& PropertyChain)
+		{
+			return FConcertPropertyChainWrapper{ PropertyChain }; 
+		});
+		return Result;
+	}
+#endif
+	return {};
+}
+
+TArray<FSoftObjectPath> UMultiUserReplicationSubsystem::GetRegisteredObjects(const FGuid& ClientId) const
+{
+#if WITH_CONCERT
+	const UE::MultiUserClient::IMultiUserReplication* ReplicationInterface = IMultiUserClientModule::Get().GetReplication();
+	if (ensureMsgf(ReplicationInterface, TEXT("We expected it to always be valid.")))
+	{
+		const FObjectReplicationMap* ObjectMap = ReplicationInterface->FindReplicationMapForClient(ClientId);
+		if (!ObjectMap)
+		{
+			return {};
+		}
+		
+		TArray<FSoftObjectPath> Result;
+		ObjectMap->ReplicatedObjects.GenerateKeyArray(Result);
+		return Result;
+	}
+#endif
+	return {};
+}
+
+TArray<FSoftObjectPath> UMultiUserReplicationSubsystem::GetReplicatedObjects(const FGuid& ClientId) const
+{
+#if WITH_CONCERT
+	const UE::MultiUserClient::IMultiUserReplication* ReplicationInterface = IMultiUserClientModule::Get().GetReplication();
+	if (ensureMsgf(ReplicationInterface, TEXT("We expected it to always be valid.")))
+	{
+		const FObjectReplicationMap* ObjectMap = ReplicationInterface->FindReplicationMapForClient(ClientId);
+		if (!ObjectMap)
+		{
+			return {};
+		}
+
+		TArray<FSoftObjectPath> Result;
+		for (const TPair<FSoftObjectPath, FReplicatedObjectInfo>& Pair : ObjectMap->ReplicatedObjects)
+		{
+			if (ReplicationInterface->IsReplicatingObject(ClientId, Pair.Key))
+			{
+				Result.Add(Pair.Key);
+			}
+		}
+		return Result;
+	}
+#endif
+	return {};
+}
+
+void UMultiUserReplicationSubsystem::Initialize(FSubsystemCollectionBase& Collection)
+{
+	Super::Initialize(Collection);  
+
+#if WITH_CONCERT
+	UE::MultiUserClient::IMultiUserReplication* ReplicationInterface = IMultiUserClientModule::Get().GetReplication();
+	if (ensureMsgf(ReplicationInterface, TEXT("We expected it to always be valid.")))
+	{
+		UObjectAdapter = MakeShared<UE::MultiUserClientLibrary::FUObjectAdapterReplicationDiscoverer>();
+		ReplicationInterface->RegisterReplicationDiscoverer(UObjectAdapter.ToSharedRef());
+	}
+#endif
+}
+
+void UMultiUserReplicationSubsystem::Deinitialize()
+{
+	Super::Deinitialize();
+
+#if WITH_CONCERT
+	if (IMultiUserClientModule::IsAvailable())
+	{
+		UE::MultiUserClient::IMultiUserReplication* ReplicationInterface = IMultiUserClientModule::Get().GetReplication();
+		if (ensureMsgf(ReplicationInterface, TEXT("We expected it to always be valid.")))
+		{
+			ReplicationInterface->RemoveReplicationDiscoverer(UObjectAdapter.ToSharedRef());
+			UObjectAdapter.Reset();
+		}
+	}
+#endif
+}
diff --git a/Engine/Plugins/Developer/Concert/ConcertApp/MultiUserClient/Source/MultiUserClientLibrary/Private/Replication/UObjectAdaperReplicationDiscoverer.cpp b/Engine/Plugins/Developer/Concert/ConcertApp/MultiUserClient/Source/MultiUserClientLibrary/Private/Replication/UObjectAdaperReplicationDiscoverer.cpp
new file mode 100644
--- /dev/null
+++ b/Engine/Plugins/Developer/Concert/ConcertApp/MultiUserClient/Source/MultiUserClientLibrary/Private/Replication/UObjectAdaperReplicationDiscoverer.cpp
@@ -0,0 +1,44 @@
+// Copyright Epic Games, Inc. All Rights Reserved.
+
+#if WITH_CONCERT
+#include "UObjectAdapterReplicationDiscoverer.h"
+
+#include "Replication/IMultiUserReplicationRegistration.h"
+
+#include "CoreGlobals.h"
+#include "Templates/UnrealTemplate.h"
+#include "UObject/Package.h"
+
+namespace UE::MultiUserClientLibrary
+{
+	FUObjectAdapterReplicationDiscoverer::FUObjectAdapterReplicationDiscoverer()
+	{
+		Context = NewObject<UMultiUserReplicationRegistrationContextImpl>(GetTransientPackage(), NAME_None, RF_Transient);
+	}
+
+	void FUObjectAdapterReplicationDiscoverer::DiscoverReplicationSettings(const MultiUserClient::FReplicationDiscoveryParams& Params)
+	{
+		if (Params.ExtendedObject.Implements<UMultiUserReplicationRegistration>()
+			&& ensureMsgf(Context, TEXT("Was it GC'ed even though it is not supposed to have been?")))
+		{
+			// Auto-generated Execute_DiscoverReplicationSettings checks whether the UFunction is marked call in editor, which is most likely isn't
+			// hence we need to allow script execution outside of PIE temporarily.
+			const TGuardValue ScriptExecutionGuard(GAllowActorScriptExecutionInEditor, true);
+			// No technical reason to reset but let's just follow good practices and not keep it dangling
+			const TGuardValue NativeContextGuard(Context->NativeContext, &Params.Context);
+			
+			IMultiUserReplicationRegistration::Execute_DiscoverReplicationSettings(&Params.ExtendedObject, { Params.EndpointId, Context });
+		}
+	}
+	
+	void FUObjectAdapterReplicationDiscoverer::AddReferencedObjects(FReferenceCollector& Collector)
+	{
+		Collector.AddReferencedObject(Context);
+	}
+
+	FString FUObjectAdapterReplicationDiscoverer::GetReferencerName() const
+	{
+		return TEXT("FUObjectAdapterReplicationDiscoverer");
+	}
+}
+#endif
\ No newline at end of file
diff --git a/Engine/Plugins/Developer/Concert/ConcertApp/MultiUserClient/Source/MultiUserClientLibrary/Private/Replication/UObjectAdapterReplicationDiscoverer.h b/Engine/Plugins/Developer/Concert/ConcertApp/MultiUserClient/Source/MultiUserClientLibrary/Private/Replication/UObjectAdapterReplicationDiscoverer.h
new file mode 100644
--- /dev/null
+++ b/Engine/Plugins/Developer/Concert/ConcertApp/MultiUserClient/Source/MultiUserClientLibrary/Private/Replication/UObjectAdapterReplicationDiscoverer.h
@@ -0,0 +1,40 @@
+// Copyright Epic Games, Inc. All Rights Reserved.
+
+#pragma once
+
+#if WITH_CONCERT
+
+#include "MultiUserReplicationRegistrationContextImpl.h"
+#include "Replication/IReplicationDiscoverer.h"
+#include "UObject/GCObject.h"
+
+namespace UE::MultiUserClientLibrary
+{
+	/**
+	 * This discovery implementation checks whether the target UObject implements IConcertReplicationRegistration, and calls the equivalent functions on it.
+	 * @see IMultiUserReplicationRegistration
+	 */
+	class FUObjectAdapterReplicationDiscoverer
+		: public MultiUserClient::IReplicationDiscoverer
+		, public FGCObject
+	{
+	public:
+
+		FUObjectAdapterReplicationDiscoverer();
+
+		//~ Begin IReplicationDiscovery Interface
+		virtual void DiscoverReplicationSettings(const MultiUserClient::FReplicationDiscoveryParams& Params) override;
+		//~ End IReplicationDiscovery Interface
+
+		//~ Begin FGCObject Interface
+		virtual void AddReferencedObjects(FReferenceCollector& Collector) override;
+		virtual FString GetReferencerName() const override;
+		//~ End FGCObject Interface
+
+	private:
+
+		/** Passed to IMultiUserReplicationRegistration. We set the IReplicationDiscoveryContext for every call of DiscoverReplicationSettings.  */
+		TObjectPtr<UMultiUserReplicationRegistrationContextImpl> Context;
+	};
+}
+#endif
diff --git a/Engine/Plugins/Developer/Concert/ConcertApp/MultiUserClient/Source/MultiUserClientLibrary/Public/MultiUserSubsystem.h b/Engine/Plugins/Developer/Concert/ConcertApp/MultiUserClient/Source/MultiUserClientLibrary/Public/MultiUserSubsystem.h
--- a/Engine/Plugins/Developer/Concert/ConcertApp/MultiUserClient/Source/MultiUserClientLibrary/Public/MultiUserSubsystem.h
+++ b/Engine/Plugins/Developer/Concert/ConcertApp/MultiUserClient/Source/MultiUserClientLibrary/Public/MultiUserSubsystem.h
@@ -52,6 +52,19 @@ public:
 	UFUNCTION(BlueprintCallable, Category = "Multi-user")
 	bool IsConnectedToSession() const;
 
+	/**
+	 * If connected to a session, gets the endpoint identifiers of client corresponding to the local editor instance.
+	 * @return Whether OutClientId contains a valid identifiers.
+	 */
+	UFUNCTION(BlueprintPure, Category = "Multi-user")
+	bool GetLocalClientId(FGuid& OutClientId) const;
+	/**
+	 * If connected to a session, gets the endpoint identifiers of all remote clients.
+	 * @return Whether OutRemoteClientIds contains the remote client endpoint identifiers.
+	 */
+	UFUNCTION(BlueprintPure, Category = "Multi-user")
+	bool GetRemoteClientIds(TArray<FGuid>& OutRemoteClientIds);
+
 	/**
 	 * Send a custom event message over multi-user.  If you a not in an active session, this function does nothing and only
 	 * reports and warning to the output log.
diff --git a/Engine/Plugins/Developer/Concert/ConcertApp/MultiUserClient/Source/MultiUserClientLibrary/Public/Replication/IMultiUserReplicationRegistration.h b/Engine/Plugins/Developer/Concert/ConcertApp/MultiUserClient/Source/MultiUserClientLibrary/Public/Replication/IMultiUserReplicationRegistration.h
new file mode 100644
--- /dev/null
+++ b/Engine/Plugins/Developer/Concert/ConcertApp/MultiUserClient/Source/MultiUserClientLibrary/Public/Replication/IMultiUserReplicationRegistration.h
@@ -0,0 +1,49 @@
+// Copyright Epic Games, Inc. All Rights Reserved.
+
+#pragma once
+
+#include "UObject/Interface.h"
+
+#include "ConcertPropertyChainWrapper.h"
+#include "IMultiUserReplicationRegistrationContext.h"
+#include "Misc/Guid.h"
+
+#include "IMultiUserReplicationRegistration.generated.h"
+
+struct FConcertPropertyChainWrapper;
+
+UINTERFACE()
+class UMultiUserReplicationRegistration : public UInterface
+{
+	GENERATED_BODY()
+};
+
+USTRUCT(BlueprintType)
+struct FMultiUserReplicationRegistrationParams
+{
+	GENERATED_BODY()
+
+	/** The client for which the objects and properties are being registered. */
+	UPROPERTY(BlueprintReadOnly, Category = "Multi-user")
+	FGuid ClientEndpointId;
+
+	/** Uses this to register properties and add additional objects for discovery. */
+	UPROPERTY(BlueprintReadOnly, Category = "Multi-user")
+	TScriptInterface<IMultiUserReplicationRegistrationContext> Context;
+};
+
+/**
+ * Exposes IReplicationDiscoverer to Blueprints.
+ * 
+ * FUObjectAdapterReplicationDiscoverer is used to adapt this interface to the native IReplicationDiscovery interface.
+ * @see FUObjectAdapterReplicationDiscoverer
+ */
+class MULTIUSERCLIENTLIBRARY_API IMultiUserReplicationRegistration
+{
+	GENERATED_BODY()
+public:
+	
+	/** Registers properties and additional objects that must be replicated. */
+	UFUNCTION(BlueprintCallable, BlueprintNativeEvent, Category = "Multi-user|Replication")
+	void DiscoverReplicationSettings(const FMultiUserReplicationRegistrationParams& Params);
+};
diff --git a/Engine/Plugins/Developer/Concert/ConcertApp/MultiUserClient/Source/MultiUserClientLibrary/Public/Replication/IMultiUserReplicationRegistrationContext.h b/Engine/Plugins/Developer/Concert/ConcertApp/MultiUserClient/Source/MultiUserClientLibrary/Public/Replication/IMultiUserReplicationRegistrationContext.h
new file mode 100644
--- /dev/null
+++ b/Engine/Plugins/Developer/Concert/ConcertApp/MultiUserClient/Source/MultiUserClientLibrary/Public/Replication/IMultiUserReplicationRegistrationContext.h
@@ -0,0 +1,47 @@
+// Copyright Epic Games, Inc. All Rights Reserved.
+
+#pragma once
+
+#include "UObject/Interface.h"
+#include "IMultiUserReplicationRegistrationContext.generated.h"
+
+struct FConcertPropertyChainWrapper;
+
+UINTERFACE(NotBlueprintable)
+class UMultiUserReplicationRegistrationContext : public UInterface
+{
+	GENERATED_BODY()
+};
+
+/**
+ * Given to IMultiUserReplicationRegistration to register properties and objects.
+ * @see IMultiUserReplicationRegistration
+ */
+class MULTIUSERCLIENTLIBRARY_API IMultiUserReplicationRegistrationContext
+{
+	GENERATED_BODY()
+public:
+
+	/**
+	 * Extends the stream with the given object, if not yet added, and adds the properties to it.
+	 *
+	 * Properties will be automatically discovered for Object, as well.
+	 * If the object implements IMultiUserReplicationRegistrationContext, it will be invoked on Object recursively.
+	 * 
+	 * @param Object Object to add, if not yet present, and assign a property to
+	 * @param PropertiesToAdd Property chain to associate with the object
+	 */
+	UFUNCTION(BlueprintCallable, Category = "Multi-user")
+	virtual void AddPropertiesToObject(UObject* Object, const TArray<FConcertPropertyChainWrapper>& PropertiesToAdd) = 0;
+
+	/**
+	 * Adds an additional object.
+	 * 
+	 * Properties will be automatically discovered for Object, as well.
+	 * If the object implements IMultiUserReplicationRegistrationContext, it will be invoked on Object recursively.
+	 * 
+	 * @param Object The additional object to add
+	 */
+	UFUNCTION(BlueprintCallable, Category = "Multi-user")
+	virtual void AddAdditionalObject(UObject* Object) = 0;
+};
diff --git a/Engine/Plugins/Developer/Concert/ConcertApp/MultiUserClient/Source/MultiUserClientLibrary/Public/Replication/MultiUserReplicationSubsystem.h b/Engine/Plugins/Developer/Concert/ConcertApp/MultiUserClient/Source/MultiUserClientLibrary/Public/Replication/MultiUserReplicationSubsystem.h
new file mode 100644
--- /dev/null
+++ b/Engine/Plugins/Developer/Concert/ConcertApp/MultiUserClient/Source/MultiUserClientLibrary/Public/Replication/MultiUserReplicationSubsystem.h
@@ -0,0 +1,82 @@
+// Copyright Epic Games, Inc. All Rights Reserved.
+
+#pragma once
+
+#include "Subsystems/EngineSubsystem.h"
+#include "ConcertPropertyChainWrapper.h"
+#include "MultiUserReplicationSubsystem.generated.h"
+
+namespace UE::MultiUserClientLibrary
+{
+	class FUObjectAdapterReplicationDiscoverer;
+}
+
+/** Exposes ways to interact with the Multi-user replication system via Blueprints. */
+UCLASS()
+class MULTIUSERCLIENTLIBRARY_API UMultiUserReplicationSubsystem : public UEngineSubsystem
+{
+	GENERATED_BODY()
+public:
+	
+	// This would be the right place to expose additional MU specific replication functions in the future
+	
+	/**
+	 * @return Whether the client is replicating the object.
+	 * @note An object can be registered but not replicated.
+	 * @see UMultiUserSubsystem::GetLocalClientId and UMultiUserSubsystem::GetRemoteClientIds
+	 */
+	UFUNCTION(BlueprintPure, Category = "Multi-user")
+	bool IsReplicatingObject(const FGuid& ClientId, const FSoftObjectPath& ObjectPath) const;
+
+	/**
+	 * @return The properties the client has registered for replication for the object.
+	 * @note An object can be registered but not replicated. Use IsReplicatingObject to find out whether the client is replicating the returned properties.
+	 * @see UMultiUserSubsystem::GetLocalClientId and UMultiUserSubsystem::GetRemoteClientIds.
+	 */
+	UFUNCTION(BlueprintPure, Category = "Multi-user")
+	TArray<FConcertPropertyChainWrapper> GetPropertiesRegisteredToObject(const FGuid& ClientId, const FSoftObjectPath& ObjectPath) const;
+
+	/**
+	 * Gets the objects the client has registered with the server.
+	 *
+	 * Just because an object is replicated, it does not mean that the object is being replicated.
+	 * Objects must be registered with the server so registration is just the first step.
+	 * GetReplicatedObjects() will always contain GetRegisteredObjects().
+	 * 
+	 * @param ClientId The client of which to get the registered objects. 
+	 * @return The objects the client has
+	 * 
+	 * @see UMultiUserSubsystem::GetLocalClientId and UMultiUserSubsystem::GetRemoteClientIds
+	 */
+	UFUNCTION(BlueprintPure, Category = "Multi-user")
+	TArray<FSoftObjectPath> GetRegisteredObjects(const FGuid& ClientId) const;
+
+	/**
+	 * Gets the objects that are currently being replicated by the client.
+	 *
+	 * @note There is a difference between registered and replicated objects! Objects are registered with the server first and later the client
+	 * can attempt to start replicating them. GetReplicatedObjects() will always contain GetRegisteredObjects().
+	 * 
+	 * @param ClientId The client of which to get the replicated objects. 
+	 * @return The objects being replicated by the client
+	 *
+	 * @see UMultiUserSubsystem::GetLocalClientId and UMultiUserSubsystem::GetRemoteClientIds
+	 */
+	UFUNCTION(BlueprintPure, Category = "Multi-user")
+	TArray<FSoftObjectPath> GetReplicatedObjects(const FGuid& ClientId) const;
+	
+	//~ Begin USubsystem Interface
+	virtual void Initialize(FSubsystemCollectionBase& Collection) override;
+	virtual void Deinitialize() override;
+	//~ End USubsystem Interface
+
+private:
+
+	/**
+	 * Registered when initialized.
+	 * 
+	 * This allows UObjects, the target being Blueprints, to implement the IConcertReplicationRegistration interface through which MU will use to
+	 * auto-add properties when registering an object to a client's replication stream.
+	 */
+	TSharedPtr<UE::MultiUserClientLibrary::FUObjectAdapterReplicationDiscoverer> UObjectAdapter;
+};
diff --git a/Engine/Plugins/Developer/Concert/ConcertScripting/ConcertReplicationScripting/ConcertReplicationScripting.uplugin b/Engine/Plugins/Developer/Concert/ConcertScripting/ConcertReplicationScripting/ConcertReplicationScripting.uplugin
new file mode 100644
--- /dev/null
+++ b/Engine/Plugins/Developer/Concert/ConcertScripting/ConcertReplicationScripting/ConcertReplicationScripting.uplugin
@@ -0,0 +1,34 @@
+﻿{
+	"FileVersion": 3,
+	"Version": 1,
+	"VersionName": "1.0",
+	"FriendlyName": "Concert Replication Scripting",
+	"Description": "Exposes Concert Replication types for scripting, e.g. in Blueprints",
+	"Category": "Networking",
+	"CreatedBy": "Epic Games, Inc.",
+	"CreatedByURL": "https://epicgames.com",
+	"DocsURL": "",
+	"MarketplaceURL": "",
+	"CanContainContent": true,
+	"IsBetaVersion": false,
+	"IsExperimentalVersion": false,
+	"Installed": false,
+	"Modules": [
+        {
+          "Name": "ConcertReplicationScripting",
+          "Type": "Runtime",
+          "LoadingPhase": "Default"
+        },
+		{
+          "Name": "ConcertReplicationScriptingEditor",
+          "Type": "Editor",
+          "LoadingPhase": "Default"
+        }
+	],
+	"Plugins": [
+		{
+			"Name": "ConcertSyncCore",
+			"Enabled": true
+		}
+	]
+}
\ No newline at end of file
diff --git a/Engine/Plugins/Developer/Concert/ConcertScripting/ConcertReplicationScripting/Source/ConcertReplicationScripting/ConcertReplicationScripting.Build.cs b/Engine/Plugins/Developer/Concert/ConcertScripting/ConcertReplicationScripting/Source/ConcertReplicationScripting/ConcertReplicationScripting.Build.cs
new file mode 100644
--- /dev/null
+++ b/Engine/Plugins/Developer/Concert/ConcertScripting/ConcertReplicationScripting/Source/ConcertReplicationScripting/ConcertReplicationScripting.Build.cs
@@ -0,0 +1,33 @@
+﻿// Copyright Epic Games, Inc. All Rights Reserved.
+
+namespace UnrealBuildTool.Rules
+{
+	public class ConcertReplicationScripting : ModuleRules
+	{
+		public ConcertReplicationScripting(ReadOnlyTargetRules Target) : base(Target)
+		{
+			PublicDependencyModuleNames.AddRange(
+				new string[]
+				{
+					// Engine
+					"Core",
+					"CoreUObject",
+					"Engine",
+					
+					// Concert
+					"ConcertSyncCore"
+				}
+			);
+			
+			PrivateDependencyModuleNames.AddRange(
+				new string[]
+				{
+					// Concert
+					"ConcertTransport", // For LogConcert
+				}
+			);
+
+			ShortName = "CrtSc";
+		}
+	}
+}
\ No newline at end of file
diff --git a/Engine/Plugins/Developer/Concert/ConcertScripting/ConcertReplicationScripting/Source/ConcertReplicationScripting/Private/ConcertReplicationBlueprintFunctionLibrary.cpp b/Engine/Plugins/Developer/Concert/ConcertScripting/ConcertReplicationScripting/Source/ConcertReplicationScripting/Private/ConcertReplicationBlueprintFunctionLibrary.cpp
new file mode 100644
--- /dev/null
+++ b/Engine/Plugins/Developer/Concert/ConcertScripting/ConcertReplicationScripting/Source/ConcertReplicationScripting/Private/ConcertReplicationBlueprintFunctionLibrary.cpp
@@ -0,0 +1,91 @@
+// Copyright Epic Games, Inc. All Rights Reserved.
+
+#include "ConcertReplicationBlueprintFunctionLibrary.h"
+
+#include "ConcertLogGlobal.h"
+#include "ConcertPropertyChainWrapper.h"
+#include "Replication/PropertyChainUtils.h"
+
+#include "Templates/SubclassOf.h"
+
+bool UConcertReplicationBlueprintFunctionLibrary::MakePropertyChainByLiteralPath(const TSubclassOf<UObject>& Class, const TArray<FName>& PathToProperty, FConcertPropertyChainWrapper& Result)
+{
+	if (!Class)
+	{
+		UE_LOG(LogConcert, Error, TEXT("UConcertReplicationBlueprintFunctionLibrary::MakePropertyChainByLiteralPath: Invalid class!"))
+		return false;
+	}
+
+	TOptional<FConcertPropertyChain> PropertyChain = FConcertPropertyChain::CreateFromPath(**Class, PathToProperty);
+	if (PropertyChain)
+	{
+		constexpr bool bLogOnFail = false;
+		const FProperty* Property = PropertyChain->ResolveProperty(**Class, bLogOnFail);
+		if (!ensure(Property) || !UE::ConcertSyncCore::PropertyChain::IsReplicatableProperty(*Property))
+		{
+			UE_LOG(LogConcert, Error, TEXT("UConcertReplicationBlueprintFunctionLibrary::MakePropertyChainByLiteralPath: Property path [%s] resolved to a property that is not replicatable!"),
+				*FString::JoinBy(PathToProperty, TEXT(", "), [](const FName& Name){ return TEXT("\"") + Name.ToString() + TEXT("\""); })
+				);
+			return false;
+		}
+
+		Result.PropertyChain = MoveTemp(*PropertyChain);
+		return true;
+	}
+
+	UE_LOG(LogConcert, Error, TEXT("UConcertReplicationBlueprintFunctionLibrary::MakePropertyChainByLiteralPath: Path [%s] is invalid!"),
+		*FString::JoinBy(PathToProperty, TEXT(", "), [](const FName& Name){ return TEXT("\"") + Name.ToString() + TEXT("\""); })
+		);
+	return false;
+}
+
+TArray<FConcertPropertyChainWrapper> UConcertReplicationBlueprintFunctionLibrary::GetPropertiesIn(const TSubclassOf<UObject>& Class, FPropertyChainPredicate Filter)
+{
+	if (!Class)
+	{
+		UE_LOG(LogConcert, Error, TEXT("UConcertReplicationBlueprintFunctionLibrary::GetPropertiesIn: Invalid class!"))
+		return {};
+	}
+
+	TArray<FConcertPropertyChainWrapper> Result;  
+	UE::ConcertSyncCore::PropertyChain::ForEachReplicatableConcertProperty(
+		**Class,
+		[&Filter, &Result](FConcertPropertyChain&& PropertyChain)
+		{
+			FConcertPropertyChainWrapper Wrapper { MoveTemp(PropertyChain) };
+			if (!Filter.IsBound() || Filter.Execute(Wrapper))
+			{
+				Result.Emplace(MoveTemp(Wrapper.PropertyChain));
+			}
+			return EBreakBehavior::Continue;
+		});
+	return Result;
+}
+
+const TArray<FName>& UConcertReplicationBlueprintFunctionLibrary::GetPropertyStringPath(const FConcertPropertyChainWrapper& Path)
+{
+	return Path.PropertyChain.GetPathToProperty();
+}
+
+FName UConcertReplicationBlueprintFunctionLibrary::GetPropertyFromRoot(const FConcertPropertyChainWrapper& Path, int32 Index)
+{
+	const TArray<FName>& StringPath = Path.PropertyChain.GetPathToProperty();
+	return StringPath.IsValidIndex(Index) ? StringPath[Index] : NAME_None;
+}
+
+FName UConcertReplicationBlueprintFunctionLibrary::GetPropertyFromLeaf(const FConcertPropertyChainWrapper& Path, int32 Index)
+{
+	const TArray<FName>& StringPath = Path.PropertyChain.GetPathToProperty();
+	const int32 RealIndex = StringPath.Num() - 1 - Index;
+	return StringPath.IsValidIndex(RealIndex) ? StringPath[RealIndex] : NAME_None;
+}
+
+bool UConcertReplicationBlueprintFunctionLibrary::IsChildOf(const FConcertPropertyChainWrapper& ToTest, const FConcertPropertyChainWrapper& Parent)
+{
+	return ToTest.PropertyChain.IsChildOf(Parent.PropertyChain);
+}
+
+bool UConcertReplicationBlueprintFunctionLibrary::IsDirectChildOf(const FConcertPropertyChainWrapper& ToTest, const FConcertPropertyChainWrapper& Parent)
+{
+	return ToTest.PropertyChain.IsDirectChildOf(Parent.PropertyChain);
+}
diff --git a/Engine/Plugins/Developer/Concert/ConcertScripting/ConcertReplicationScripting/Source/ConcertReplicationScripting/Private/ConcertReplicationScriptingModule.cpp b/Engine/Plugins/Developer/Concert/ConcertScripting/ConcertReplicationScripting/Source/ConcertReplicationScripting/Private/ConcertReplicationScriptingModule.cpp
new file mode 100644
--- /dev/null
+++ b/Engine/Plugins/Developer/Concert/ConcertScripting/ConcertReplicationScripting/Source/ConcertReplicationScripting/Private/ConcertReplicationScriptingModule.cpp
@@ -0,0 +1,16 @@
+// Copyright Epic Games, Inc. All Rights Reserved.
+
+#include "ConcertReplicationScriptingModule.h"
+
+#include "Modules/ModuleManager.h"
+
+namespace UE::ConcertReplicationScripting
+{
+	void FConcertReplicationScriptingModule::StartupModule()
+	{}
+
+	void FConcertReplicationScriptingModule::ShutdownModule()
+	{}
+}
+
+IMPLEMENT_MODULE(UE::ConcertReplicationScripting::FConcertReplicationScriptingModule, ConcertReplicationScripting);
\ No newline at end of file
diff --git a/Engine/Plugins/Developer/Concert/ConcertScripting/ConcertReplicationScripting/Source/ConcertReplicationScripting/Private/ConcertReplicationScriptingModule.h b/Engine/Plugins/Developer/Concert/ConcertScripting/ConcertReplicationScripting/Source/ConcertReplicationScripting/Private/ConcertReplicationScriptingModule.h
new file mode 100644
--- /dev/null
+++ b/Engine/Plugins/Developer/Concert/ConcertScripting/ConcertReplicationScripting/Source/ConcertReplicationScripting/Private/ConcertReplicationScriptingModule.h
@@ -0,0 +1,18 @@
+// Copyright Epic Games, Inc. All Rights Reserved.
+
+#pragma once
+
+#include "Modules/ModuleInterface.h"
+
+namespace UE::ConcertReplicationScripting
+{
+	class FConcertReplicationScriptingModule : public IModuleInterface
+	{
+	public:
+	
+		//~ Begin IModuleInterface Interface
+		virtual void StartupModule() override;
+		virtual void ShutdownModule() override;
+		//~ End IModuleInterface Interface
+	};
+}
\ No newline at end of file
diff --git a/Engine/Plugins/Developer/Concert/ConcertScripting/ConcertReplicationScripting/Source/ConcertReplicationScripting/Public/ConcertPropertyChainWrapper.h b/Engine/Plugins/Developer/Concert/ConcertScripting/ConcertReplicationScripting/Source/ConcertReplicationScripting/Public/ConcertPropertyChainWrapper.h
new file mode 100644
--- /dev/null
+++ b/Engine/Plugins/Developer/Concert/ConcertScripting/ConcertReplicationScripting/Source/ConcertReplicationScripting/Public/ConcertPropertyChainWrapper.h
@@ -0,0 +1,20 @@
+// Copyright Epic Games, Inc. All Rights Reserved.
+
+#pragma once
+
+#include "Replication/Data/ConcertPropertySelection.h"
+#include "ConcertPropertyChainWrapper.generated.h"
+
+USTRUCT(BlueprintType, meta = (DisplayName = "Concert Property Chain"))
+struct CONCERTREPLICATIONSCRIPTING_API FConcertPropertyChainWrapper
+{
+	GENERATED_BODY()
+
+	UPROPERTY()
+	FConcertPropertyChain PropertyChain;
+	
+	FConcertPropertyChainWrapper() = default;
+	explicit FConcertPropertyChainWrapper(FConcertPropertyChain InPropertyChain)
+		: PropertyChain(MoveTemp(InPropertyChain))
+	{}
+};
\ No newline at end of file
diff --git a/Engine/Plugins/Developer/Concert/ConcertScripting/ConcertReplicationScripting/Source/ConcertReplicationScripting/Public/ConcertReplicationBlueprintFunctionLibrary.h b/Engine/Plugins/Developer/Concert/ConcertScripting/ConcertReplicationScripting/Source/ConcertReplicationScripting/Public/ConcertReplicationBlueprintFunctionLibrary.h
new file mode 100644
--- /dev/null
+++ b/Engine/Plugins/Developer/Concert/ConcertScripting/ConcertReplicationScripting/Source/ConcertReplicationScripting/Public/ConcertReplicationBlueprintFunctionLibrary.h
@@ -0,0 +1,101 @@
+// Copyright Epic Games, Inc. All Rights Reserved.
+
+#pragma once
+
+#include "Kismet/BlueprintFunctionLibrary.h"
+#include "ConcertReplicationBlueprintFunctionLibrary.generated.h"
+
+struct FConcertPropertyChainWrapper;
+
+UCLASS(BlueprintType)
+class CONCERTREPLICATIONSCRIPTING_API UConcertReplicationBlueprintFunctionLibrary : public UBlueprintFunctionLibrary
+{
+	GENERATED_BODY()
+public:
+	
+	DECLARE_DYNAMIC_DELEGATE_RetVal_OneParam(bool, FPropertyChainPredicate, const FConcertPropertyChainWrapper&, PathToProperty);
+
+	/******************** Factories ********************/
+	
+	/**
+	 * Tries to make concert property chain by specifying a path manually.
+	 * 
+	 * @param Class The class in which to search for the property path 
+	 * @param PathToProperty The path to the property
+	 * @param Result The resulting path
+	 * 
+	 * @return Whether Result is valid.
+	 */
+	UFUNCTION(BlueprintPure, Category = "Concert|Replication")
+	static bool MakePropertyChainByLiteralPath(const TSubclassOf<UObject>& Class, const TArray<FName>& PathToProperty, FConcertPropertyChainWrapper& Result);
+	
+	/**
+	 * Builds an array of property paths that pass the given filter.
+	 * 
+	 * @param Class The class in which to search for properties
+	 * @param Filter Decides whether a property path should be included
+	 * 
+	 * @return An array of all property paths in the given class that passed the filter
+	 */
+	UFUNCTION(BlueprintPure, Category = "Concert|Replication", meta = (KeyWords = "Get Properties In Find Enumerate List"))
+	static TArray<FConcertPropertyChainWrapper> GetPropertiesIn(const TSubclassOf<UObject>& Class, FPropertyChainPredicate Filter);
+	
+	/******************** Getters ********************/
+	
+	/**
+	 * Gets the path as string array.
+	 * Example: ["RelativeLocation", "X"]
+	 *
+	 * @param Path The property path to extract from
+	 *
+	 * @return The path as string array
+	 */
+	UFUNCTION(BlueprintPure, Category = "Concert|Replication")
+	static const TArray<FName>& GetPropertyStringPath(const FConcertPropertyChainWrapper& Path);
+
+	/**
+	 * Gets the property at Index starting from the root most property.
+	 * Example: The root Index 0 for ["RelativeLocation", "X"] would return "RelativeLocation".
+	 * 
+	 * @param Path The property path to get the sub-property from.
+	 * @param Index The index in the path counting from the left, root property.
+	 * 
+	 * @return The property name at Index or None if Index is invalid.
+	 */
+	UFUNCTION(BlueprintPure, Category = "Concert|Replication")
+	static FName GetPropertyFromRoot(const FConcertPropertyChainWrapper& Path, int32 Index = 0);
+
+	/**
+	 * Gets the property at Index starting from the leaf most property.
+	 * Example: The leaf Index 0 for ["RelativeLocation", "X"] would return "X".
+	 * 
+	 * @param Path The property path to get the sub-property from.
+	 * @param Index The index in the path counting from the right, leaf property.
+	 * 
+	 * @return The property name at Index or None if Index is invalid.
+	 */
+	UFUNCTION(BlueprintPure, Category = "Concert|Replication")
+	static FName GetPropertyFromLeaf(const FConcertPropertyChainWrapper& Path, int32 Index = 0);
+
+	/**
+	 * Checks whether ToTest is a child property of Parent.
+	 * 
+	 * @param ToTest The property for which to check whether it is a child property
+	 * @param Parent The property which is supposed to be the parent of ToTest
+	 * 
+	 * @return Whether ToTest is a child property of Parent
+	 */
+	UFUNCTION(BlueprintPure, Category = "Concert|Replication")
+	static bool IsChildOf(const FConcertPropertyChainWrapper& ToTest, const FConcertPropertyChainWrapper& Parent);
+
+	/**
+	 * Checks whether ToTest is a direct child property of Parent.
+	 * 
+	 * @param ToTest The property for which to check whether it is a child property
+	 * @param Parent The property which is supposed to be the parent of ToTest
+	 * 
+	 * @return Whether ToTest is a direct child property of Parent
+	 */
+	UFUNCTION(BlueprintPure, Category = "Concert|Replication")
+	static bool IsDirectChildOf(const FConcertPropertyChainWrapper& ToTest, const FConcertPropertyChainWrapper& Parent);
+};
\ No newline at end of file
diff --git a/Engine/Plugins/Developer/Concert/ConcertScripting/ConcertReplicationScripting/Source/ConcertReplicationScriptingEditor/ConcertReplicationScriptingEditor.build.cs b/Engine/Plugins/Developer/Concert/ConcertScripting/ConcertReplicationScripting/Source/ConcertReplicationScriptingEditor/ConcertReplicationScriptingEditor.build.cs
new file mode 100644
--- /dev/null
+++ b/Engine/Plugins/Developer/Concert/ConcertScripting/ConcertReplicationScripting/Source/ConcertReplicationScriptingEditor/ConcertReplicationScriptingEditor.build.cs
@@ -0,0 +1,41 @@
+// Copyright Epic Games, Inc. All Rights Reserved.
+
+namespace UnrealBuildTool.Rules
+{
+	public class ConcertReplicationScriptingEditor : ModuleRules
+	{
+		public ConcertReplicationScriptingEditor(ReadOnlyTargetRules Target) : base(Target)
+		{
+			PublicDependencyModuleNames.AddRange(
+				new string[]
+				{
+					// Engine
+					"Core",
+					"CoreUObject",
+					"Engine",
+					"Slate",
+					"SlateCore",
+					
+					// Concert
+					"ConcertReplicationScripting",
+					"ConcertSyncCore"
+				}
+			);
+			
+			PrivateDependencyModuleNames.AddRange(
+				new string[]
+				{
+					// Engine
+					"Slate",
+					"SlateCore",
+					"UnrealEd",
+					
+					// Concert
+					"ConcertTransport", // For LogConcert
+				}
+			);
+
+			ShortName = "CrtScEd";
+		}
+	}
+}
diff --git a/Engine/Plugins/Developer/Concert/ConcertScripting/ConcertReplicationScripting/Source/ConcertReplicationScriptingEditor/Private/ConcertReplicationScriptingEditorModule.cpp b/Engine/Plugins/Developer/Concert/ConcertScripting/ConcertReplicationScripting/Source/ConcertReplicationScriptingEditor/Private/ConcertReplicationScriptingEditorModule.cpp
new file mode 100644
--- /dev/null
+++ b/Engine/Plugins/Developer/Concert/ConcertScripting/ConcertReplicationScripting/Source/ConcertReplicationScriptingEditor/Private/ConcertReplicationScriptingEditorModule.cpp
@@ -0,0 +1,16 @@
+// Copyright Epic Games, Inc. All Rights Reserved.
+
+#include "ConcertReplicationScriptingEditorModule.h"
+
+#include "Modules/ModuleManager.h"
+
+namespace UE::ConcertReplicationScriptingEditor
+{
+	void FConcertReplicationScriptingEditorModule::StartupModule()
+	{}
+
+	void FConcertReplicationScriptingEditorModule::ShutdownModule()
+	{}
+}
+
+IMPLEMENT_MODULE(UE::ConcertReplicationScriptingEditor::FConcertReplicationScriptingEditorModule, ConcertReplicationScriptingEditor);
\ No newline at end of file
diff --git a/Engine/Plugins/Developer/Concert/ConcertScripting/ConcertReplicationScripting/Source/ConcertReplicationScriptingEditor/Private/ConcertReplicationScriptingEditorModule.h b/Engine/Plugins/Developer/Concert/ConcertScripting/ConcertReplicationScripting/Source/ConcertReplicationScriptingEditor/Private/ConcertReplicationScriptingEditorModule.h
new file mode 100644
--- /dev/null
+++ b/Engine/Plugins/Developer/Concert/ConcertScripting/ConcertReplicationScripting/Source/ConcertReplicationScriptingEditor/Private/ConcertReplicationScriptingEditorModule.h
@@ -0,0 +1,18 @@
+// Copyright Epic Games, Inc. All Rights Reserved.
+
+#pragma once
+
+#include "Modules/ModuleInterface.h"
+
+namespace UE::ConcertReplicationScriptingEditor
+{
+	class FConcertReplicationScriptingEditorModule : public IModuleInterface
+	{
+	public:
+	
+		//~ Begin IModuleInterface Interface
+		virtual void StartupModule() override;
+		virtual void ShutdownModule() override;
+		//~ End IModuleInterface Interface
+	};
+}
\ No newline at end of file
diff --git a/Engine/Plugins/Developer/Concert/ConcertSync/ConcertSyncCore/Source/ConcertSyncCore/Private/Replication/Data/ConcertPropertySelection.cpp b/Engine/Plugins/Developer/Concert/ConcertSync/ConcertSyncCore/Source/ConcertSyncCore/Private/Replication/Data/ConcertPropertySelection.cpp
--- a/Engine/Plugins/Developer/Concert/ConcertSync/ConcertSyncCore/Source/ConcertSyncCore/Private/Replication/Data/ConcertPropertySelection.cpp
+++ b/Engine/Plugins/Developer/Concert/ConcertSync/ConcertSyncCore/Source/ConcertSyncCore/Private/Replication/Data/ConcertPropertySelection.cpp
@@ -117,9 +117,13 @@ bool FConcertPropertyChain::IsDirectChildOf(const FConcertPropertyChain& ParentT
 bool FConcertPropertyChain::MatchesExactly(const FArchiveSerializedPropertyChain* OptionalChain, const FProperty& LeafProperty) const
 {
 	using namespace UE::ConcertSyncCore::PropertyChain;
-	const int32 OptionalChainLength = OptionalChain ? OptionalChain->GetNumProperties() : 0;
+	if (PathToProperty.IsEmpty())
+	{
+		return false;
+	}
 	
-	const bool bLeafPropertiesMatch = !PathToProperty.IsEmpty() && IsInnerContainerProperty(LeafProperty)
+	const int32 OptionalChainLength = OptionalChain ? OptionalChain->GetNumProperties() : 0;
+	const bool bLeafPropertiesMatch = IsInnerContainerProperty(LeafProperty)
 		// The only place FConcertPropertyChain contains inner container properties is at the end; it is named InternalContainerPropertyValueName.
 		// In that case the "real" leaf property is the owning container property.
 		? (OptionalChainLength >= 1 && PathToProperty.Num() > 1) && OptionalChain->GetPropertyFromStack(0)->GetFName() == PathToProperty[PathToProperty.Num() - 2]
@@ -179,7 +183,7 @@ bool FConcertPropertyChain::MatchesExactly(const FArchiveSerializedPropertyChain
 	return bArePathsEqual && bVisitedEveryProperty;
 }
 
-FProperty* FConcertPropertyChain::ResolveProperty(UStruct& Class, bool bLogOnFail)
+FProperty* FConcertPropertyChain::ResolveProperty(UStruct& Class, bool bLogOnFail) const
 {
 	// FConcertPropertyChain::ResolveProperty exists only for visibility to developers since
 	// the class is the first place one would look and not in the utils namespace.
diff --git a/Engine/Plugins/Developer/Concert/ConcertSync/ConcertSyncCore/Source/ConcertSyncCore/Public/Replication/Data/ConcertPropertySelection.h b/Engine/Plugins/Developer/Concert/ConcertSync/ConcertSyncCore/Source/ConcertSyncCore/Public/Replication/Data/ConcertPropertySelection.h
--- a/Engine/Plugins/Developer/Concert/ConcertSync/ConcertSyncCore/Source/ConcertSyncCore/Public/Replication/Data/ConcertPropertySelection.h
+++ b/Engine/Plugins/Developer/Concert/ConcertSync/ConcertSyncCore/Source/ConcertSyncCore/Public/Replication/Data/ConcertPropertySelection.h
@@ -65,7 +65,7 @@ struct CONCERTSYNCCORE_API FConcertPropertyChain
 	bool MatchesExactly(const FArchiveSerializedPropertyChain* OptionalChain, const FProperty& LeafProperty) const;
 
 	/** @return Attempts to resolve this property given the class */
-	FProperty* ResolveProperty(UStruct& Class, bool bLogOnFail = true);
+	FProperty* ResolveProperty(UStruct& Class, bool bLogOnFail = true) const;
 
 	const TArray<FName>& GetPathToProperty() const { return PathToProperty; }
 
diff --git a/Engine/Plugins/Developer/Concert/ConcertUI/ConcertClientSharedSlate/Source/ConcertClientSharedSlate/Private/Replication/Editor/Model/GenericReplicationStreamModel.cpp b/Engine/Plugins/Developer/Concert/ConcertUI/ConcertClientSharedSlate/Source/ConcertClientSharedSlate/Private/Replication/Editor/Model/GenericReplicationStreamModel.cpp
--- a/Engine/Plugins/Developer/Concert/ConcertUI/ConcertClientSharedSlate/Source/ConcertClientSharedSlate/Private/Replication/Editor/Model/GenericReplicationStreamModel.cpp
+++ b/Engine/Plugins/Developer/Concert/ConcertUI/ConcertClientSharedSlate/Source/ConcertClientSharedSlate/Private/Replication/Editor/Model/GenericReplicationStreamModel.cpp
@@ -284,13 +284,19 @@ namespace UE::ConcertClientSharedSlate
 				, ObjectsAddedSoFar(ObjectsAddedSoFar)
 			{}
 
-			virtual void AddPropertyTo(UObject& Object, FConcertPropertyChain&& PropertyChain) override
+			virtual void AddPropertyTo(UObject& Object, FConcertPropertyChain PropertyChain) override
 			{
+				UClass* ObjectClass = Object.GetClass();
+				if (!ensure(ObjectClass))
+				{
+					return;
+				}
+				
 				AddAdditionalObject(Object);
 				FReplicatedObjectInfo& ObjectInfo = ReplicationMap.ReplicatedObjects[&Object];
 
 				constexpr bool bLog = false;
-				const FProperty* ResolvedProperty = PropertyChain.ResolveProperty(*Object.GetClass(), bLog);
+				const FProperty* ResolvedProperty = PropertyChain.ResolveProperty(*ObjectClass, bLog);
 				if (!ResolvedProperty || !ConcertSyncCore::PropertyChain::IsReplicatableProperty(*ResolvedProperty))
 				{
 					UE_LOG(LogConcert, Warning, TEXT("Property \"%s\" is not a valid property to assign to object \"%s\"."), *PropertyChain.ToString(), *Object.GetPathName());
@@ -300,8 +306,8 @@ namespace UE::ConcertClientSharedSlate
 				if (!ObjectInfo.PropertySelection.ReplicatedProperties.Contains(PropertyChain))
 				{
 					TArray<FConcertPropertyChain>& Properties = ObjectInfo.PropertySelection.ReplicatedProperties;
+					Private::AddParentProperties(*ObjectClass, PropertyChain, Properties);
 					Properties.Emplace(MoveTemp(PropertyChain));
-					Private::AddParentProperties(*Object.GetClass(), PropertyChain, Properties);
 					// TODO UE-202079: Make sure to append FConcertPropertyChain::InternalContainerPropertyValueName if needed
 				}
 			}
@@ -327,4 +333,4 @@ namespace UE::ConcertClientSharedSlate
 			Extender->ExtendStream(*CurrentObject, Context);
 		}
 	}
-}
+}
\ No newline at end of file
diff --git a/Engine/Plugins/Developer/Concert/ConcertUI/ConcertClientSharedSlate/Source/ConcertClientSharedSlate/Public/Replication/Editor/Model/Extension/IStreamExtensionContext.h b/Engine/Plugins/Developer/Concert/ConcertUI/ConcertClientSharedSlate/Source/ConcertClientSharedSlate/Public/Replication/Editor/Model/Extension/IStreamExtensionContext.h
--- a/Engine/Plugins/Developer/Concert/ConcertUI/ConcertClientSharedSlate/Source/ConcertClientSharedSlate/Public/Replication/Editor/Model/Extension/IStreamExtensionContext.h
+++ b/Engine/Plugins/Developer/Concert/ConcertUI/ConcertClientSharedSlate/Source/ConcertClientSharedSlate/Public/Replication/Editor/Model/Extension/IStreamExtensionContext.h
@@ -21,7 +21,7 @@ namespace UE::ConcertClientSharedSlate
 		 * @param Object Object to add, if not yet present, and assign a property to
 		 * @param PropertyChain Property chain to associate with the object
 		 */
-		virtual void AddPropertyTo(UObject& Object, FConcertPropertyChain&& PropertyChain) = 0;
+		virtual void AddPropertyTo(UObject& Object, FConcertPropertyChain PropertyChain) = 0;
 
 		/**
 		 * Adds an additional object.
