Allow Blueprints to change streams and authority set up by Multi-User

diff --git a/Engine/Plugins/Developer/Concert/ConcertApp/MultiUserClient/Source/MultiUserClient/Private/Replication/Client/ReplicationClient.cpp b/Engine/Plugins/Developer/Concert/ConcertApp/MultiUserClient/Source/MultiUserClient/Private/Replication/Client/ReplicationClient.cpp
--- a/Engine/Plugins/Developer/Concert/ConcertApp/MultiUserClient/Source/MultiUserClient/Private/Replication/Client/ReplicationClient.cpp
+++ b/Engine/Plugins/Developer/Concert/ConcertApp/MultiUserClient/Source/MultiUserClient/Private/Replication/Client/ReplicationClient.cpp
@@ -28,6 +28,11 @@ namespace UE::MultiUserClient
 		, AuthoritySynchronizer(MoveTemp(InAuthoritySynchronizer))
 		, SubmissionWorkflow(MoveTemp(InSubmissionWorkflow))
 		, SubmissionQueue(*SubmissionWorkflow)
+		, ExternalRequestHandler(
+			StreamSynchronizer->GetStreamId(),
+			FGetStreamContent::CreateLambda([this](){ return &StreamSynchronizer->GetServerState(); }),
+			SubmissionQueue
+			)
 		, LocalClientEditModel(CreatePropertySelectionModel(
 			*ClientContentStorage->Stream,
 			ClientContentStorage->Stream->MakeReplicationMapGetterAttribute(),
diff --git a/Engine/Plugins/Developer/Concert/ConcertApp/MultiUserClient/Source/MultiUserClient/Private/Replication/Client/ReplicationClient.h b/Engine/Plugins/Developer/Concert/ConcertApp/MultiUserClient/Source/MultiUserClient/Private/Replication/Client/ReplicationClient.h
--- a/Engine/Plugins/Developer/Concert/ConcertApp/MultiUserClient/Source/MultiUserClient/Private/Replication/Client/ReplicationClient.h
+++ b/Engine/Plugins/Developer/Concert/ConcertApp/MultiUserClient/Source/MultiUserClient/Private/Replication/Client/ReplicationClient.h
@@ -8,6 +8,7 @@
 #include "Replication/Stream/StreamChangeTracker.h"
 #include "Replication/Submission/AutoSubmissionPolicy.h"
 #include "Replication/Submission/ChangeRequestBuilder.h"
+#include "Replication/Submission/External/ExternalClientChangeRequestHandler.h"
 #include "Replication/Submission/ISubmissionWorkflow.h"
 #include "Replication/Submission/Queue/SubmissionQueue.h"
 
@@ -29,7 +30,7 @@ namespace UE::MultiUserClient
 	class ISubmissionWorkflow;
 	
 	/**
-	 * Holds on to shared info about a local or remote client.
+	 * Initializes systems that are common for local or remote clients.
 	 * This class' responsibility is to initialize all systems that exist for the life time of a client in a session.
 	 */
 	class FReplicationClient : public FNoncopyable
@@ -79,6 +80,8 @@ namespace UE::MultiUserClient
 		ISubmissionWorkflow& GetSubmissionWorkflow() { return *SubmissionWorkflow; }
 		/** @return Implements simple game-thread based queue for SubmissionWorkflow. Only one submission can be in progress at any given time. */
 		FSubmissionQueue& GetSubmissionQueue() { return SubmissionQueue; }
+		/** @return Gets the adapter handling requests received by IMultiUserReplication::EnqueuesChanges */
+		FExternalClientChangeRequestHandler& GetExternalRequestHandler() { return ExternalRequestHandler; }
 
 		/** @return The endpoint ID of this client in the Concert session. */
 		const FGuid& GetEndpointId() const { return EndpointId; }
@@ -113,6 +116,8 @@ namespace UE::MultiUserClient
 
 		/** Allows systems to queue pending submissions to SubmissionWorkflow in case a submission is in progress. */
 		FSubmissionQueue SubmissionQueue;
+		/** Enqueues external change requests into the SubmissionQueue. */
+		FExternalClientChangeRequestHandler ExternalRequestHandler;
 		
 		/**
 		 * Used to detect changes made to the client's config by the local editor.
diff --git a/Engine/Plugins/Developer/Concert/ConcertApp/MultiUserClient/Source/MultiUserClient/Private/Replication/MultiUserReplicationManager.cpp b/Engine/Plugins/Developer/Concert/ConcertApp/MultiUserClient/Source/MultiUserClient/Private/Replication/MultiUserReplicationManager.cpp
--- a/Engine/Plugins/Developer/Concert/ConcertApp/MultiUserClient/Source/MultiUserClient/Private/Replication/MultiUserReplicationManager.cpp
+++ b/Engine/Plugins/Developer/Concert/ConcertApp/MultiUserClient/Source/MultiUserClient/Private/Replication/MultiUserReplicationManager.cpp
@@ -3,6 +3,7 @@
 #include "MultiUserReplicationManager.h"
 
 #include "IConcertSyncClient.h"
+#include "Replication/ChangeOperationTypes.h"
 #include "Replication/IConcertClientReplicationManager.h"
 
 #include "Containers/Ticker.h"
@@ -99,11 +100,6 @@ namespace UE::MultiUserClient
 		OnReplicationConnectionStateChangedDelegate.Broadcast(ConnectionState);
 	}
 
-	FGuid FMultiUserReplicationManager::GetMultiUserStreamId() const
-	{
-		return UMultiUserReplicationClientPreset::MultiUserStreamID;
-	}
-
 	const FObjectReplicationMap* FMultiUserReplicationManager::FindReplicationMapForClient(const FGuid& ClientId) const
 	{
 		if (ConnectedState)
@@ -136,6 +132,18 @@ namespace UE::MultiUserClient
 		DiscoveryContainer.RemoveDiscoverer(Discoverer);
 	}
 
+	TSharedRef<IClientChangeOperation> FMultiUserReplicationManager::EnqueueChanges(const FGuid& ClientId, TAttribute<FChangeClientReplicationRequest> SubmissionParams)
+	{
+		if (ConnectedState)
+		{
+			FReplicationClient* ReplicationClient = ConnectedState->ClientManager.FindClient(ClientId);
+			return ReplicationClient
+				? ReplicationClient->GetExternalRequestHandler().HandleRequest(MoveTemp(SubmissionParams))
+				: FExternalClientChangeRequestHandler::MakeFailedOperation(EChangeStreamOperationResult::UnknownClient, EChangeAuthorityOperationResult::UnknownClient);
+		}
+		return FExternalClientChangeRequestHandler::MakeFailedOperation(EChangeStreamOperationResult::NotInSession, EChangeAuthorityOperationResult::NotInSession);
+	}
+
 	FMultiUserReplicationManager::FConnectedState::FConnectedState(TSharedRef<IConcertSyncClient> InClient, FReplicationDiscoveryContainer& InDiscoveryContainer)
 		: ClientManager(InClient, InClient->GetConcertClient()->GetCurrentSession().ToSharedRef(), InDiscoveryContainer)
 	{}
diff --git a/Engine/Plugins/Developer/Concert/ConcertApp/MultiUserClient/Source/MultiUserClient/Private/Replication/MultiUserReplicationManager.h b/Engine/Plugins/Developer/Concert/ConcertApp/MultiUserClient/Source/MultiUserClient/Private/Replication/MultiUserReplicationManager.h
--- a/Engine/Plugins/Developer/Concert/ConcertApp/MultiUserClient/Source/MultiUserClient/Private/Replication/MultiUserReplicationManager.h
+++ b/Engine/Plugins/Developer/Concert/ConcertApp/MultiUserClient/Source/MultiUserClient/Private/Replication/MultiUserReplicationManager.h
@@ -3,7 +3,6 @@
 #pragma once
 
 #include "Client/ReplicationClientManager.h"
-
 #include "IConcertSession.h"
 #include "Replication/IMultiUserReplication.h"
 #include "Replication/Stream/Discovery/ReplicationDiscoveryContainer.h"
@@ -28,11 +27,6 @@ namespace UE::ConcertClientSharedSlate
 
 namespace UE::MultiUserClient
 {
-	namespace Replication
-	{
-		struct FJoinSessionResult;
-	}
-
 	enum class EMultiUserReplicationConnectionState : uint8
 	{
 		Connecting,
@@ -75,11 +69,11 @@ namespace UE::MultiUserClient
 		EMultiUserReplicationConnectionState GetConnectionState() const { return ConnectionState; }
 
 		//~ Begin IMultiUserReplication Interface
-		virtual FGuid GetMultiUserStreamId() const override;
 		virtual const FObjectReplicationMap* FindReplicationMapForClient(const FGuid& ClientId) const override;
 		virtual bool IsReplicatingObject(const FGuid& ClientId, const FSoftObjectPath& ObjectPath) const override;
 		virtual void RegisterReplicationDiscoverer(TSharedRef<IReplicationDiscoverer> Discoverer) override;
 		virtual void RemoveReplicationDiscoverer(const TSharedRef<IReplicationDiscoverer>& Discoverer) override;
+		virtual TSharedRef<IClientChangeOperation> EnqueueChanges(const FGuid& ClientId, TAttribute<FChangeClientReplicationRequest> SubmissionParams) override;
 		//~ End IMultiUserReplication Interface
 
 	private:
diff --git a/Engine/Plugins/Developer/Concert/ConcertApp/MultiUserClient/Source/MultiUserClient/Private/Replication/Submission/External/ClientChangeConversionUtils.cpp b/Engine/Plugins/Developer/Concert/ConcertApp/MultiUserClient/Source/MultiUserClient/Private/Replication/Submission/External/ClientChangeConversionUtils.cpp
new file mode 100644
--- /dev/null
+++ b/Engine/Plugins/Developer/Concert/ConcertApp/MultiUserClient/Source/MultiUserClient/Private/Replication/Submission/External/ClientChangeConversionUtils.cpp
@@ -0,0 +1,174 @@
+// Copyright Epic Games, Inc. All Rights Reserved.
+
+#include "ClientChangeConversionUtils.h"
+
+#include "Replication/ChangeOperationTypes.h"
+#include "Replication/Util/StreamRequestUtils.h"
+
+#include "Algo/RemoveIf.h"
+#include "Misc/Guid.h"
+#include "Replication/Submission/Data/AuthoritySubmission.h"
+#include "Replication/Submission/Data/StreamSubmission.h"
+
+namespace UE::MultiUserClient::ClientChangeConversionUtils
+{
+	namespace Private
+	{
+		static TOptional<FConcertReplication_ChangeStream_PutObject> GeneratePutRequest(
+			const UObject& Object,
+			FPropertyChange&& PropertyChange,
+			const FObjectReplicationMap& ClientStreamContent
+			)
+		{
+			// TODO UE-201166: This step could be simplified if we had an append operation in FConcertReplication_ChangeStream_PutObject
+			const FReplicatedObjectInfo* ExistingObjectInfo = ClientStreamContent.ReplicatedObjects.Find(&Object);
+			if (ExistingObjectInfo)
+			{
+				FReplicatedObjectInfo NewObjectInfo = FReplicatedObjectInfo::Make(Object);
+				TArray<FConcertPropertyChain>& ReplicatedProperties = NewObjectInfo.PropertySelection.ReplicatedProperties;
+				switch (PropertyChange.ChangeType)
+				{
+				case EPropertyChangeType::Put:
+					ReplicatedProperties = MoveTemp(PropertyChange.Properties);
+					break;
+				case EPropertyChangeType::Add:
+					ReplicatedProperties.Append(PropertyChange.Properties);
+					break;
+				case EPropertyChangeType::Remove:
+					ReplicatedProperties.SetNum(Algo::RemoveIf(ReplicatedProperties, [&PropertyChange](const FConcertPropertyChain& Property)
+					{
+						return PropertyChange.Properties.Contains(Property);
+					}));
+					break;
+				default: checkNoEntry(); break;
+				}
+				return FConcertReplication_ChangeStream_PutObject::MakeFromChange(*ExistingObjectInfo, NewObjectInfo);
+			}
+			else
+			{
+				FReplicatedObjectInfo NewObjectInfo = FReplicatedObjectInfo::Make(Object);
+				TArray<FConcertPropertyChain>& ReplicatedProperties = NewObjectInfo.PropertySelection.ReplicatedProperties;
+				switch (PropertyChange.ChangeType)
+				{
+				case EPropertyChangeType::Put:
+				case EPropertyChangeType::Add:
+					ReplicatedProperties = MoveTemp(PropertyChange.Properties);
+					return FConcertReplication_ChangeStream_PutObject::MakeFromInfo(NewObjectInfo);
+				case EPropertyChangeType::Remove:
+					return {};
+				default: checkNoEntry(); break;
+				}
+			}
+				
+			return {};
+		}
+
+		static void BuildStreamChanges(
+			TMap<UObject*, FPropertyChange>&& PropertyChanges,
+			const FGuid& ClientStreamId,
+			const FObjectReplicationMap& ClientStreamContent,
+			FStreamChangelist& StreamChangelist
+			)
+		{
+			for (TPair<UObject*, FPropertyChange>& PropertyChange : PropertyChanges)
+			{
+				UObject* Object = PropertyChange.Key; 
+				if (!Object)
+				{
+					continue;
+				}
+				
+				TOptional<FConcertReplication_ChangeStream_PutObject> PutRequest = GeneratePutRequest(*Object, MoveTemp(PropertyChange.Value), ClientStreamContent);
+				if (PutRequest)
+				{
+					// Gracefully handle APi caller (e.g. Blueprinter) forgetting to specify parent properties,
+					// e.g. if they only specify ["Vector", "X"], we'll add ["Vector"] as well because it is needed for proper replication. 
+					PutRequest->Properties.DiscoverAndAddImplicitParentProperties();
+					
+					StreamChangelist.ObjectsToPut.Add(
+						FObjectInStreamID { ClientStreamId, Object },
+						MoveTemp(*PutRequest)
+					);
+				}
+			}
+		}
+
+		static void BuildAuthorityChanges(TSet<FSoftObjectPath>&& ObjectsToRemove, const FGuid& ClientStreamId, FStreamChangelist& StreamChangelist)
+		{
+			for (FSoftObjectPath& RemovedObject : ObjectsToRemove)
+			{
+				StreamChangelist.ObjectsToRemove.Emplace(
+					FObjectInStreamID{ ClientStreamId, MoveTemp(RemovedObject) }
+				);
+			}
+		}
+	}
+
+	TOptional<FConcertReplication_ChangeStream_Request> Transform(
+		FChangeStreamRequest Request,
+		const FGuid& ClientStreamId,
+		const FObjectReplicationMap& ClientStreamContent
+		)
+	{
+		FStreamChangelist StreamChangelist;
+		Private::BuildStreamChanges(MoveTemp(Request.PropertyChanges), ClientStreamId, ClientStreamContent, StreamChangelist);
+		Private::BuildAuthorityChanges(MoveTemp(Request.ObjectsToRemove), ClientStreamId, StreamChangelist);
+		if (StreamChangelist.ObjectsToPut.IsEmpty() && StreamChangelist.ObjectsToRemove.IsEmpty())
+		{
+			return {};
+		}
+		
+		const bool bNeedsToRegisterStream = ClientStreamContent.IsEmpty();
+		return bNeedsToRegisterStream
+			? StreamRequestUtils::BuildChangeRequest_CreateNewStream(ClientStreamId, StreamChangelist)
+			: StreamRequestUtils::BuildChangeRequest_UpdateExistingStream(MoveTemp(StreamChangelist));
+	}
+
+	TOptional<FConcertReplication_ChangeAuthority_Request> Transform(
+		FChangeAuthorityRequest Request,
+		const FGuid& ClientStreamId
+		)
+	{
+		FConcertReplication_ChangeAuthority_Request Result;
+		for (FSoftObjectPath& TakeAuthority : Request.ObjectsToStartReplicating)
+		{
+			Result.TakeAuthority.Emplace(MoveTemp(TakeAuthority))
+				.StreamIds = { { ClientStreamId } };
+		}
+		for (FSoftObjectPath& ReleaseAuthority : Request.ObjectToStopReplicating)
+		{
+			Result.ReleaseAuthority.Emplace(MoveTemp(ReleaseAuthority))
+				.StreamIds = { { ClientStreamId } };
+		}
+		return Result;
+	}
+
+	EChangeStreamOperationResult Transform(const FSubmitStreamChangesResponse& Response)
+	{
+		switch (Response.ErrorCode)
+		{
+		case EStreamSubmissionErrorCode::Success: return EChangeStreamOperationResult::Success;
+		case EStreamSubmissionErrorCode::NoChange: return EChangeStreamOperationResult::NoChanges; 
+		case EStreamSubmissionErrorCode::Timeout: return EChangeStreamOperationResult::Timeout; 
+		case EStreamSubmissionErrorCode::Cancelled: return EChangeStreamOperationResult::Cancelled;
+		default:
+			checkNoEntry();
+			return EChangeStreamOperationResult::Cancelled;
+		}
+	}
+	
+	EChangeAuthorityOperationResult Transform(const FSubmitAuthorityChangesResponse& Response)
+	{
+		switch (Response.ErrorCode)
+		{
+		case EAuthoritySubmissionResponseErrorCode::Success: return EChangeAuthorityOperationResult::Success;
+		case EAuthoritySubmissionResponseErrorCode::NoChange: return EChangeAuthorityOperationResult::NoChanges; 
+		case EAuthoritySubmissionResponseErrorCode::Timeout: return EChangeAuthorityOperationResult::Timeout; 
+		case EAuthoritySubmissionResponseErrorCode::CancelledDueToStreamUpdate: return EChangeAuthorityOperationResult::CancelledDueToStreamUpdate; 
+		case EAuthoritySubmissionResponseErrorCode::Cancelled: return EChangeAuthorityOperationResult::Cancelled;
+		default:
+			checkNoEntry();
+			return EChangeAuthorityOperationResult::Cancelled;
+		}
+	}
+}
diff --git a/Engine/Plugins/Developer/Concert/ConcertApp/MultiUserClient/Source/MultiUserClient/Private/Replication/Submission/External/ClientChangeConversionUtils.h b/Engine/Plugins/Developer/Concert/ConcertApp/MultiUserClient/Source/MultiUserClient/Private/Replication/Submission/External/ClientChangeConversionUtils.h
new file mode 100644
--- /dev/null
+++ b/Engine/Plugins/Developer/Concert/ConcertApp/MultiUserClient/Source/MultiUserClient/Private/Replication/Submission/External/ClientChangeConversionUtils.h
@@ -0,0 +1,69 @@
+// Copyright Epic Games, Inc. All Rights Reserved.
+
+#pragma once
+
+#include "Replication/ChangeOperationTypes.h"
+#include "Replication/Messages/ChangeAuthority.h"
+#include "Replication/Messages/ChangeStream.h"
+
+#include "Misc/Optional.h"
+
+struct FObjectReplicationMap;
+struct FGuid;
+
+namespace UE::MultiUserClient
+{
+	enum class EChangeStreamOperationResult : uint8;
+	enum class EChangeAuthorityOperationResult : uint8;
+	
+	struct FSubmitAuthorityChangesResponse;
+	struct FSubmitStreamChangesResponse;
+}
+
+/**
+ * Utils for transforming the simplified params in ChangeOperationTypes.h to the ones used by ConcertSyncCore.
+ */
+namespace UE::MultiUserClient::ClientChangeConversionUtils
+{
+	/** Transforms FChangeStreamRequest to FConcertReplication_ChangeStream_Request, if the request is correctly formatted. */
+	TOptional<FConcertReplication_ChangeStream_Request> Transform(
+		FChangeStreamRequest Request,
+		const FGuid& ClientStreamId,
+		const FObjectReplicationMap& ClientStreamContent
+		);
+	
+	/** Transforms FChangeAuthorityRequest to FConcertReplication_ChangeStream_Request, if the request is correctly formatted. */
+	TOptional<FConcertReplication_ChangeAuthority_Request> Transform(
+		FChangeAuthorityRequest Request,
+		const FGuid& ClientStreamId
+		);
+
+
+	/** Extract an error code from FSubmitStreamChangesResponse. */
+	EChangeStreamOperationResult Transform(const FSubmitStreamChangesResponse& Response);
+	/** Extract an error code from FSubmitAuthorityChangesResponse. */
+	EChangeAuthorityOperationResult Transform(const FSubmitAuthorityChangesResponse& Response);
+
+
+	inline TOptional<FConcertReplication_ChangeStream_Request> Transform(
+		TOptional<FChangeStreamRequest> Request,
+		const FGuid& ClientStreamId,
+		const FObjectReplicationMap& ClientStreamContent
+		)
+	{
+		return Request
+			? Transform(MoveTemp(*Request), ClientStreamId, ClientStreamContent)
+			: TOptional<FConcertReplication_ChangeStream_Request>{};
+	}
+	
+	inline TOptional<FConcertReplication_ChangeAuthority_Request> Transform(
+		TOptional<FChangeAuthorityRequest> Request,
+		const FGuid& ClientStreamId
+		)
+	{
+		return Request
+			? Transform(MoveTemp(*Request), ClientStreamId)
+			: TOptional<FConcertReplication_ChangeAuthority_Request>{};
+	}
+}
+
diff --git a/Engine/Plugins/Developer/Concert/ConcertApp/MultiUserClient/Source/MultiUserClient/Private/Replication/Submission/External/ClientChangeOperation.cpp b/Engine/Plugins/Developer/Concert/ConcertApp/MultiUserClient/Source/MultiUserClient/Private/Replication/Submission/External/ClientChangeOperation.cpp
new file mode 100644
--- /dev/null
+++ b/Engine/Plugins/Developer/Concert/ConcertApp/MultiUserClient/Source/MultiUserClient/Private/Replication/Submission/External/ClientChangeOperation.cpp
@@ -0,0 +1,126 @@
+// Copyright Epic Games, Inc. All Rights Reserved.
+
+#include "ClientChangeOperation.h"
+
+#include "ClientChangeConversionUtils.h"
+#include "ConcertLogGlobal.h"
+#include "Replication/Submission/ISubmissionWorkflow.h"
+#include "Replication/Submission/Queue/SubmissionQueue.h"
+
+#include "Replication/Submission/ISubmissionOperation.h"
+
+namespace UE::MultiUserClient
+{
+	TSharedRef<FClientChangeOperation> FClientChangeOperation::StartOperation(
+		const FGuid& InStreamId,
+		FGetStreamContent InGetStreamContentDelegate,
+		FSubmissionQueue& InSubmissionQueue,
+		FOnOperationCompleted InOnOperationCompletedFunc,
+		TAttribute<FChangeClientReplicationRequest> InMakeSubmissionParamsAttribute
+		)
+	{
+		TSharedRef<FClientChangeOperation> Operation = MakeShared<FClientChangeOperation>(
+			InStreamId,
+			MoveTemp(InGetStreamContentDelegate),
+			MoveTemp(InOnOperationCompletedFunc),
+			MoveTemp(InMakeSubmissionParamsAttribute)
+			);
+		
+		// Must fully construct the TSharedRef first because FClientChangeOperation::SetUpOperation requires it for keeping a weak ptr
+		InSubmissionQueue.SubmitNowOrEnqueue_GameThread(*Operation);
+		
+		return Operation;
+	}
+	
+	FClientChangeOperation::FClientChangeOperation(
+		const FGuid& InStreamId,
+		FGetStreamContent InGetStreamContentDelegate,
+		FOnOperationCompleted OnOperationCompletedFunc,
+		TAttribute<FChangeClientReplicationRequest> InMakeSubmissionParamsAttribute
+		)
+		: StreamId(InStreamId)
+		, GetStreamContentDelegate(MoveTemp(InGetStreamContentDelegate))
+		, OnOperationCompletedFunc(MoveTemp(OnOperationCompletedFunc))
+		, MakeSubmissionParamsAttribute(MoveTemp(InMakeSubmissionParamsAttribute))
+	{}
+
+	FClientChangeOperation::~FClientChangeOperation()
+	{
+		const bool bIsAlreadyCompleted = StreamResult && AuthorityResult;
+		if (!StreamResult)
+		{
+			EmplaceStreamResult(EChangeStreamOperationResult::Cancelled);
+		}
+		if (!AuthorityResult)
+		{
+			// Also executes OnOperationCompleted
+			EmplaceAuthorityResult(EChangeAuthorityOperationResult::Cancelled);
+		}
+		if (!bIsAlreadyCompleted)
+		{
+			OperationCompletedPromise.EmplaceValue(FChangeClientReplicationResult{ *StreamResult, *AuthorityResult });
+		}
+	}
+
+	void FClientChangeOperation::PerformSubmission_GameThread(ISubmissionWorkflow& Workflow)
+	{
+		check(!StreamResult && !AuthorityResult);
+
+		FChangeClientReplicationRequest Request = MakeSubmissionParamsAttribute.Get();
+		if (!Request.StreamChangeRequest && !Request.AuthorityChangeRequest)
+		{
+			EmplaceStreamResult(EChangeStreamOperationResult::NoChanges);
+			EmplaceAuthorityResult(EChangeAuthorityOperationResult::NoChanges);
+			CompleteOperation();
+		}
+		else
+		{
+			const TSharedPtr<ISubmissionOperation> SubmissionOperation = Workflow.SubmitChanges({
+				ClientChangeConversionUtils::Transform(MoveTemp(Request.StreamChangeRequest), StreamId, *GetStreamContentDelegate.Execute()),
+				ClientChangeConversionUtils::Transform(MoveTemp(Request.AuthorityChangeRequest), StreamId)
+			});
+			SetUpOperation(SubmissionOperation);
+		}
+	}
+
+	void FClientChangeOperation::SetUpOperation(const TSharedPtr<ISubmissionOperation>& SubmissionOperation)
+	{
+		if (!SubmissionOperation)
+		{
+			UE_LOG(LogConcert, Error, TEXT("Failed to submit request to server. Check log for additional errors."));
+			EmplaceStreamResult(EChangeStreamOperationResult::FailedToSendRequest);
+			EmplaceAuthorityResult(EChangeAuthorityOperationResult::FailedToSendRequest);
+			CompleteOperation();
+			return;
+		}
+
+		// When submission workflow is destroyed, it will finish the below futures.
+		// Keep weak reference to us in case the owning FExternalClientChangeRequestHandler is destroyed before the submission workflow.
+		// This way our code does not rely on whether FExternalClientChangeRequestHandler or ISubmissionWorkflow is destroyed first.
+		SubmissionOperation->OnCompleteStreamChangesFuture()
+			.Next([WeakThis = SharedThis(this).ToWeakPtr()](FSubmitStreamChangesResponse&& Response)
+			{
+				if (const TSharedPtr<FClientChangeOperation> ThisPin = WeakThis.Pin())
+				{
+					ThisPin->EmplaceStreamResult(ClientChangeConversionUtils::Transform(Response));
+				}
+			});
+		SubmissionOperation->OnCompleteAuthorityChangeFuture()
+			.Next([WeakThis = SharedThis(this).ToWeakPtr()](FSubmitAuthorityChangesResponse&& Response)
+			{
+				if (const TSharedPtr<FClientChangeOperation> ThisPin = WeakThis.Pin())
+				{
+					ThisPin->EmplaceAuthorityResult(ClientChangeConversionUtils::Transform(Response));
+				}
+			});
+		SubmissionOperation->OnCompletedOperation()
+			.Next([WeakThis = SharedThis(this).ToWeakPtr()](ESubmissionOperationCompletedCode Response)
+			{
+				if (const TSharedPtr<FClientChangeOperation> ThisPin = WeakThis.Pin())
+				{
+					check(ThisPin->StreamResult && ThisPin->AuthorityResult);
+					ThisPin->CompleteOperation();
+				}
+			});
+	}
+}
diff --git a/Engine/Plugins/Developer/Concert/ConcertApp/MultiUserClient/Source/MultiUserClient/Private/Replication/Submission/External/ClientChangeOperation.h b/Engine/Plugins/Developer/Concert/ConcertApp/MultiUserClient/Source/MultiUserClient/Private/Replication/Submission/External/ClientChangeOperation.h
new file mode 100644
--- /dev/null
+++ b/Engine/Plugins/Developer/Concert/ConcertApp/MultiUserClient/Source/MultiUserClient/Private/Replication/Submission/External/ClientChangeOperation.h
@@ -0,0 +1,123 @@
+// Copyright Epic Games, Inc. All Rights Reserved.
+
+#pragma once
+
+#include "ExternalClientChangeDelegates.h"
+#include "Replication/ChangeOperationTypes.h"
+#include "Replication/IClientChangeOperation.h"
+#include "Replication/Submission/Queue/DeferredSubmitter.h"
+
+#include "Misc/Attribute.h"
+#include "Misc/Optional.h"
+
+namespace UE::MultiUserClient
+{
+	class ISubmissionOperation;
+	
+	/**
+	 * Abstracts an operation invoked externally, i.e. IMultiUserReplication::EnqueueChanges part of our public API.
+	 *
+	 * It stores all the operation data that was passed to EnqueueChanges and handles enqueuing itself to the replication client's FSubmissionQueue.
+	 * FClientChangeOperation is created by FExternalClientChangeRequestHandler, which ensures it is kept alive for at least the lifetime of the
+	 * client the request is made for.
+	 * Since this operation is exposed to the public API caller, it can outlive the client it was created for.
+	 */
+	class FClientChangeOperation
+		: public IClientChangeOperation
+		, public IDeferredSubmitter
+		, public TSharedFromThis<FClientChangeOperation>
+	{
+		template <typename ObjectType, ESPMode Mode>
+		friend class SharedPointerInternals::TIntrusiveReferenceController;
+	public:
+
+		using FOnOperationCompleted = TUniqueFunction<void()>;
+
+		/**
+		 * @param InStreamId The stream ID the client uses.
+		 * @param InGetStreamContentDelegate Delegate for getting the content of the client's stream
+		 * @param InSubmissionQueue The client's queue, which is used to enqueue this operation. The caller, which is FExternalClientChangeRequestHandler,
+		 * keeps the constructed operation alive until either it is executed by the queue or the queue is destroyed. The caller handles dequeuing this
+		 * operation if the client is destroyed.
+		 * @param InOnOperationCompletedFunc Callback function which let's FExternalClientChangeRequestHandler know when the operation is done. Executing
+		 * this function causes FExternalClientChangeRequestHandler to stop reference this operation so it may be destroyed immediately when invoked!
+		 * @param InMakeSubmissionParamsAttribute Callback to build submission params when task is ready to be executed.
+		 */
+		static TSharedRef<FClientChangeOperation> StartOperation(
+			const FGuid& InStreamId,
+			FGetStreamContent InGetStreamContentDelegate,
+			FSubmissionQueue& InSubmissionQueue,
+			FOnOperationCompleted InOnOperationCompletedFunc,
+			TAttribute<FChangeClientReplicationRequest> InMakeSubmissionParamsAttribute
+			);
+		virtual ~FClientChangeOperation() override;
+
+		//~ Begin IClientChangeOperation Interface
+		virtual TFuture<EChangeStreamOperationResult> OnChangeStream() override { return ChangeStreamPromise.GetFuture(); }
+		virtual TFuture<EChangeAuthorityOperationResult> OnChangeAuthority() override { return ChangeAuthorityPromise.GetFuture(); }
+		virtual TFuture<FChangeClientReplicationResult> OnOperationCompleted() override { return OperationCompletedPromise.GetFuture(); }
+		//~ End IClientChangeOperation Interface
+
+		//~ Begin IDeferredSubmitter Interface
+		virtual void PerformSubmission_GameThread(ISubmissionWorkflow& Workflow) override;
+		//~ End IDeferredSubmitter Interface
+
+	private:
+
+		/** The stream ID the client uses. */
+		const FGuid StreamId;
+		/** Delegate for getting the content of the client's stream */
+		const FGetStreamContent GetStreamContentDelegate;
+		
+		/**
+		 * Invokes when this operation completes.
+		 * Causes this instance to no longer be referenced by FExternalClientChangeRequestHandler, which means calling this may invoke the destructor.
+		 * @important Invoking this function may cause this instance to be destroyed immediately.
+		 */
+		const FOnOperationCompleted OnOperationCompletedFunc;
+		/** Callback to build submission params when task is ready to be executed. */
+		const TAttribute<FChangeClientReplicationRequest> MakeSubmissionParamsAttribute;
+
+		TPromise<EChangeStreamOperationResult> ChangeStreamPromise;
+		TPromise<EChangeAuthorityOperationResult> ChangeAuthorityPromise;
+		TPromise<FChangeClientReplicationResult> OperationCompletedPromise;
+
+		/** Set when ChangeStreamPromise is fulfilled. */
+		TOptional<EChangeStreamOperationResult> StreamResult;
+		/** Set when ChangeAuthorityPromise is fulfilled. */
+		TOptional<EChangeAuthorityOperationResult> AuthorityResult;
+
+		FClientChangeOperation(
+			const FGuid& InStreamId,
+			FGetStreamContent InGetStreamContentDelegate,
+			FOnOperationCompleted InOnOperationCompletedFunc,
+			TAttribute<FChangeClientReplicationRequest> InMakeSubmissionParamsAttribute
+			);
+
+		void EmplaceStreamResult(EChangeStreamOperationResult Result)
+		{
+			if (ensure(!StreamResult))
+			{
+				StreamResult = Result;
+				ChangeStreamPromise.EmplaceValue(Result);
+			}
+		}
+		void EmplaceAuthorityResult(EChangeAuthorityOperationResult Result)
+		{
+			if (ensure(!AuthorityResult))
+			{
+				AuthorityResult = Result;
+				ChangeAuthorityPromise.EmplaceValue(Result);
+			}
+		}
+		void CompleteOperation()
+		{
+			OperationCompletedPromise.EmplaceValue(FChangeClientReplicationResult{ *StreamResult, *AuthorityResult });
+			OnOperationCompletedFunc();
+		}
+
+		/** Handles an operation received from ISubmissionWorkflow. */
+		void SetUpOperation(const TSharedPtr<ISubmissionOperation>& SubmissionOperation);
+	};
+}
+
diff --git a/Engine/Plugins/Developer/Concert/ConcertApp/MultiUserClient/Source/MultiUserClient/Private/Replication/Submission/External/ExternalClientChangeDelegates.h b/Engine/Plugins/Developer/Concert/ConcertApp/MultiUserClient/Source/MultiUserClient/Private/Replication/Submission/External/ExternalClientChangeDelegates.h
new file mode 100644
--- /dev/null
+++ b/Engine/Plugins/Developer/Concert/ConcertApp/MultiUserClient/Source/MultiUserClient/Private/Replication/Submission/External/ExternalClientChangeDelegates.h
@@ -0,0 +1,12 @@
+// Copyright Epic Games, Inc. All Rights Reserved.
+
+#pragma once
+
+#include "Delegates/Delegate.h"
+
+struct FObjectReplicationMap;
+
+namespace UE::MultiUserClient
+{
+	DECLARE_DELEGATE_RetVal(const FObjectReplicationMap*, FGetStreamContent);
+}
diff --git a/Engine/Plugins/Developer/Concert/ConcertApp/MultiUserClient/Source/MultiUserClient/Private/Replication/Submission/External/ExternalClientChangeRequestHandler.cpp b/Engine/Plugins/Developer/Concert/ConcertApp/MultiUserClient/Source/MultiUserClient/Private/Replication/Submission/External/ExternalClientChangeRequestHandler.cpp
new file mode 100644
--- /dev/null
+++ b/Engine/Plugins/Developer/Concert/ConcertApp/MultiUserClient/Source/MultiUserClient/Private/Replication/Submission/External/ExternalClientChangeRequestHandler.cpp
@@ -0,0 +1,120 @@
+// Copyright Epic Games, Inc. All Rights Reserved.
+
+#include "ExternalClientChangeRequestHandler.h"
+
+#include "ClientChangeOperation.h"
+#include "Replication/Submission/Queue/SubmissionQueue.h"
+
+#include "Async/Async.h"
+#include "Async/TaskGraphInterfaces.h"
+#include "Async/Future.h"
+
+namespace UE::MultiUserClient
+{
+	TSharedRef<IClientChangeOperation> FExternalClientChangeRequestHandler::MakeFailedOperation(
+		EChangeStreamOperationResult StreamResult,
+		EChangeAuthorityOperationResult AuthorityResult
+		)
+	{
+		class FNotInSessionOperation : public IClientChangeOperation
+		{
+			TPromise<EChangeStreamOperationResult> ChangeStreamPromise;
+			TPromise<EChangeAuthorityOperationResult> ChangeAuthorityPromise;
+			TPromise<FChangeClientReplicationResult> OperationCompletedPromise;
+			
+		public:
+
+			FNotInSessionOperation(EChangeStreamOperationResult StreamResult, EChangeAuthorityOperationResult AuthorityResult)
+				: ChangeStreamPromise(MakeFulfilledPromise<EChangeStreamOperationResult>(StreamResult))
+				, ChangeAuthorityPromise(MakeFulfilledPromise<EChangeAuthorityOperationResult>(AuthorityResult))
+				, OperationCompletedPromise(MakeFulfilledPromise<FChangeClientReplicationResult>(FChangeClientReplicationResult{ StreamResult, AuthorityResult }))
+			{}
+			
+			//~ Begin IClientChangeOperation Interface
+			virtual TFuture<EChangeStreamOperationResult> OnChangeStream() override { return ChangeStreamPromise.GetFuture(); }
+			virtual TFuture<EChangeAuthorityOperationResult> OnChangeAuthority() override { return ChangeAuthorityPromise.GetFuture(); }
+			virtual TFuture<FChangeClientReplicationResult> OnOperationCompleted() override { return OperationCompletedPromise.GetFuture(); }
+			//~ End IClientChangeOperation Interface
+		};
+
+		return MakeShared<FNotInSessionOperation>(StreamResult, AuthorityResult);
+	}
+
+	FExternalClientChangeRequestHandler::FExternalClientChangeRequestHandler(
+		const FGuid& InStreamId, 
+		FGetStreamContent InGetStreamContentDelegate,
+		FSubmissionQueue& InSubmissionQueue
+		)
+		: StreamId(InStreamId)
+		, GetStreamContentDelegate(MoveTemp(InGetStreamContentDelegate))
+		, SubmissionQueue(InSubmissionQueue)
+	{}
+
+	FExternalClientChangeRequestHandler::~FExternalClientChangeRequestHandler()
+	{
+		// We keeps the all operations alive until either they are executed by the queue or the queue is destroyed.
+		// There is no technical reason to dequeue since the queue will not reference us when it is destroyed.
+		// However 1. we'll follow RAII and 2. FSubmissionQueue::SubmitNowOrEnqueue_GameThread's contract dictates we call Dequeue_GameThread when destroyed.
+		// So we'll dequeue it anyway.
+		for (const TPair<FOperationId, TSharedRef<FClientChangeOperation>>& Op : PendingOperations)
+		{
+			// Dequeue_GameThread will ensure already if we're not in game thread so no need to double-check again.
+			SubmissionQueue.Dequeue_GameThread(*Op.Value);
+		}
+	}
+
+	TSharedRef<IClientChangeOperation> FExternalClientChangeRequestHandler::HandleRequest(TAttribute<FChangeClientReplicationRequest> SubmissionParams)
+	{
+		TSharedRef<FClientChangeOperation> Result = FClientChangeOperation::StartOperation(
+			StreamId,
+			GetStreamContentDelegate,
+			SubmissionQueue,
+			[this, OperationId = NextOperationId, WeakToken = LifetimeToken.ToWeakPtr()]()
+			{
+				OnOperationCompleted(OperationId, WeakToken, this);
+			},
+			MoveTemp(SubmissionParams)
+		);
+		PendingOperations.Add(NextOperationId, Result);
+		++NextOperationId;
+		return Result;
+	}
+
+	void FExternalClientChangeRequestHandler::OnOperationCompleted(
+		FOperationId OperationId,
+		TWeakPtr<FToken> WeakToken,
+		FExternalClientChangeRequestHandler* GuardedThis)
+	{
+		const TSharedPtr<FToken> TokenPin = WeakToken.Pin();
+		if (!TokenPin)
+		{
+			// FExternalClientChangeRequestHandler was destroyed.
+			// Happens when external module keeps  operation alive through strong reference and the operation is cancelled as part of ISubmissionWorkflow being destroyed.
+			return;
+		}
+
+		auto Unregister = [OperationId, GuardedThis]()
+		{
+			const int32 NumRemoved = GuardedThis->PendingOperations.Remove(OperationId);
+			const bool bFailedInstantly = OperationId == GuardedThis->NextOperationId;
+			check(NumRemoved > 0 || bFailedInstantly);
+		};
+		
+		if (IsInGameThread())
+		{
+			Unregister();
+		}
+		else
+		{
+			// PendingOperations may only be modified on the game thread to avoid race conditions.
+			AsyncTask(ENamedThreads::GameThread, [WeakToken, Unregister]()
+			{
+				// Latent callback. Must check for destruction of FExternalClientChangeRequestHandler again.
+				if (const TSharedPtr<FToken> TokenPin = WeakToken.Pin())
+				{
+					Unregister();
+				}
+			});
+		}
+	}
+}
diff --git a/Engine/Plugins/Developer/Concert/ConcertApp/MultiUserClient/Source/MultiUserClient/Private/Replication/Submission/External/ExternalClientChangeRequestHandler.h b/Engine/Plugins/Developer/Concert/ConcertApp/MultiUserClient/Source/MultiUserClient/Private/Replication/Submission/External/ExternalClientChangeRequestHandler.h
new file mode 100644
--- /dev/null
+++ b/Engine/Plugins/Developer/Concert/ConcertApp/MultiUserClient/Source/MultiUserClient/Private/Replication/Submission/External/ExternalClientChangeRequestHandler.h
@@ -0,0 +1,95 @@
+// Copyright Epic Games, Inc. All Rights Reserved.
+
+#pragma once
+
+#include "ExternalClientChangeDelegates.h"
+#include "Replication/IToken.h"
+
+#include "Misc/Attribute.h"
+#include "Misc/Guid.h"
+#include "Templates/UnrealTemplate.h"
+
+namespace UE::MultiUserClient
+{
+	class FClientChangeOperation;
+	class FReplicationClientManager;
+	class FSubmissionQueue;
+	class IClientChangeOperation;
+	
+	enum class EChangeAuthorityOperationResult : uint8;
+	enum class EChangeStreamOperationResult : uint8;
+	
+	struct FChangeClientReplicationRequest;
+	
+	/**
+	 * Handles change requests made by external modules via IMultiUserReplication::EnqueueChanges.
+	 * This is instantiated per client. As such, each instance manages change requests made for a specific client.
+	 *
+	 * This class implements the Adapter Pattern: it transforms the public IMultiUserReplication::EnqueueChanges API to our internal API using
+	 * ISubmissionOperation & ISubmissionWorkflow. This way, the public and private APIs can change independently.
+	 */
+	class FExternalClientChangeRequestHandler : public FNoncopyable
+	{
+	public:
+
+		/** @return An operation to return when not in any session. */
+		static TSharedRef<IClientChangeOperation> MakeFailedOperation(EChangeStreamOperationResult StreamResult, EChangeAuthorityOperationResult AuthorityResult);
+
+		/**
+		 * @param InStreamId The stream ID the client uses.
+		 * @param InGetStreamContentDelegate Delegate for getting the content of the client's stream
+		 * @param InSubmissionQueue Used to enqueue the external requests. The owning FMultiUserReplicationManager ensures it outlives the constructed instance.
+		 */
+		FExternalClientChangeRequestHandler(
+			const FGuid& InStreamId,
+			FGetStreamContent InGetStreamContentDelegate,
+			FSubmissionQueue& InSubmissionQueue
+			);
+		~FExternalClientChangeRequestHandler();
+		
+		/** IMultiUserReplication::EnqueueChanges forwards requests to this function. */
+		TSharedRef<IClientChangeOperation> HandleRequest(TAttribute<FChangeClientReplicationRequest> SubmissionParams);
+
+	private:
+
+		/**
+		 * Detects whether the FExternalClientChangeRequestHandler has been / is being destroyed when processing an operation's completion.
+		 *
+		 * Here's what could happen, and is likely:
+		 * - ~FReplicationClient invokes ~FExternalClientChangeRequestHandler 
+		 * - ~FReplicationClient invokes ~ISubmissionWorkflow, which cancels an operation originally started by FExternalClientChangeRequestHandler.
+		 * This problem could be solved in a synchronized way (like unsubscribing our OnOperationCompleted), but here's another case albeit unlikely:
+		 * - ~FReplicationClient invokes ~FExternalClientChangeRequestHandler
+		 * - Just before ~ISubmissionWorkflow is invoked, the UDP unfortunately decides to timeout one of our operations.
+		 * 
+		 * TL;DR: OnOperationCompleted checks whether this is valid before doing anything.
+		 */
+		const TSharedRef<FToken> LifetimeToken = FToken::Make();
+
+		/** The stream ID the client uses. */
+		const FGuid StreamId;
+		/** Delegate for getting the content of the client's stream */
+		const FGetStreamContent GetStreamContentDelegate;
+		
+		/**
+		 * Used to enqueue the external requests.
+		 * The owning FMultiUserReplicationManager ensures it outlives the constructed instance.
+		 */
+		FSubmissionQueue& SubmissionQueue;
+
+		using FOperationId = uint32;
+		
+		FOperationId NextOperationId = 0;
+		/**
+		 * Each time IMultiUserReplication::EnqueueChanges invokes HandleRequest, an operation is added here and returned by HandleRequest.
+		 * The operations invoke OnOperationCompleted when the complete, which removes them from this array.
+		 *
+		 * This property can only be read & written from the game thread. Relevant to OnOperationCompleted being executed as part of a timeout.
+		 */
+		TMap<FOperationId, TSharedRef<FClientChangeOperation>> PendingOperations;
+
+		/** Removes the operation from PendingOperations. */
+		static void OnOperationCompleted(FOperationId OperationId, TWeakPtr<FToken> WeakToken, FExternalClientChangeRequestHandler* GuardedThis);
+	};
+}
+
diff --git a/Engine/Plugins/Developer/Concert/ConcertApp/MultiUserClient/Source/MultiUserClient/Private/Replication/Submission/Queue/SubmissionQueue.h b/Engine/Plugins/Developer/Concert/ConcertApp/MultiUserClient/Source/MultiUserClient/Private/Replication/Submission/Queue/SubmissionQueue.h
--- a/Engine/Plugins/Developer/Concert/ConcertApp/MultiUserClient/Source/MultiUserClient/Private/Replication/Submission/Queue/SubmissionQueue.h
+++ b/Engine/Plugins/Developer/Concert/ConcertApp/MultiUserClient/Source/MultiUserClient/Private/Replication/Submission/Queue/SubmissionQueue.h
@@ -26,6 +26,9 @@ namespace UE::MultiUserClient
 	{
 	public:
 		
+		/**
+		 * @param InWorkflow Used to dispatch requests. The caller ensures it outlives the constructed instance.
+		 */
 		FSubmissionQueue(ISubmissionWorkflow& InWorkflow);
 		virtual ~FSubmissionQueue();
 
@@ -53,7 +56,10 @@ namespace UE::MultiUserClient
 		/** Implements the queue. The head is the next to execute. Used instead of TQueue because TQueue does not support removing. */
 		TDoubleLinkedList<IDeferredSubmitter*> Queue;
 
-		/** The workflow that is being enqueued for. */
+		/**
+		 * The workflow that is being enqueued for.
+		 * Outlives this object.
+		 */
 		ISubmissionWorkflow& Workflow;
 
 		/** Set to true when a task finishes outside of the game thread. Written exactly once by other threads other than game thread. Read by game thread. */
diff --git a/Engine/Plugins/Developer/Concert/ConcertApp/MultiUserClient/Source/MultiUserClient/Public/Replication/ChangeOperationTypes.h b/Engine/Plugins/Developer/Concert/ConcertApp/MultiUserClient/Source/MultiUserClient/Public/Replication/ChangeOperationTypes.h
new file mode 100644
--- /dev/null
+++ b/Engine/Plugins/Developer/Concert/ConcertApp/MultiUserClient/Source/MultiUserClient/Public/Replication/ChangeOperationTypes.h
@@ -0,0 +1,147 @@
+// Copyright Epic Games, Inc. All Rights Reserved.
+
+#pragma once
+
+#include "Replication/Data/ConcertPropertySelection.h"
+
+/**
+ * This public, high-level Multi-User API for replication  effectively encapsulates the lower-level ConcertSyncCore replication API,
+ * @see FConcertReplication_ChangeStream_Request and FConcertReplication_ChangeAuthority_Request.
+ *
+ * The main intentions behind the encapsulation are:
+ *	1. on an architectural level, the users of MU API should not care be made to care about the internal replication protocols
+ *	2. on a functional level, MU internally registers a single stream per client. FConcertReplication_ChangeStream_Request, etc. expose stream IDs directly.
+ *	We want to protect the API users in case anything about the way MU handles stream IDs changes.
+ *
+ * These types can (unfortunately) not be exposed for Runtime Blueprints directly in this module because it is only Editor.
+ * They are wrapped yet again in MultiUserClientLibrary.
+ */
+namespace UE::MultiUserClient
+{
+	/** Result of changing a client's stream */
+	enum class EChangeStreamOperationResult : uint8
+	{
+		/** All changes made or none were requested. */
+		Success,
+		/** No changes were requested to be made */
+		NoChanges,
+
+		// Realistic failures
+		
+		/** The changes were rejected by the server. */
+		Rejected,
+		/** Failed because the request timed out. */
+		Timeout,
+		/** The submission progress was cancelled, e.g. because the target client disconnected. */
+		Cancelled,
+		
+		/** Unlikely. ISubmissionWorkflow validation rejected the change. See log. Could be e.g. that the remote client does not allow remote editing. */
+		FailedToSendRequest,
+
+		// Bad API usage
+		
+		/** No submission took place because the local editor is not in any session. */
+		NotInSession,
+		/** No submission took place because the client ID was invalid.  */
+		UnknownClient,
+
+		/** Not an actual return code. Make sure it's always last. */
+		Count
+	};
+
+	/** Result of changing a client's authority */
+	enum class EChangeAuthorityOperationResult : uint8
+	{
+		/** All changes made or none were requested. */
+		Success,
+		/** No changes were requested to be made */
+		NoChanges,
+
+		// Realistic failures
+
+		/** The changes were rejected by the server. */
+		Rejected,
+		/** Failed because the request timed out. */
+		Timeout,
+		/** The authority change was not submitted because the stream change was unsuccessful. */
+		CancelledDueToStreamUpdate,
+		/** The submission progress was cancelled, e.g. because the target client disconnected. */
+		Cancelled,
+		
+		/** Unlikely. ISubmissionWorkflow validation rejected the change. See log. Could be e.g. that the remote client does not allow remote editing. */
+		FailedToSendRequest,
+
+		// Bad API usage
+		
+		/** No submission took place because the local editor is not in any session. */
+		NotInSession,
+		/** No submission took place because the client ID was invalid.  */
+		UnknownClient,
+
+		/** Not an actual return code. Make sure it's always last. */
+		Count
+	};
+	
+	enum class EPropertyChangeType : uint8
+	{
+		/** Replace all assigned properties with the given properties. */
+		Put,
+		/** Append the given properties, keeping any preexisting properties. */
+		Add,
+		/** Remove the given properties. */
+		Remove,
+		
+		/** Not an actual parameter. Make sure it's always last. */
+		Count
+	};
+	
+	struct FPropertyChange
+	{
+		TArray<FConcertPropertyChain> Properties;
+		EPropertyChangeType ChangeType;
+	};
+
+	/** Params for changing a client's stream. */
+	struct FChangeStreamRequest
+	{
+		/** Property changes to make to objects. */
+		TMap<UObject*, FPropertyChange> PropertyChanges;
+
+		/** Objects that should be unregistered (they will also stop replicating if added here) */
+		TSet<FSoftObjectPath> ObjectsToRemove;
+	};
+
+	/** Params for changing a client's authority (what they're replicating). */
+	struct FChangeAuthorityRequest
+	{
+		/** Objects that should start replicating. The objects must previously have been registered. */
+		TSet<FSoftObjectPath> ObjectsToStartReplicating;
+
+		/** Objects that should stop replicating. */
+		TSet<FSoftObjectPath> ObjectToStopReplicating;
+	};
+
+	/** Params for changing a client's stream or authority. */
+	struct FChangeClientReplicationRequest
+	{
+		/**
+		 * Changes to make to the client's registered objects.
+		 * Performed before AuthorityChangeRequest. If it fails, AuthorityChangeRequest will also fail.
+		 */
+		TOptional<FChangeStreamRequest> StreamChangeRequest;
+		
+		/**
+		 * Changes to make to the objects the client is replicating.
+		 * Fails if StreamChangeRequest fails.
+		 */
+		TOptional<FChangeAuthorityRequest> AuthorityChangeRequest;
+	};
+
+	/** Result of processing a FChangeClientReplicationRequest. */
+	struct FChangeClientReplicationResult
+	{
+		EChangeStreamOperationResult StreamChangeResult;
+		EChangeAuthorityOperationResult AuthorityChangeResult;
+	};
+}
+
diff --git a/Engine/Plugins/Developer/Concert/ConcertApp/MultiUserClient/Source/MultiUserClient/Public/Replication/IClientChangeOperation.h b/Engine/Plugins/Developer/Concert/ConcertApp/MultiUserClient/Source/MultiUserClient/Public/Replication/IClientChangeOperation.h
new file mode 100644
--- /dev/null
+++ b/Engine/Plugins/Developer/Concert/ConcertApp/MultiUserClient/Source/MultiUserClient/Public/Replication/IClientChangeOperation.h
@@ -0,0 +1,45 @@
+// Copyright Epic Games, Inc. All Rights Reserved.
+
+#pragma once
+
+#include "ChangeOperationTypes.h"
+
+#include "Async/Future.h"
+
+namespace UE::MultiUserClient
+{
+	/**
+	 * Abstracts an operation that is changing a client's stream, authority, or both.
+	 * They are the result of IMultiUserReplication::EnqueueChanges.
+	 *
+	 * IClientChangeOperation wraps ISubmissionOperation, which is internal to Multi-User.
+	 * This allows the two interface to change independently.
+	 */
+	class MULTIUSERCLIENT_API IClientChangeOperation
+	{
+	public:
+
+		/**
+		 * Completes when the operation of changing streams has completed.
+		 * @note This can be called at most once; subsequent calls result in an unset future.
+		 * @note This future can execute on any thread!
+		 */
+		virtual TFuture<EChangeStreamOperationResult> OnChangeStream() = 0;
+		
+		/**
+		 * Completes when the operation of changing authority has completed.
+		 * @note This can be called at most once; subsequent calls result in an unset future.
+		 * @note This future can execute on any thread!
+		 */
+		virtual TFuture<EChangeAuthorityOperationResult> OnChangeAuthority() = 0;
+		
+		/**
+		 * Completes all sub-operations have completed.
+		 * @note This can be called at most once; subsequent calls result in an unset future.
+		 * @note This future can execute on any thread!
+		 */
+		virtual TFuture<FChangeClientReplicationResult> OnOperationCompleted() = 0;
+		
+		virtual ~IClientChangeOperation() = default;
+	};
+}
diff --git a/Engine/Plugins/Developer/Concert/ConcertApp/MultiUserClient/Source/MultiUserClient/Public/Replication/IMultiUserReplication.h b/Engine/Plugins/Developer/Concert/ConcertApp/MultiUserClient/Source/MultiUserClient/Public/Replication/IMultiUserReplication.h
--- a/Engine/Plugins/Developer/Concert/ConcertApp/MultiUserClient/Source/MultiUserClient/Public/Replication/IMultiUserReplication.h
+++ b/Engine/Plugins/Developer/Concert/ConcertApp/MultiUserClient/Source/MultiUserClient/Public/Replication/IMultiUserReplication.h
@@ -2,6 +2,7 @@
 
 #pragma once
 
+#include "Misc/Attribute.h"
 #include "Templates/SharedPointer.h"
 
 struct FGuid;
@@ -9,21 +10,16 @@ struct FObjectReplicationMap;
 
 namespace UE::MultiUserClient
 {
+	class IClientChangeOperation;
 	class IReplicationDiscoverer;
 	
+	struct FChangeClientReplicationRequest;
+	
 	/** Interface for interacting with Multi-User replication, which uses the Concert replication system. */
 	class MULTIUSERCLIENT_API IMultiUserReplication
 	{
 	public:
 
-		/**
-		 * The server keeps track of multiple streams per client and the ID is used to identify streams by client.
-		 * MU only uses one stream per client and this is its ID.
-		 * 
-		 * @return The stream ID that every MU client uses to register its replication stream with the server.
-		 */
-		virtual FGuid GetMultiUserStreamId() const = 0;
-
 		/**
 		 * @return Gets the last known server map of objects registered for replication for a given client.
 		 * This server state is regularly polled whilst the local client state should always be in synch.*/
@@ -42,6 +38,19 @@ namespace UE::MultiUserClient
 		/** Unregisters a previously registered discoverer */
 		virtual void RemoveReplicationDiscoverer(const TSharedRef<IReplicationDiscoverer>& Discoverer) = 0;
 
+		/**
+		 * Enqueues a request for changing a client's stream and authority.
+		 * The request is enqueued with the other requests that Multi-User might have ongoing already (like those triggered by the UI).
+		 *
+		 * A stream is the mapping of objects to properties.
+		 * The authority state specifies which of the registered objects should actually be sending data.
+		 * The stream change is requested first and is followed by the authority change.
+		 *
+		 * @param ClientId The client for which to change authority
+		 * @param SubmissionParams Once the request is ready to be sent to the server, this attribute is used to generate the change request
+		 */
+		virtual TSharedRef<IClientChangeOperation> EnqueueChanges(const FGuid& ClientId, TAttribute<FChangeClientReplicationRequest> SubmissionParams) = 0;
+
 		virtual ~IMultiUserReplication() = default;
 	};
 }
diff --git a/Engine/Plugins/Developer/Concert/ConcertApp/MultiUserClient/Source/MultiUserClientLibrary/Private/Replication/Async/ChangeClientAsyncAction.cpp b/Engine/Plugins/Developer/Concert/ConcertApp/MultiUserClient/Source/MultiUserClientLibrary/Private/Replication/Async/ChangeClientAsyncAction.cpp
new file mode 100644
--- /dev/null
+++ b/Engine/Plugins/Developer/Concert/ConcertApp/MultiUserClient/Source/MultiUserClientLibrary/Private/Replication/Async/ChangeClientAsyncAction.cpp
@@ -0,0 +1,62 @@
+// Copyright Epic Games, Inc. All Rights Reserved.
+
+#include "Replication/Async/ChangeClientAsyncAction.h"
+
+#if WITH_CONCERT
+#include "IMultiUserClientModule.h"
+#include "Replication/IClientChangeOperation.h"
+#include "Replication/IMultiUserReplication.h"
+
+#include "Async/Async.h"
+#include "Async/TaskGraphInterfaces.h"
+#endif
+
+UChangeClientAsyncAction* UChangeClientAsyncAction::ChangeClient(const FGuid& ClientId, FMultiUserChangeClientReplicationRequest Request)
+{
+	UChangeClientAsyncAction* AsyncAction = NewObject<UChangeClientAsyncAction>();
+	AsyncAction->ClientId = ClientId;
+	AsyncAction->ReplicationRequest = MoveTemp(Request);
+	return AsyncAction;
+}
+
+void UChangeClientAsyncAction::Activate()
+{
+	Super::Activate();
+	
+#if WITH_CONCERT
+	UE::MultiUserClient::IMultiUserReplication* MultiUserModule = IMultiUserClientModule::IsAvailable()
+		? IMultiUserClientModule::Get().GetReplication()
+		: nullptr;
+	if (MultiUserModule)
+	{
+		MultiUserModule->EnqueueChanges(ClientId, UE::MultiUserClientLibrary::Transform(MoveTemp(ReplicationRequest)))
+			->OnOperationCompleted().Next([WeakThis = TWeakObjectPtr<UChangeClientAsyncAction>(this)](UE::MultiUserClient::FChangeClientReplicationResult&& Result)
+			{
+				if (IsInGameThread())
+				{
+					// Do not call IsValid it outside the game thread!
+					if (WeakThis.IsValid())
+					{
+						WeakThis->OnCompleted.Broadcast(UE::MultiUserClientLibrary::Transform(MoveTemp(Result)));
+					}
+				}
+				else
+				{
+					// Result is on stack... do not reference in latent operation.
+					UE::MultiUserClient::FChangeClientReplicationResult ResultCopy = MoveTemp(Result);
+					AsyncTask(ENamedThreads::GameThread, [WeakThis, Result = MoveTemp(ResultCopy)]() mutable
+					{
+						if (WeakThis.IsValid())
+						{
+							WeakThis->OnCompleted.Broadcast(UE::MultiUserClientLibrary::Transform(MoveTemp(Result)));
+						}
+					});
+				}
+			});
+	}
+	else
+#endif
+	{
+		OnCompleted.Broadcast({ EMultiUserChangeStreamOperationResult::NotAvailable, EMultiUserChangeAuthorityOperationResult::NotAvailable });
+	}
+}
diff --git a/Engine/Plugins/Developer/Concert/ConcertApp/MultiUserClient/Source/MultiUserClientLibrary/Private/Replication/Async/ChangeClientBlueprintParams.cpp b/Engine/Plugins/Developer/Concert/ConcertApp/MultiUserClient/Source/MultiUserClientLibrary/Private/Replication/Async/ChangeClientBlueprintParams.cpp
new file mode 100644
--- /dev/null
+++ b/Engine/Plugins/Developer/Concert/ConcertApp/MultiUserClient/Source/MultiUserClientLibrary/Private/Replication/Async/ChangeClientBlueprintParams.cpp
@@ -0,0 +1,71 @@
+// Copyright Epic Games, Inc. All Rights Reserved.
+
+#include "Replication/Async/ChangeClientBlueprintParams.h"
+
+namespace UE::MultiUserClientLibrary
+{
+#if WITH_CONCERT
+	EMultiUserChangeStreamOperationResult Transform(MultiUserClient::EChangeStreamOperationResult Data)
+	{
+		static_assert(static_cast<int32>(MultiUserClient::EChangeStreamOperationResult::Count) == 8, "Update this location if you changed the enum");
+		return static_cast<EMultiUserChangeStreamOperationResult>(Data);
+	}
+	EMultiUserChangeAuthorityOperationResult Transform(MultiUserClient::EChangeAuthorityOperationResult Data)
+	{
+		static_assert(static_cast<int32>(MultiUserClient::EChangeAuthorityOperationResult::Count) == 9, "Update this location if you changed the enum");
+		return static_cast<EMultiUserChangeAuthorityOperationResult>(Data);
+	}
+	FMultiUserChangeClientReplicationResult Transform(MultiUserClient::FChangeClientReplicationResult Data)
+	{
+		return { Transform(Data.StreamChangeResult), Transform(Data.AuthorityChangeResult) };
+	}
+	
+	MultiUserClient::EPropertyChangeType Transform(EMultiUserPropertyChangeType Data)
+	{
+		static_assert(static_cast<int32>(MultiUserClient::EPropertyChangeType::Count) == 3, "Update this location if you changed the enum");
+		return static_cast<MultiUserClient::EPropertyChangeType>(Data);
+	}
+	MultiUserClient::FPropertyChange Transform(FMultiUserPropertyChange Data)
+	{
+		MultiUserClient::FPropertyChange Result { .ChangeType =  Transform(MoveTemp(Data.ChangeType)) };
+		
+		if (!Data.Properties.IsEmpty())
+		{
+			Result.Properties.Reserve(Data.Properties.Num());
+			for (FConcertPropertyChainWrapper& Wrapper : Data.Properties)
+			{
+				Result.Properties.Emplace(MoveTemp(Wrapper.PropertyChain));
+			}
+		}
+		
+		return Result;
+	}
+	
+	MultiUserClient::FChangeStreamRequest Transform(FMultiUserChangeStreamRequest Data)
+	{
+		MultiUserClient::FChangeStreamRequest Result { .ObjectsToRemove = MoveTemp(Data.ObjectsToRemove) };
+		
+		if (!Data.PropertyChanges.IsEmpty())
+		{
+			Result.PropertyChanges.Reserve(Data.PropertyChanges.Num());
+			for (TPair<UObject*, FMultiUserPropertyChange>& Change : Data.PropertyChanges)
+			{
+				Result.PropertyChanges.Emplace(Change.Key, Transform(MoveTemp(Change.Value)));
+			}
+		}
+		
+		return Result;
+	}
+	MultiUserClient::FChangeAuthorityRequest Transform(FMultiUserChangeAuthorityRequest Data)
+	{
+		return { MoveTemp(Data.ObjectsToStartReplicating), MoveTemp(Data.ObjectToStopReplicating) };
+	}
+	MultiUserClient::FChangeClientReplicationRequest Transform(FMultiUserChangeClientReplicationRequest Data)
+	{
+		return {
+			Data.StreamChangeRequest.IsEmpty() ? TOptional<MultiUserClient::FChangeStreamRequest>{} : Transform(MoveTemp(Data.StreamChangeRequest)),
+			Data.AuthorityChangeRequest.IsEmpty() ? TOptional<MultiUserClient::FChangeAuthorityRequest>{} : Transform(MoveTemp(Data.AuthorityChangeRequest)),
+		};
+	}
+#endif
+}
\ No newline at end of file
diff --git a/Engine/Plugins/Developer/Concert/ConcertApp/MultiUserClient/Source/MultiUserClientLibrary/Public/Replication/Async/ChangeClientAsyncAction.h b/Engine/Plugins/Developer/Concert/ConcertApp/MultiUserClient/Source/MultiUserClientLibrary/Public/Replication/Async/ChangeClientAsyncAction.h
new file mode 100644
--- /dev/null
+++ b/Engine/Plugins/Developer/Concert/ConcertApp/MultiUserClient/Source/MultiUserClientLibrary/Public/Replication/Async/ChangeClientAsyncAction.h
@@ -0,0 +1,49 @@
+// Copyright Epic Games, Inc. All Rights Reserved.
+
+#pragma once
+
+#include "Kismet/BlueprintAsyncActionBase.h"
+#include "ChangeClientBlueprintParams.h"
+#include "ChangeClientAsyncAction.generated.h"
+
+/**
+ * Enqueues a request for changing a specific client's stream and / or authority.
+ * This class exposes IMultiUserReplication::EnqueueChanges.
+ */
+UCLASS()
+class MULTIUSERCLIENTLIBRARY_API UChangeClientAsyncAction : public UBlueprintAsyncActionBase
+{
+	GENERATED_BODY()
+public:
+
+	DECLARE_DYNAMIC_MULTICAST_DELEGATE_OneParam(FOnChangeOperationCompleted, const FMultiUserChangeClientReplicationResult&, Response);
+
+	/** Event that triggers when the operation completes. */
+	UPROPERTY(BlueprintAssignable, Category = "Multi-user")
+	FOnChangeOperationCompleted OnCompleted;
+	
+	/**
+	 * Allows changing of a client's stream and authority.
+	 *
+	 * A stream is a mapping of objects to properties, i.e. registered objects. 
+	 * A client must request to take authority over objects they have registered previously.
+	 * If there is no conflict with other clients, the server approves the request and the client starts replicating automatically.
+	 *
+	 * @param ClientId The client for which this request is to be sent; can be the local or a remote client.
+	 * @param Request The request specifying a change to stream, authority, or both.
+	 */
+	UFUNCTION(BlueprintCallable, Category = "Multi-user", meta = (BlueprintInternalUseOnly = "true", Keywords = "Change Client Stream Authority Ownership Replication"))
+	static UChangeClientAsyncAction* ChangeClient(const FGuid& ClientId, FMultiUserChangeClientReplicationRequest Request);
+
+	//~ Begin UBlueprintAsyncActionBase Interface
+	virtual void Activate() override;
+	//~ End UBlueprintAsyncActionBase Interface
+
+private:
+
+	UPROPERTY(Transient)
+	FGuid ClientId; 
+	
+	UPROPERTY(Transient)
+	FMultiUserChangeClientReplicationRequest ReplicationRequest;
+};
diff --git a/Engine/Plugins/Developer/Concert/ConcertApp/MultiUserClient/Source/MultiUserClientLibrary/Public/Replication/Async/ChangeClientBlueprintParams.h b/Engine/Plugins/Developer/Concert/ConcertApp/MultiUserClient/Source/MultiUserClientLibrary/Public/Replication/Async/ChangeClientBlueprintParams.h
new file mode 100644
--- /dev/null
+++ b/Engine/Plugins/Developer/Concert/ConcertApp/MultiUserClient/Source/MultiUserClientLibrary/Public/Replication/Async/ChangeClientBlueprintParams.h
@@ -0,0 +1,187 @@
+// Copyright Epic Games, Inc. All Rights Reserved.
+
+#pragma once
+
+#include "ConcertPropertyChainWrapper.h"
+#include "UObject/SoftObjectPath.h"
+
+#if WITH_CONCERT
+#include "Replication/ChangeOperationTypes.h"
+#endif
+
+#include "ChangeClientBlueprintParams.generated.h"
+
+/** Result of changing a client's stream */
+UENUM(BlueprintType)
+enum class EMultiUserChangeStreamOperationResult : uint8
+{
+	/** All changes made or none were requested. */
+	Success,
+	/** No changes were requested to be made */
+	NoChanges,
+
+	// Realistic failures
+	
+	/** The changes were rejected by the server. */
+	Rejected,
+	/** Failed because the request timed out. */
+	Timeout,
+	/** The submission progress was cancelled, e.g. because the target client disconnected. */
+	Cancelled,
+	
+	/** Unlikely. ISubmissionWorkflow validation rejected the change. See log. Could be e.g. that the remote client does not allow remote editing. */
+	FailedToSendRequest,
+
+	// Bad API usage
+	
+	/** No submission took place because the local editor is not in any session. */
+	NotInSession,
+	/** No submission took place because the client ID was invalid.  */
+	UnknownClient,
+	/** The MultiUserClient module is not available. Usually when you're in a runtime build but the operation can only be run in editor builds. */
+	NotAvailable
+};
+
+/** Result of changing a client's authority */
+UENUM(BlueprintType)
+enum class EMultiUserChangeAuthorityOperationResult : uint8
+{
+	/** All changes made or none were requested. */
+	Success,
+	/** No changes were requested to be made */
+	NoChanges,
+
+	// Realistic failures
+
+	/** The changes were rejected by the server. */
+	Rejected,
+	/** Failed because the request timed out. */
+	Timeout,
+	/** The authority change was not submitted because the stream change was unsuccessful. */
+	CancelledDueToStreamUpdate,
+	/** The submission progress was cancelled, e.g. because the target client disconnected. */
+	Cancelled,
+	
+	/** Unlikely. ISubmissionWorkflow validation rejected the change. See log. Could be e.g. that the remote client does not allow remote editing. */
+	FailedToSendRequest,
+
+	// Bad API usage
+	
+	/** No submission took place because the local editor is not in any session. */
+	NotInSession,
+	/** No submission took place because the client ID was invalid.  */
+	UnknownClient,
+	/** The MultiUserClient module is not available. Usually when you're in a runtime build but the operation can only be run in editor builds. */
+	NotAvailable
+};
+
+UENUM(BlueprintType)
+enum class EMultiUserPropertyChangeType : uint8
+{
+	/** Replace all assigned properties with the given properties. */
+	Put,
+	/** Append the given properties, keeping any preexisting properties. */
+	Add,
+	/** Remove the given properties. */
+	Remove
+};
+
+USTRUCT(BlueprintType)
+struct FMultiUserPropertyChange
+{
+	GENERATED_BODY()
+
+	/** The properties for the object. See ChangeType for how they are used */
+	UPROPERTY(BlueprintReadWrite, Category = "Multi-user")
+	TArray<FConcertPropertyChainWrapper> Properties;
+
+	/** How you want Properties to be applied to the object. */
+	UPROPERTY(BlueprintReadWrite, Category = "Multi-user")
+	EMultiUserPropertyChangeType ChangeType = EMultiUserPropertyChangeType::Add;
+};
+
+/** Params for changing a client's stream. */
+USTRUCT(BlueprintType)
+struct FMultiUserChangeStreamRequest
+{
+	GENERATED_BODY()
+	
+	/** Property changes to make to objects. */
+	UPROPERTY(BlueprintReadWrite, Category = "Multi-user")
+	TMap<UObject*, FMultiUserPropertyChange> PropertyChanges;
+
+	/** Objects that should be unregistered (they will also stop replicating if added here) */
+	UPROPERTY(BlueprintReadWrite, Category = "Multi-user")
+	TSet<FSoftObjectPath> ObjectsToRemove;
+	
+	bool IsEmpty() const { return PropertyChanges.IsEmpty() && ObjectsToRemove.IsEmpty(); }
+};
+
+/** Params for changing a client's authority (what they're replicating). */
+USTRUCT(BlueprintType)
+struct FMultiUserChangeAuthorityRequest
+{
+	GENERATED_BODY()
+	
+	/** Objects that should start replicating. The objects must previously have been registered. */
+	UPROPERTY(BlueprintReadWrite, Category = "Multi-user")
+	TSet<FSoftObjectPath> ObjectsToStartReplicating;
+
+	/** Objects that should stop replicating. */
+	UPROPERTY(BlueprintReadWrite, Category = "Multi-user")
+	TSet<FSoftObjectPath> ObjectToStopReplicating;
+
+	bool IsEmpty() const { return ObjectsToStartReplicating.IsEmpty() && ObjectToStopReplicating.IsEmpty(); }
+};
+
+/** Params for changing a client's stream or authority. */
+USTRUCT(BlueprintType)
+struct FMultiUserChangeClientReplicationRequest
+{
+	GENERATED_BODY()
+	
+	/**
+	 * Changes to make to the client's registered objects.
+	 * Performed before AuthorityChangeRequest. If it fails, AuthorityChangeRequest will also fail.
+	 */
+	UPROPERTY(BlueprintReadWrite, Category = "Multi-user")
+	FMultiUserChangeStreamRequest StreamChangeRequest;
+	
+	/**
+	 * Changes to make to the objects the client is replicating.
+	 * Fails if StreamChangeRequest fails.
+	 */
+	UPROPERTY(BlueprintReadWrite, Category = "Multi-user")
+	FMultiUserChangeAuthorityRequest AuthorityChangeRequest;
+};
+
+/** Result of processing a FChangeClientReplicationRequest. */
+USTRUCT(BlueprintType)
+struct FMultiUserChangeClientReplicationResult
+{
+	GENERATED_BODY()
+
+	/** The result of changing streams. */
+	UPROPERTY(BlueprintReadOnly, Category = "Multi-user")
+	EMultiUserChangeStreamOperationResult StreamChangeResult = EMultiUserChangeStreamOperationResult::Success;
+
+	/** The result of changing authority. Fails if StreamChangeResult fails. */
+	UPROPERTY(BlueprintReadOnly, Category = "Multi-user")
+	EMultiUserChangeAuthorityOperationResult AuthorityChangeResult = EMultiUserChangeAuthorityOperationResult::Success;
+};
+
+namespace UE::MultiUserClientLibrary
+{
+#if WITH_CONCERT
+	MULTIUSERCLIENTLIBRARY_API EMultiUserChangeStreamOperationResult Transform(MultiUserClient::EChangeStreamOperationResult Data);
+	MULTIUSERCLIENTLIBRARY_API EMultiUserChangeAuthorityOperationResult Transform(MultiUserClient::EChangeAuthorityOperationResult Data);
+	MULTIUSERCLIENTLIBRARY_API FMultiUserChangeClientReplicationResult Transform(MultiUserClient::FChangeClientReplicationResult Data);
+	
+	MULTIUSERCLIENTLIBRARY_API MultiUserClient::EPropertyChangeType Transform(EMultiUserPropertyChangeType Data);
+	MULTIUSERCLIENTLIBRARY_API MultiUserClient::FPropertyChange Transform(FMultiUserPropertyChange Data);
+	
+	MULTIUSERCLIENTLIBRARY_API MultiUserClient::FChangeStreamRequest Transform(FMultiUserChangeStreamRequest Data);
+	MULTIUSERCLIENTLIBRARY_API MultiUserClient::FChangeAuthorityRequest Transform(FMultiUserChangeAuthorityRequest Data);
+	MULTIUSERCLIENTLIBRARY_API MultiUserClient::FChangeClientReplicationRequest Transform(FMultiUserChangeClientReplicationRequest Data);
+#endif
+}
\ No newline at end of file
diff --git a/Engine/Plugins/Developer/Concert/ConcertSync/ConcertSyncCore/Source/ConcertSyncCore/Private/Replication/Data/ConcertPropertySelection.cpp b/Engine/Plugins/Developer/Concert/ConcertSync/ConcertSyncCore/Source/ConcertSyncCore/Private/Replication/Data/ConcertPropertySelection.cpp
--- a/Engine/Plugins/Developer/Concert/ConcertSync/ConcertSyncCore/Source/ConcertSyncCore/Private/Replication/Data/ConcertPropertySelection.cpp
+++ b/Engine/Plugins/Developer/Concert/ConcertSync/ConcertSyncCore/Source/ConcertSyncCore/Private/Replication/Data/ConcertPropertySelection.cpp
@@ -190,6 +190,18 @@ FProperty* FConcertPropertyChain::ResolveProperty(UStruct& Class, bool bLogOnFai
 	return UE::ConcertSyncCore::PropertyChain::ResolveProperty(Class, *this, bLogOnFail);
 }
 
+FConcertPropertyChain FConcertPropertyChain::GetParent() const
+{
+	if (PathToProperty.Num() <= 1)
+	{
+		return {};
+	}
+	
+	FConcertPropertyChain Result = *this;
+	Result.PathToProperty.RemoveAt(Result.PathToProperty.Num() - 1);
+	return Result;
+}
+
 FString FConcertPropertyChain::ToString(EToStringMethod Method) const
 {
 	switch (Method)
@@ -219,6 +231,21 @@ bool FConcertPropertySelection::Includes(const FConcertPropertySelection& Other)
 	});
 }
 
+void FConcertPropertySelection::DiscoverAndAddImplicitParentProperties()
+{
+	// Iterate in reverse to ignore the properties we're about to add
+	for (int32 Index = ReplicatedProperties.Num() - 1; ReplicatedProperties.IsValidIndex(Index); ++Index)
+	{
+		FConcertPropertyChain Current = ReplicatedProperties[Index].GetParent();
+		while (!Current.IsEmpty())
+		{
+			// This part makes the algorithm O(n^2) but could be O(n) with a TSet. Suboptimal, but it's simple... it's quite unlikely this will cause us performance trouble
+			ReplicatedProperties.AddUnique(Current);
+			Current = Current.GetParent();
+		}
+	}
+}
+
 bool FConcertPropertySelection::EnumeratePropertyOverlaps(
 	TConstArrayView<FConcertPropertyChain> First,
 	TConstArrayView<FConcertPropertyChain> Second,
diff --git a/Engine/Plugins/Developer/Concert/ConcertSync/ConcertSyncCore/Source/ConcertSyncCore/Public/Replication/Data/ConcertPropertySelection.h b/Engine/Plugins/Developer/Concert/ConcertSync/ConcertSyncCore/Source/ConcertSyncCore/Public/Replication/Data/ConcertPropertySelection.h
--- a/Engine/Plugins/Developer/Concert/ConcertSync/ConcertSyncCore/Source/ConcertSyncCore/Public/Replication/Data/ConcertPropertySelection.h
+++ b/Engine/Plugins/Developer/Concert/ConcertSync/ConcertSyncCore/Source/ConcertSyncCore/Public/Replication/Data/ConcertPropertySelection.h
@@ -67,6 +67,7 @@ struct CONCERTSYNCCORE_API FConcertPropertyChain
 	/** @return Attempts to resolve this property given the class */
 	FProperty* ResolveProperty(UStruct& Class, bool bLogOnFail = true) const;
 
+	FConcertPropertyChain GetParent() const;
 	const TArray<FName>& GetPathToProperty() const { return PathToProperty; }
 
 	enum class EToStringMethod
@@ -187,6 +188,14 @@ struct CONCERTSYNCCORE_API FConcertPropertySelection
 
 	/** @return Whether this includes all properties of Other */
 	bool Includes(const FConcertPropertySelection& Other) const;
+
+	/**
+	 * Adds all parent properties if they are missing.
+	 * 
+	 * Suppose ReplicatedProperties = { ["Vector", "X"] }.
+	 * After execution, it would be { ["Vector", "X"], ["Vector"] }
+	 */
+	void DiscoverAndAddImplicitParentProperties();
 	
 	/**
 	 * Determines all properties that overlap.
diff --git a/Engine/Plugins/Developer/Concert/ConcertSync/ConcertSyncCore/Source/ConcertSyncCore/Public/Replication/Data/ObjectReplicationMap.h b/Engine/Plugins/Developer/Concert/ConcertSync/ConcertSyncCore/Source/ConcertSyncCore/Public/Replication/Data/ObjectReplicationMap.h
--- a/Engine/Plugins/Developer/Concert/ConcertSync/ConcertSyncCore/Source/ConcertSyncCore/Public/Replication/Data/ObjectReplicationMap.h
+++ b/Engine/Plugins/Developer/Concert/ConcertSync/ConcertSyncCore/Source/ConcertSyncCore/Public/Replication/Data/ObjectReplicationMap.h
@@ -3,6 +3,7 @@
 #pragma once
 
 #include "Replication/Data/ConcertPropertySelection.h"
+#include "UObject/Object.h"
 #include "ObjectReplicationMap.generated.h"
 
 USTRUCT()
@@ -16,6 +17,12 @@ struct FReplicatedObjectInfo
 	UPROPERTY()
 	FConcertPropertySelection PropertySelection;
 
+	// Use static factories instead of user-provided constructors to continue allowing brace-initialization
+	static FReplicatedObjectInfo Make(const UObject& Object)
+	{
+		return { Object.GetClass() };
+	}
+
 	/** @return Whether this data is valid for sending to the server. */
 	bool IsValidForSendingToServer() const { return ClassPath.IsValid() && !PropertySelection.ReplicatedProperties.IsEmpty(); }
 	
