Expose replication freuquency to Multi User UI

diff --git a/Engine/Plugins/Developer/Concert/ConcertApp/MultiUserClient/Source/MultiUserClient/Private/Replication/Client/ReplicationClient.h b/Engine/Plugins/Developer/Concert/ConcertApp/MultiUserClient/Source/MultiUserClient/Private/Replication/Client/ReplicationClient.h
--- a/Engine/Plugins/Developer/Concert/ConcertApp/MultiUserClient/Source/MultiUserClient/Private/Replication/Client/ReplicationClient.h
+++ b/Engine/Plugins/Developer/Concert/ConcertApp/MultiUserClient/Source/MultiUserClient/Private/Replication/Client/ReplicationClient.h
@@ -77,6 +77,7 @@ namespace UE::MultiUserClient
 		FAuthorityChangeTracker& GetAuthorityDiffer() { return LocalAuthorityDiffer; }
 
 		/** @return The object with which you can submit changes. You can use it to listen to general events. If you want to submit changes, you should prefer to use the SubmissionQueue. */
+		const ISubmissionWorkflow& GetSubmissionWorkflow() const { return *SubmissionWorkflow; }
 		ISubmissionWorkflow& GetSubmissionWorkflow() { return *SubmissionWorkflow; }
 		/** @return Implements simple game-thread based queue for SubmissionWorkflow. Only one submission can be in progress at any given time. */
 		FSubmissionQueue& GetSubmissionQueue() { return SubmissionQueue; }
diff --git a/Engine/Plugins/Developer/Concert/ConcertApp/MultiUserClient/Source/MultiUserClient/Private/Replication/Stream/IClientStreamSynchronizer.h b/Engine/Plugins/Developer/Concert/ConcertApp/MultiUserClient/Source/MultiUserClient/Private/Replication/Stream/IClientStreamSynchronizer.h
--- a/Engine/Plugins/Developer/Concert/ConcertApp/MultiUserClient/Source/MultiUserClient/Private/Replication/Stream/IClientStreamSynchronizer.h
+++ b/Engine/Plugins/Developer/Concert/ConcertApp/MultiUserClient/Source/MultiUserClient/Private/Replication/Stream/IClientStreamSynchronizer.h
@@ -5,6 +5,7 @@
 #include "Async/Future.h"
 #include "Delegates/Delegate.h"
 
+struct FConcertStreamFrequencySettings;
 struct FObjectReplicationMap;
 
 namespace UE::MultiUserClient
@@ -21,9 +22,13 @@ namespace UE::MultiUserClient
 		
 		/** @return The managed client's stream ID. */
 		virtual FGuid GetStreamId() const = 0;
+		
 		/** @return What local instance thinks the client's server state is. */
 		virtual const FObjectReplicationMap& GetServerState() const = 0;
 
+		/** @return What the local instance thinks the client's replication frequencies are. */
+		virtual const FConcertStreamFrequencySettings& GetFrequencySettings() const = 0;
+
 		DECLARE_MULTICAST_DELEGATE(FOnServerStateChanged);
 		/** @return Event executed when the result of GetServerState has been updated. */
 		virtual FOnServerStateChanged& OnServerStateChanged() = 0;
diff --git a/Engine/Plugins/Developer/Concert/ConcertApp/MultiUserClient/Source/MultiUserClient/Private/Replication/Stream/StreamSynchronizer_LocalClient.cpp b/Engine/Plugins/Developer/Concert/ConcertApp/MultiUserClient/Source/MultiUserClient/Private/Replication/Stream/StreamSynchronizer_LocalClient.cpp
--- a/Engine/Plugins/Developer/Concert/ConcertApp/MultiUserClient/Source/MultiUserClient/Private/Replication/Stream/StreamSynchronizer_LocalClient.cpp
+++ b/Engine/Plugins/Developer/Concert/ConcertApp/MultiUserClient/Source/MultiUserClient/Private/Replication/Stream/StreamSynchronizer_LocalClient.cpp
@@ -4,7 +4,6 @@
 
 #include "IConcertSyncClient.h"
 #include "Replication/IConcertClientReplicationManager.h"
-#include "Replication/Messages/ChangeStream.h"
 
 #define LOCTEXT_NAMESPACE "FLocalClientStreamDiffer"
 
@@ -33,32 +32,42 @@ namespace UE::MultiUserClient
 	}
 
 	const FObjectReplicationMap& FStreamSynchronizer_LocalClient::GetServerState() const
+	{
+		const FSharedReplicationStreamDescription* Description = GetLocalMultiUserStream();
+		return Description ? Description->ReplicationMap : EmptyState.ReplicationMap;
+	}
+
+	const FConcertStreamFrequencySettings& FStreamSynchronizer_LocalClient::GetFrequencySettings() const
+	{
+		const FSharedReplicationStreamDescription* Description = GetLocalMultiUserStream();
+		return Description ? Description->FrequencySettings : EmptyState.FrequencySettings;
+	}
+
+	void FStreamSynchronizer_LocalClient::OnPostStreamsChanged()
+	{
+		// Note that this fires even for non MU-related changes, e.g. due to some other system's Concert replication API usage.
+		OnServerStateChangedDelegate.Broadcast();
+	}
+
+	const FSharedReplicationStreamDescription* FStreamSynchronizer_LocalClient::GetLocalMultiUserStream() const
 	{
 		IConcertClientReplicationManager* ReplicationManager = LocalClient->GetReplicationManager();
 		if (!ensure(ReplicationManager))
 		{
-			return EmptyState;
+			return nullptr;
 		}
 		
-		const FObjectReplicationMap* Result = nullptr;
+		const FSharedReplicationStreamDescription* Result = nullptr;
 		ReplicationManager->ForEachRegisteredStream([this, &Result](const FReplicationStreamDescription& Stream)
 		{
 			if (Stream.BaseDescription.Identifier == LocalClientStreamId)
 			{
-				Result = &Stream.BaseDescription.ReplicationMap;
+				Result = &Stream.BaseDescription;
 				return EBreakBehavior::Break;
 			}
 			return EBreakBehavior::Continue;
 		});
-		return Result
-			? *Result
-			: EmptyState;
-	}
-
-	void FStreamSynchronizer_LocalClient::OnPostStreamsChanged()
-	{
-		// Note that this fires even for non MU-related changes, e.g. due to some other system's Concert replication API usage.
-		OnServerStateChangedDelegate.Broadcast();
+		return Result;
 	}
 }
 
diff --git a/Engine/Plugins/Developer/Concert/ConcertApp/MultiUserClient/Source/MultiUserClient/Private/Replication/Stream/StreamSynchronizer_LocalClient.h b/Engine/Plugins/Developer/Concert/ConcertApp/MultiUserClient/Source/MultiUserClient/Private/Replication/Stream/StreamSynchronizer_LocalClient.h
--- a/Engine/Plugins/Developer/Concert/ConcertApp/MultiUserClient/Source/MultiUserClient/Private/Replication/Stream/StreamSynchronizer_LocalClient.h
+++ b/Engine/Plugins/Developer/Concert/ConcertApp/MultiUserClient/Source/MultiUserClient/Private/Replication/Stream/StreamSynchronizer_LocalClient.h
@@ -4,6 +4,7 @@
 
 #include "IClientStreamSynchronizer.h"
 #include "Replication/Data/ObjectReplicationMap.h"
+#include "Replication/Data/ReplicationStreamDescription.h"
 #include "Templates/UnrealTemplate.h"
 
 class IConcertSyncClient;
@@ -21,13 +22,14 @@ namespace UE::MultiUserClient
 		//~ Begin IClientStreamSynchronizer Interface
 		virtual FGuid GetStreamId() const override { return LocalClientStreamId; }
 		virtual const FObjectReplicationMap& GetServerState() const override;
+		virtual const FConcertStreamFrequencySettings& GetFrequencySettings() const override;
 		virtual FOnServerStateChanged& OnServerStateChanged() override { return OnServerStateChangedDelegate; }
 		//~ End IClientStreamSynchronizer Interface
 
 	private:
 		
 		/** Returned by GetServerState when there is no registered stream. */
-		const FObjectReplicationMap EmptyState; 
+		const FSharedReplicationStreamDescription EmptyState; 
 		
 		/** Owning client. Used to send change requests to the server. */
 		const TSharedRef<IConcertSyncClient> LocalClient;
@@ -38,6 +40,7 @@ namespace UE::MultiUserClient
 		FOnServerStateChanged OnServerStateChangedDelegate;
 		
 		void OnPostStreamsChanged();
+		const FSharedReplicationStreamDescription* GetLocalMultiUserStream() const;
 	};
 }
 
diff --git a/Engine/Plugins/Developer/Concert/ConcertApp/MultiUserClient/Source/MultiUserClient/Private/Replication/Stream/StreamSynchronizer_RemoteClient.cpp b/Engine/Plugins/Developer/Concert/ConcertApp/MultiUserClient/Source/MultiUserClient/Private/Replication/Stream/StreamSynchronizer_RemoteClient.cpp
--- a/Engine/Plugins/Developer/Concert/ConcertApp/MultiUserClient/Source/MultiUserClient/Private/Replication/Stream/StreamSynchronizer_RemoteClient.cpp
+++ b/Engine/Plugins/Developer/Concert/ConcertApp/MultiUserClient/Source/MultiUserClient/Private/Replication/Stream/StreamSynchronizer_RemoteClient.cpp
@@ -40,13 +40,13 @@ namespace UE::MultiUserClient
 		});
 		if (MultiUserDescription
 			// Avoid unnecessary Broadcast()s
-			&& LastKnownServerState != MultiUserDescription->ReplicationMap)
+			&& LastKnownServerState != *MultiUserDescription)
 		{
-			LastKnownServerState = MultiUserDescription->ReplicationMap;
+			LastKnownServerState = *MultiUserDescription;
 			OnServerStateChangedDelegate.Broadcast();
 		}
 		// If the remote client removed the last object, the stream is implicitly deleted.
-		else if (!MultiUserDescription && !LastKnownServerState.ReplicatedObjects.IsEmpty())
+		else if (!MultiUserDescription && !LastKnownServerState.ReplicationMap.ReplicatedObjects.IsEmpty())
 		{
 			LastKnownServerState = {};
 			OnServerStateChangedDelegate.Broadcast();
diff --git a/Engine/Plugins/Developer/Concert/ConcertApp/MultiUserClient/Source/MultiUserClient/Private/Replication/Stream/StreamSynchronizer_RemoteClient.h b/Engine/Plugins/Developer/Concert/ConcertApp/MultiUserClient/Source/MultiUserClient/Private/Replication/Stream/StreamSynchronizer_RemoteClient.h
--- a/Engine/Plugins/Developer/Concert/ConcertApp/MultiUserClient/Source/MultiUserClient/Private/Replication/Stream/StreamSynchronizer_RemoteClient.h
+++ b/Engine/Plugins/Developer/Concert/ConcertApp/MultiUserClient/Source/MultiUserClient/Private/Replication/Stream/StreamSynchronizer_RemoteClient.h
@@ -23,7 +23,8 @@ namespace UE::MultiUserClient
 
 		//~ Begin IClientStreamSynchronizer Interface
 		virtual FGuid GetStreamId() const override;
-		virtual const FObjectReplicationMap& GetServerState() const override { return LastKnownServerState; }
+		virtual const FObjectReplicationMap& GetServerState() const override { return LastKnownServerState.ReplicationMap; }
+		virtual const FConcertStreamFrequencySettings& GetFrequencySettings() const override { return LastKnownServerState.FrequencySettings; }
 		virtual FOnServerStateChanged& OnServerStateChanged() override { return OnServerStateChangedDelegate; }
 		//~ End IClientStreamSynchronizer Interface
 
@@ -35,7 +36,7 @@ namespace UE::MultiUserClient
 		const FDelegateHandle QueryStreamHandle; 
 
 		/** Represents what the local client thinks the replication map on the server currently looks like. */
-		FObjectReplicationMap LastKnownServerState;
+		FSharedReplicationStreamDescription LastKnownServerState;
 		
 		/** Event executed when the result of GetServerState has been synched. */
 		FOnServerStateChanged OnServerStateChangedDelegate;
diff --git a/Engine/Plugins/Developer/Concert/ConcertApp/MultiUserClient/Source/MultiUserClient/Private/Replication/Submission/External/ClientChangeOperation.cpp b/Engine/Plugins/Developer/Concert/ConcertApp/MultiUserClient/Source/MultiUserClient/Private/Replication/Submission/External/ClientChangeOperation.cpp
--- a/Engine/Plugins/Developer/Concert/ConcertApp/MultiUserClient/Source/MultiUserClient/Private/Replication/Submission/External/ClientChangeOperation.cpp
+++ b/Engine/Plugins/Developer/Concert/ConcertApp/MultiUserClient/Source/MultiUserClient/Private/Replication/Submission/External/ClientChangeOperation.cpp
@@ -6,7 +6,6 @@
 #include "ConcertLogGlobal.h"
 #include "Replication/Submission/ISubmissionWorkflow.h"
 #include "Replication/Submission/Queue/SubmissionQueue.h"
-
 #include "Replication/Submission/ISubmissionOperation.h"
 
 namespace UE::MultiUserClient
@@ -97,7 +96,7 @@ namespace UE::MultiUserClient
 		// When submission workflow is destroyed, it will finish the below futures.
 		// Keep weak reference to us in case the owning FExternalClientChangeRequestHandler is destroyed before the submission workflow.
 		// This way our code does not rely on whether FExternalClientChangeRequestHandler or ISubmissionWorkflow is destroyed first.
-		SubmissionOperation->OnCompleteStreamChangesFuture()
+		SubmissionOperation->OnCompleteStreamChangesFuture_AnyThread()
 			.Next([WeakThis = SharedThis(this).ToWeakPtr()](FSubmitStreamChangesResponse&& Response)
 			{
 				if (const TSharedPtr<FClientChangeOperation> ThisPin = WeakThis.Pin())
@@ -105,7 +104,7 @@ namespace UE::MultiUserClient
 					ThisPin->EmplaceStreamResult(ClientChangeConversionUtils::Transform(Response));
 				}
 			});
-		SubmissionOperation->OnCompleteAuthorityChangeFuture()
+		SubmissionOperation->OnCompleteAuthorityChangeFuture_AnyThread()
 			.Next([WeakThis = SharedThis(this).ToWeakPtr()](FSubmitAuthorityChangesResponse&& Response)
 			{
 				if (const TSharedPtr<FClientChangeOperation> ThisPin = WeakThis.Pin())
@@ -113,7 +112,7 @@ namespace UE::MultiUserClient
 					ThisPin->EmplaceAuthorityResult(ClientChangeConversionUtils::Transform(Response));
 				}
 			});
-		SubmissionOperation->OnCompletedOperation()
+		SubmissionOperation->OnCompletedOperation_AnyThread()
 			.Next([WeakThis = SharedThis(this).ToWeakPtr()](ESubmissionOperationCompletedCode Response)
 			{
 				if (const TSharedPtr<FClientChangeOperation> ThisPin = WeakThis.Pin())
diff --git a/Engine/Plugins/Developer/Concert/ConcertApp/MultiUserClient/Source/MultiUserClient/Private/Replication/Submission/ISubmissionOperation.h b/Engine/Plugins/Developer/Concert/ConcertApp/MultiUserClient/Source/MultiUserClient/Private/Replication/Submission/ISubmissionOperation.h
--- a/Engine/Plugins/Developer/Concert/ConcertApp/MultiUserClient/Source/MultiUserClient/Private/Replication/Submission/ISubmissionOperation.h
+++ b/Engine/Plugins/Developer/Concert/ConcertApp/MultiUserClient/Source/MultiUserClient/Private/Replication/Submission/ISubmissionOperation.h
@@ -17,8 +17,10 @@ namespace UE::MultiUserClient
 	};
 	
 	/**
-	 * Exposes the stages of submitting to the server.
+	 * Exposes several events that happen as part of submitting to the server.
 	 * Every time ISubmissionWorkflow::SubmitChanges is called a new instance is created.
+	 *
+	 * The events are exposed as futures, which can complete on any thread.
 	 */
 	class ISubmissionOperation
 	{
@@ -29,27 +31,35 @@ namespace UE::MultiUserClient
 		
 		/**
 		 * Completes when the operation of changing streams has completed.
+		 * 
 		 * @note This can be called at most once; subsequent calls result in an unset future.
+		 * @note This can complete on any thread. Usually this finishes on the game thread but timeouts usually occur on the messaging (UDP) thread.
 		 */
-		virtual TFuture<FSubmitStreamChangesResponse> OnCompleteStreamChangesFuture() = 0;
+		virtual TFuture<FSubmitStreamChangesResponse> OnCompleteStreamChangesFuture_AnyThread() = 0;
 
 		/**
 		 * Completes when the authority change request has been sent to the server.
+		 * 
 		 * @note This can be called at most once; subsequent calls result in an unset future.
+		 * @note This can complete on any thread. Usually this finishes on the game thread but timeouts usually occur on the messaging (UDP) thread.
 		 */
-		virtual TFuture<FSubmitAuthorityChangesRequest> OnRequestAuthorityChangeFuture() = 0;
+		virtual TFuture<FSubmitAuthorityChangesRequest> OnRequestAuthorityChangeFuture_AnyThread() = 0;
 		
 		/**
 		 * Completes when the operation of changing authority has completed.
+		 * 
 		 * @note This can be called at most once; subsequent calls result in an unset future.
+		 * @note This can complete on any thread. Usually this finishes on the game thread but timeouts usually occur on the messaging (UDP) thread.
 		 */
-		virtual TFuture<FSubmitAuthorityChangesResponse> OnCompleteAuthorityChangeFuture() = 0;
+		virtual TFuture<FSubmitAuthorityChangesResponse> OnCompleteAuthorityChangeFuture_AnyThread() = 0;
 
 		/**
 		 * Completes when the operation is done. No further work will be performed.
+		 * 
 		 * @note This can be called at most once; subsequent calls result in an unset future.
+		 * @note This can complete on any thread. Usually this finishes on the game thread but timeouts usually occur on the messaging (UDP) thread.
 		 */
-		virtual TFuture<ESubmissionOperationCompletedCode> OnCompletedOperation() = 0;
+		virtual TFuture<ESubmissionOperationCompletedCode> OnCompletedOperation_AnyThread() = 0;
 		
 		virtual ~ISubmissionOperation() = default;
 	};
diff --git a/Engine/Plugins/Developer/Concert/ConcertApp/MultiUserClient/Source/MultiUserClient/Private/Replication/Submission/MultiEdit/ReassignObjectPropertiesLogic.cpp b/Engine/Plugins/Developer/Concert/ConcertApp/MultiUserClient/Source/MultiUserClient/Private/Replication/Submission/MultiEdit/ReassignObjectPropertiesLogic.cpp
--- a/Engine/Plugins/Developer/Concert/ConcertApp/MultiUserClient/Source/MultiUserClient/Private/Replication/Submission/MultiEdit/ReassignObjectPropertiesLogic.cpp
+++ b/Engine/Plugins/Developer/Concert/ConcertApp/MultiUserClient/Source/MultiUserClient/Private/Replication/Submission/MultiEdit/ReassignObjectPropertiesLogic.cpp
@@ -341,7 +341,7 @@ namespace UE::MultiUserClient
 		}
 		
 		InProgressOperation = { MoveTemp(ObjectsToReassign), ClientId, MoveTemp(OldRegisteredObjects), MoveTemp(ReassignedAuthority), ParallelSubmissionOperation.ToSharedRef() };
-		ParallelSubmissionOperation->GetFuture()
+		ParallelSubmissionOperation->OnCompletedFuture_AnyThread()
 			.Next([this](FParallelExecutionResult&& Result)
 			{
 				// When "this" is being destroyed, there is no sense in further processing.
@@ -462,7 +462,7 @@ namespace UE::MultiUserClient
 			return;
 		}
 		
-		SubmitOperation->OnCompletedOperation().Next([this](ESubmissionOperationCompletedCode)
+		SubmitOperation->OnCompletedOperation_AnyThread().Next([this](ESubmissionOperationCompletedCode)
 		{
 			// We're not on game thread if OnCompletedOperation completes due to a timeout.
 			// By contract, IParallelSubmissionOperation must be destroyed on the game thread!
diff --git a/Engine/Plugins/Developer/Concert/ConcertApp/MultiUserClient/Source/MultiUserClient/Private/Replication/Submission/Remote/RemoteSubmissionListener.cpp b/Engine/Plugins/Developer/Concert/ConcertApp/MultiUserClient/Source/MultiUserClient/Private/Replication/Submission/Remote/RemoteSubmissionListener.cpp
--- a/Engine/Plugins/Developer/Concert/ConcertApp/MultiUserClient/Source/MultiUserClient/Private/Replication/Submission/Remote/RemoteSubmissionListener.cpp
+++ b/Engine/Plugins/Developer/Concert/ConcertApp/MultiUserClient/Source/MultiUserClient/Private/Replication/Submission/Remote/RemoteSubmissionListener.cpp
@@ -89,7 +89,7 @@ namespace UE::MultiUserClient
 
 		if (!RequestData.StreamChangeRequest.IsEmpty())
 		{
-			Operation->OnCompleteStreamChangesFuture().Next([this, Token = LifetimeToken.ToWeakPtr()](FSubmitStreamChangesResponse&& Response)
+			Operation->OnCompleteStreamChangesFuture_AnyThread().Next([this, Token = LifetimeToken.ToWeakPtr()](FSubmitStreamChangesResponse&& Response)
 			{
 				if (const TSharedPtr<FToken> TokenPin = Token.Pin())
 				{
@@ -99,7 +99,7 @@ namespace UE::MultiUserClient
 		}
 		if (!RequestData.AuthorityRequest.IsEmpty())
 		{
-			Operation->OnCompleteAuthorityChangeFuture().Next([this, Token = LifetimeToken.ToWeakPtr()](FSubmitAuthorityChangesResponse&& Response)
+			Operation->OnCompleteAuthorityChangeFuture_AnyThread().Next([this, Token = LifetimeToken.ToWeakPtr()](FSubmitAuthorityChangesResponse&& Response)
 			{
 				if (const TSharedPtr<FToken> TokenPin = Token.Pin())
 				{
diff --git a/Engine/Plugins/Developer/Concert/ConcertApp/MultiUserClient/Source/MultiUserClient/Private/Replication/Submission/SingleClientSubmissionOperation.h b/Engine/Plugins/Developer/Concert/ConcertApp/MultiUserClient/Source/MultiUserClient/Private/Replication/Submission/SingleClientSubmissionOperation.h
--- a/Engine/Plugins/Developer/Concert/ConcertApp/MultiUserClient/Source/MultiUserClient/Private/Replication/Submission/SingleClientSubmissionOperation.h
+++ b/Engine/Plugins/Developer/Concert/ConcertApp/MultiUserClient/Source/MultiUserClient/Private/Replication/Submission/SingleClientSubmissionOperation.h
@@ -68,10 +68,10 @@ namespace UE::MultiUserClient
 		
 		//~ Begin ISubmissionOperation Interface
 		virtual bool IsModifyingStreams() const override { return bModifiesStreams; }
-		virtual TFuture<FSubmitStreamChangesResponse> OnCompleteStreamChangesFuture() override { return StreamChangesPromise.GetFuture();  }
-		virtual TFuture<FSubmitAuthorityChangesRequest> OnRequestAuthorityChangeFuture() override { return AuthorityChangeRequestPromise.GetFuture(); }
-		virtual TFuture<FSubmitAuthorityChangesResponse> OnCompleteAuthorityChangeFuture() override { return AuthorityChangeResponsePromise.GetFuture(); }
-		virtual TFuture<ESubmissionOperationCompletedCode> OnCompletedOperation() override { return CompleteOperationPromise.GetFuture(); }
+		virtual TFuture<FSubmitStreamChangesResponse> OnCompleteStreamChangesFuture_AnyThread() override { return StreamChangesPromise.GetFuture();  }
+		virtual TFuture<FSubmitAuthorityChangesRequest> OnRequestAuthorityChangeFuture_AnyThread() override { return AuthorityChangeRequestPromise.GetFuture(); }
+		virtual TFuture<FSubmitAuthorityChangesResponse> OnCompleteAuthorityChangeFuture_AnyThread() override { return AuthorityChangeResponsePromise.GetFuture(); }
+		virtual TFuture<ESubmissionOperationCompletedCode> OnCompletedOperation_AnyThread() override { return CompleteOperationPromise.GetFuture(); }
 		//~ End ISubmissionOperation Interface
 
 	private:
diff --git a/Engine/Plugins/Developer/Concert/ConcertApp/MultiUserClient/Source/MultiUserClient/Private/Replication/Util/FrequencyUtils.cpp b/Engine/Plugins/Developer/Concert/ConcertApp/MultiUserClient/Source/MultiUserClient/Private/Replication/Util/FrequencyUtils.cpp
new file mode 100644
--- /dev/null
+++ b/Engine/Plugins/Developer/Concert/ConcertApp/MultiUserClient/Source/MultiUserClient/Private/Replication/Util/FrequencyUtils.cpp
@@ -0,0 +1,48 @@
+// Copyright Epic Games, Inc. All Rights Reserved.
+
+#include "FrequencyUtils.h"
+
+namespace UE::MultiUserClient::FrequencyUtils
+{
+	TSharedPtr<IParallelSubmissionOperation> SetFrequencySettingForClients(
+		const FSoftObjectPath& ContextObject,
+		const FInlineClientArray& Clients,
+		FReplicationClientManager& InClientManager,
+		TFunctionRef<void(FConcertObjectReplicationSettings&)> ApplySettingChange
+		)
+	{
+		TMap<FGuid, FConcertReplication_ChangeStream_Request> ParallelOperations;
+		
+		for (const FGuid& ClientId : Clients)
+		{
+			FReplicationClient* Client = InClientManager.FindClient(ClientId);
+			if (Client && CanChangeFrequencySettings(ContextObject, *Client))
+			{
+				const FGuid StreamId = Client->GetStreamSynchronizer().GetStreamId();
+				const FConcertStreamFrequencySettings& AllSettings = Client->GetStreamSynchronizer().GetFrequencySettings();
+				FConcertObjectReplicationSettings PerObjectOverride = AllSettings.GetSettingsFor(ContextObject);
+
+				const bool bNeedsToApplyDefault = !AllSettings.ObjectOverrides.Contains(ContextObject);
+				// TODO UE-203350: Specify default replication settings from Multi User project settings
+				PerObjectOverride = bNeedsToApplyDefault ? FConcertObjectReplicationSettings{ EConcertObjectReplicationMode::SpecifiedRate, 30 } : PerObjectOverride;
+				ApplySettingChange(PerObjectOverride);
+				if (!ensure(PerObjectOverride.ReplicationRate > 0))
+				{
+					continue;
+				}
+				
+				FConcertReplication_ChangeStream_Frequency FrequencyChange;
+				FrequencyChange.OverridesToAdd.Add(ContextObject, PerObjectOverride);
+				FConcertReplication_ChangeStream_Request StreamChangeRequest;
+				StreamChangeRequest.FrequencyChanges.Add(StreamId, FrequencyChange);
+				
+				ParallelOperations.Add(ClientId, StreamChangeRequest);
+			}
+		}
+
+		return ParallelOperations.IsEmpty()
+			? nullptr
+			: ExecuteParallelStreamChanges(InClientManager, ParallelOperations);
+	}
+	
+}
\ No newline at end of file
diff --git a/Engine/Plugins/Developer/Concert/ConcertApp/MultiUserClient/Source/MultiUserClient/Private/Replication/Util/FrequencyUtils.h b/Engine/Plugins/Developer/Concert/ConcertApp/MultiUserClient/Source/MultiUserClient/Private/Replication/Util/FrequencyUtils.h
new file mode 100644
--- /dev/null
+++ b/Engine/Plugins/Developer/Concert/ConcertApp/MultiUserClient/Source/MultiUserClient/Private/Replication/Util/FrequencyUtils.h
@@ -0,0 +1,130 @@
+// Copyright Epic Games, Inc. All Rights Reserved.
+
+#pragma once
+
+#include "Replication/Client/ReplicationClientManager.h"
+#include "Replication/Data/ReplicationFrequencySettings.h"
+
+#include "Algo/AllOf.h"
+#include "Async/Future.h"
+#include "Misc/Optional.h"
+#include "Templates/Function.h"
+#include "Templates/SharedPointer.h"
+
+#include <type_traits>
+
+namespace UE::MultiUserClient
+{
+	class IParallelSubmissionOperation;
+}
+
+namespace UE::MultiUserClient::FrequencyUtils
+{
+	// Low effort optimization to minimize heap allocations
+	using FInlineClientArray = TArray<FGuid, TInlineAllocator<8>>;
+	
+	/** Gets the shared value from FConcertObjectReplicationSettings, if all clients share the same value. Transform selects the member. */
+	template<typename TSubSetting, typename TSelectSubSetting>
+	TOptional<TSubSetting> FindSharedFrequencySetting(
+		const FSoftObjectPath& ContextObject,
+		const FInlineClientArray& Clients,
+		const FReplicationClientManager& InClientManager,
+		TSelectSubSetting&& Transform
+		)
+	requires std::is_invocable_r_v<TSubSetting, TSelectSubSetting, const FConcertObjectReplicationSettings&>;
+
+	enum class EChangeResult : uint8 { Success, Failure };
+	/**
+	 * Goes through all clients and sets a FConcertObjectReplicationSettings for ContextObject. The setting is applied by invoking ApplySettingChange.
+	 * @return Operations object that informs you when the change operations have completed.
+	 */
+	TSharedPtr<IParallelSubmissionOperation> SetFrequencySettingForClients(
+		const FSoftObjectPath& ContextObject,
+		const FInlineClientArray& Clients,
+		FReplicationClientManager& InClientManager,
+		TFunctionRef<void(FConcertObjectReplicationSettings&)> ApplySettingChange
+		);
+
+	/** Gets the replication rate of all clients, if they share the same value. */
+	inline TOptional<uint8> FindSharedFrequencyRate(const FSoftObjectPath& ContextObject, const FInlineClientArray& Clients, const FReplicationClientManager& InClientManager)
+	{
+		bool bHasRightMode = false;
+		const TOptional<uint8> SharedMode = FindSharedFrequencySetting<uint8>(ContextObject, Clients, InClientManager,
+			[&bHasRightMode](const FConcertObjectReplicationSettings& Setting)
+			{
+				bHasRightMode |= Setting.ReplicationMode == EConcertObjectReplicationMode::SpecifiedRate;
+				return Setting.ReplicationRate;
+			});
+		return bHasRightMode ? SharedMode : TOptional<uint8>{};
+	}
+
+	/** @return The replication mode that all clients are using, if they share the same value. */
+	inline TOptional<EConcertObjectReplicationMode> FindSharedReplicationMode(const FSoftObjectPath& ContextObject, const FInlineClientArray& Clients, const FReplicationClientManager& InClientManager)
+	{
+		return FindSharedFrequencySetting<EConcertObjectReplicationMode>(ContextObject, Clients, InClientManager,
+			[](const FConcertObjectReplicationSettings& Setting)
+			{
+				return Setting.ReplicationMode;
+			});
+	}
+	
+	/** @return Whether all clients have the given replication mode */
+	inline bool AllClientsHaveMode(const EConcertObjectReplicationMode SearchedMode, const FSoftObjectPath& ContextObject, const FInlineClientArray& Clients, const FReplicationClientManager& InClientManager) 
+	{
+		return SearchedMode == FindSharedReplicationMode(ContextObject, Clients, InClientManager);
+	}
+
+	/** @return Whether ContextObject's frequency settings can be changed for Client. */
+	inline bool CanChangeFrequencySettings(const FSoftObjectPath& ContextObject, const FReplicationClient& Client)
+	{
+		const bool bHasProperties = Client.GetStreamSynchronizer().GetServerState().HasProperties(ContextObject);
+		const bool bAllowsEditing = Client.AllowsEditing();
+		const bool bNotProcessingOtherRequests = Client.GetSubmissionWorkflow().CanSubmit();
+		return bHasProperties && bAllowsEditing && bNotProcessingOtherRequests;
+	}
+
+	/** @return Whether ContextObject's frequency settings can be changed for all given Clients. */
+	inline bool CanChangeFrequencySettings(const FSoftObjectPath& ContextObject, const FInlineClientArray& Clients, const FReplicationClientManager& InClientManager)
+	{
+		return Algo::AllOf(Clients, [ContextObject, &InClientManager](const FGuid& ClientId)
+		{
+			const FReplicationClient* Client = InClientManager.FindClient(ClientId);
+			return Client && CanChangeFrequencySettings(ContextObject, *Client);
+		});
+	}
+}
+
+namespace UE::MultiUserClient::FrequencyUtils
+{
+	template<typename TSubSetting, typename TSelectSubSetting>
+	TOptional<TSubSetting> FindSharedFrequencySetting(
+		const FSoftObjectPath& ContextObject,
+		const FInlineClientArray& Clients,
+		const FReplicationClientManager& InClientManager,
+		TSelectSubSetting&& Transform
+		)
+	requires std::is_invocable_r_v<TSubSetting, TSelectSubSetting, const FConcertObjectReplicationSettings&>
+	{
+		TOptional<TSubSetting> SharedSetting;
+		for (const FGuid& ClientId : Clients)
+		{
+			const FReplicationClient* Client = InClientManager.FindClient(ClientId);
+			if (!Client)
+			{
+				continue;
+			}
+
+			const FConcertObjectReplicationSettings& ObjectSettings = Client->GetStreamSynchronizer().GetFrequencySettings().GetSettingsFor(ContextObject);
+			const TSubSetting SubSetting = Transform(ObjectSettings);
+			if (SharedSetting && SharedSetting != SubSetting)
+			{
+				return {};
+			}
+			else
+			{
+				SharedSetting = SubSetting;
+			}
+		}
+		return SharedSetting;
+	}
+}
diff --git a/Engine/Plugins/Developer/Concert/ConcertApp/MultiUserClient/Source/MultiUserClient/Private/Replication/Util/SynchronizedRequestUtils.cpp b/Engine/Plugins/Developer/Concert/ConcertApp/MultiUserClient/Source/MultiUserClient/Private/Replication/Util/SynchronizedRequestUtils.cpp
--- a/Engine/Plugins/Developer/Concert/ConcertApp/MultiUserClient/Source/MultiUserClient/Private/Replication/Util/SynchronizedRequestUtils.cpp
+++ b/Engine/Plugins/Developer/Concert/ConcertApp/MultiUserClient/Source/MultiUserClient/Private/Replication/Util/SynchronizedRequestUtils.cpp
@@ -136,7 +136,7 @@ namespace UE::MultiUserClient
 			}
 
 			//~ Begin IParallelSubmissionOperation Interface
-			virtual TFuture<FParallelExecutionResult> GetFuture() override { return Promise.GetFuture(); }
+			virtual TFuture<FParallelExecutionResult> OnCompletedFuture_AnyThread() override { return Promise.GetFuture(); }
 			//~ End IParallelSubmissionOperation Interface
 
 			void OnCompleteStream(const FGuid& ClientId, FSubmitStreamChangesResponse&& Response)
@@ -230,7 +230,7 @@ namespace UE::MultiUserClient
 			const TSharedPtr<ISubmissionOperation> SubmissionOperation = Workflow.SubmitChanges(SubmissionParams);
 			if (ensure(SubmissionOperation))
 			{
-				SubmissionOperation->OnCompleteStreamChangesFuture()
+				SubmissionOperation->OnCompleteStreamChangesFuture_AnyThread()
 					.Next([this, WeakOwner = Owner.AsWeak()](FSubmitStreamChangesResponse&& Response)
 					{
 						if (const TSharedPtr<FSyncOperation> SyncOperation = WeakOwner.Pin())
@@ -238,7 +238,7 @@ namespace UE::MultiUserClient
 							SyncOperation->OnCompleteStream(ClientId, MoveTemp(Response));
 						}
 					});
-				SubmissionOperation->OnCompleteAuthorityChangeFuture()
+				SubmissionOperation->OnCompleteAuthorityChangeFuture_AnyThread()
 					.Next([this, WeakOwner = Owner.AsWeak()](FSubmitAuthorityChangesResponse&& Response)
 					{
 						if (const TSharedPtr<FSyncOperation> SyncOperation = WeakOwner.Pin())
diff --git a/Engine/Plugins/Developer/Concert/ConcertApp/MultiUserClient/Source/MultiUserClient/Private/Replication/Util/SynchronizedRequestUtils.h b/Engine/Plugins/Developer/Concert/ConcertApp/MultiUserClient/Source/MultiUserClient/Private/Replication/Util/SynchronizedRequestUtils.h
--- a/Engine/Plugins/Developer/Concert/ConcertApp/MultiUserClient/Source/MultiUserClient/Private/Replication/Util/SynchronizedRequestUtils.h
+++ b/Engine/Plugins/Developer/Concert/ConcertApp/MultiUserClient/Source/MultiUserClient/Private/Replication/Util/SynchronizedRequestUtils.h
@@ -31,8 +31,11 @@ namespace UE::MultiUserClient
 	{
 	public:
 
-		/** @return Future completed when all parallel submissions have been completed. */
-		virtual TFuture<FParallelExecutionResult> GetFuture() = 0;
+		/**
+		 * @return Future completed when all parallel submissions have been completed.
+		 * @note This can complete on any thread. Usually this finishes on the game thread but timeouts usually occur on the messaging (UDP) thread.
+		 */
+		virtual TFuture<FParallelExecutionResult> OnCompletedFuture_AnyThread() = 0;
 
 		virtual ~IParallelSubmissionOperation() = default;
 	};
diff --git a/Engine/Plugins/Developer/Concert/ConcertApp/MultiUserClient/Source/MultiUserClient/Private/Widgets/ActiveSession/Replication/Client/FrequencyContextMenuUtils.cpp b/Engine/Plugins/Developer/Concert/ConcertApp/MultiUserClient/Source/MultiUserClient/Private/Widgets/ActiveSession/Replication/Client/FrequencyContextMenuUtils.cpp
new file mode 100644
--- /dev/null
+++ b/Engine/Plugins/Developer/Concert/ConcertApp/MultiUserClient/Source/MultiUserClient/Private/Widgets/ActiveSession/Replication/Client/FrequencyContextMenuUtils.cpp
@@ -0,0 +1,330 @@
+// Copyright Epic Games, Inc. All Rights Reserved.
+
+#include "FrequencyContextMenuUtils.h"
+
+#include "Replication/Client/ReplicationClientManager.h"
+#include "Replication/Util/FrequencyUtils.h"
+
+#include "Async/Async.h"
+#include "Async/TaskGraphInterfaces.h"
+#include "Framework/Commands/UIAction.h"
+#include "Framework/MultiBox/MultiBoxBuilder.h"
+#include "Textures/SlateIcon.h"
+#include "Widgets/Images/SThrobber.h"
+#include "Widgets/Input/SNumericEntryBox.h"
+#include "Widgets/Layout/SScaleBox.h"
+
+#define LOCTEXT_NAMESPACE "FrequencyContextMenuUtils"
+
+namespace UE::MultiUserClient::FrequencyContextMenuUtils
+{
+	namespace Private
+	{
+		using FInlineClientArray = FrequencyUtils::FInlineClientArray;
+
+		static FText GetRealtimeTooltip() { return LOCTEXT("ListView_Realtime.ToolTip", "The object will replicate as fast as possible (every tick)."); }
+		static FText GetSpecifiedRateTooltip() { return LOCTEXT("ListView_SpecifiedRate.ToolTip", "The object replicates at the specified rate per second."); }
+		static FText GetEditBoxInstructionsTooltip() { return LOCTEXT("Frequency.ToolTip", "Specifies how often this object should replicate in 1 second."); }
+
+		/** Displayed in context menu for changing an object's replication frequency. */
+		class SFrequencyNumericBox : public SCompoundWidget
+		{
+		public:
+
+			SLATE_BEGIN_ARGS(SFrequencyNumericBox)
+			{}
+				SLATE_ARGUMENT(FSoftObjectPath, SelectedObject)
+				SLATE_ATTRIBUTE(FInlineClientArray, Clients)
+			SLATE_END_ARGS()
+			
+			void Construct(const FArguments& InArgs, FReplicationClientManager& InClientManager)
+			{
+				ContextObject = InArgs._SelectedObject;
+				ClientsAttribute = InArgs._Clients;
+				ClientManager = &InClientManager;
+				
+				ChildSlot
+				[
+					SNew(SBox)
+					.MinDesiredWidth(100.f)
+					[
+						SNew(SHorizontalBox)
+
+						+SHorizontalBox::Slot()
+						.FillWidth(1.f)
+						[
+							SNew(SNumericEntryBox<uint8>)
+							.AllowSpin(false)
+							.IsEnabled(this, &SFrequencyNumericBox::GetIsEnabled)
+							.Value(this, &SFrequencyNumericBox::GetValue)
+							.OnValueCommitted(this, &SFrequencyNumericBox::OnValueCommitted)
+							.ToolTipText(this, &SFrequencyNumericBox::GetToolTipText)
+							.UndeterminedString(LOCTEXT("Single.Undetermined", "n/a"))
+						]
+						
+						+SHorizontalBox::Slot()
+						.AutoWidth()
+						[
+							SNew(SScaleBox)
+							[
+								SNew(SThrobber)
+								.Visibility_Lambda([this](){ return InProgressChange.IsSet() ? EVisibility::HitTestInvisible : EVisibility::Collapsed; })
+							]
+						]
+					]
+				];
+			}
+
+		private:
+
+			FSoftObjectPath ContextObject;
+			TAttribute<FInlineClientArray> ClientsAttribute;
+			FReplicationClientManager* ClientManager = nullptr;
+
+			struct FChangeOperation
+			{
+				TSharedPtr<IParallelSubmissionOperation> ChangeOperation;
+				uint8 ValueBeingSet;
+			};
+			/** Set for as long a change operation is in progress. */
+			TOptional<FChangeOperation> InProgressChange;
+
+			/** @return Gets all clients that have the given mode set for ContextObject. */
+			FInlineClientArray GetClientsWithMode(const EConcertObjectReplicationMode Mode) const
+			{
+				FInlineClientArray Result = ClientsAttribute.Get();
+				for (auto It = Result.CreateIterator(); It; ++It)
+				{
+					const TOptional<EConcertObjectReplicationMode> ClientMode = FrequencyUtils::FindSharedReplicationMode(ContextObject, { *It }, *ClientManager);
+					if (!ClientMode || *ClientMode != Mode)
+					{
+						It.RemoveCurrent();
+					}
+				}
+				return Result;
+			}
+
+			/** @return Whether all clients with specified rate can be changed. */
+			bool GetIsEnabled() const
+			{
+				const FInlineClientArray Clients = GetClientsWithMode(EConcertObjectReplicationMode::SpecifiedRate);
+				return !InProgressChange && !Clients.IsEmpty() && FrequencyUtils::CanChangeFrequencySettings(ContextObject, Clients, *ClientManager);
+			}
+			
+			void OnValueCommitted(uint8 NewValue, ETextCommit::Type CommitType)
+			{
+				if (CommitType != ETextCommit::OnEnter || NewValue == 0 || InProgressChange.IsSet())
+				{
+					return;
+				}
+			
+				const TSharedPtr<IParallelSubmissionOperation> ChangeOperation = FrequencyUtils::SetFrequencySettingForClients(ContextObject, GetClientsWithMode(EConcertObjectReplicationMode::SpecifiedRate), *ClientManager,
+					[NewValue](FConcertObjectReplicationSettings& Override)
+					{
+						Override.ReplicationRate = NewValue;
+					});
+				InProgressChange = { ChangeOperation, NewValue };
+
+				// For better UX, while we're waiting on the server to progress the request this widget should show what the user entered last.
+				// Once the request completes, the widget will show what the server state (so it may revert if the request failed, e.g. due to time out).
+				InProgressChange->ChangeOperation->OnCompletedFuture_AnyThread()
+					.Next([WeakThis = TWeakPtr<SFrequencyNumericBox>(SharedThis(this))](FParallelExecutionResult&& ExecutionResult)
+					{
+						// User may have close the menu containing this widget before the request completed.
+						if (const TSharedPtr<SFrequencyNumericBox> ThisPin = WeakThis.Pin())
+						{
+							// GuardedExecuteOnGameThread internally checks that "this" is valid before executing lambda
+							ThisPin->GuardedExecuteOnGameThread([WeakThis]() { WeakThis.Pin()->InProgressChange.Reset(); });
+						}
+					});
+			}
+
+			void GuardedExecuteOnGameThread(TUniqueFunction<void()> Callback) const
+			{
+				if (IsInGameThread())
+				{
+					Callback();
+				}
+				else
+				{
+					AsyncTask(ENamedThreads::GameThread, [Callback = MoveTemp(Callback), WeakThis = AsWeak()]()
+					{
+						if (WeakThis.IsValid())
+						{
+							Callback();
+						}
+					});
+				}
+			}
+
+			TOptional<uint8> GetValue() const
+			{
+				// For better UX, if a change is in progress show whatever the user specified for the request. Once the request finishes, show whatever the state is on the server.
+				return InProgressChange
+					? InProgressChange->ValueBeingSet
+					: FrequencyUtils::FindSharedFrequencyRate(ContextObject, GetClientsWithMode(EConcertObjectReplicationMode::SpecifiedRate), *ClientManager);
+			}
+			
+			FText GetToolTipText() const
+			{
+				if (GetValue())
+				{
+					return GetEditBoxInstructionsTooltip();
+				}
+
+				const TOptional<EConcertObjectReplicationMode> Mode = FrequencyUtils::FindSharedReplicationMode(ContextObject, GetClientsWithMode(EConcertObjectReplicationMode::SpecifiedRate), *ClientManager);
+				return Mode
+					? LOCTEXT("NotApplicable.Mixed", "Multiple clients are replicating this object with different rates.")
+					: LOCTEXT("NotApplicable.AllRealtime", "This object is repliacting in realtime");
+			}
+		};
+		
+		static void AppendReplicationModeToggleButtons(FMenuBuilder& MenuBuilder, const FSoftObjectPath& ContextObject, TAttribute<FInlineClientArray> Clients, FReplicationClientManager& InClientManager)
+		{
+			const auto SetReplicationMode = [ContextObject, Clients, &InClientManager](const EConcertObjectReplicationMode ModeToSet)
+			{
+				FrequencyUtils::SetFrequencySettingForClients(ContextObject, Clients.Get(), InClientManager,
+					[ModeToSet](FConcertObjectReplicationSettings& SettingToOverride)
+					{
+						SettingToOverride.ReplicationMode = ModeToSet;
+					});
+			};
+			
+			MenuBuilder.AddMenuEntry(
+				LOCTEXT("ListView_Realtime.Label", "Realtime"),
+				GetRealtimeTooltip(),
+				FSlateIcon(),
+				FUIAction(
+					FExecuteAction::CreateLambda([SetReplicationMode]()
+					{
+						SetReplicationMode(EConcertObjectReplicationMode::Realtime);
+					}),
+					FCanExecuteAction::CreateLambda([ContextObject, Clients, &InClientManager]()
+					{
+						return FrequencyUtils::CanChangeFrequencySettings(ContextObject, Clients.Get(), InClientManager);
+					}),
+					FIsActionChecked::CreateLambda([ContextObject, Clients, &InClientManager]()
+					{
+						return FrequencyUtils::AllClientsHaveMode(EConcertObjectReplicationMode::Realtime, ContextObject, Clients.Get(), InClientManager);
+					})
+				),
+				NAME_None,
+				EUserInterfaceActionType::RadioButton
+			);
+			
+			MenuBuilder.AddMenuEntry(
+				LOCTEXT("ListView_SpecifiedRate.Label", "Specified rate"),
+				GetSpecifiedRateTooltip(),
+				FSlateIcon(),
+				FUIAction(
+					FExecuteAction::CreateLambda([SetReplicationMode]()
+					{
+						SetReplicationMode(EConcertObjectReplicationMode::SpecifiedRate);
+					}),
+					FCanExecuteAction::CreateLambda([ContextObject, Clients, &InClientManager]()
+					{
+						return FrequencyUtils::CanChangeFrequencySettings(ContextObject, Clients.Get(), InClientManager);
+					}),
+					FIsActionChecked::CreateLambda([ContextObject, Clients, &InClientManager]()
+					{
+						return FrequencyUtils::AllClientsHaveMode(EConcertObjectReplicationMode::SpecifiedRate, ContextObject, Clients.Get(), InClientManager);
+					})
+				),
+				NAME_None,
+				EUserInterfaceActionType::RadioButton
+				);
+		}
+
+		static void AddReplicationModeSubMenu(FMenuBuilder& MenuBuilder, const FSoftObjectPath& ContextObject, TAttribute<FInlineClientArray> GetClientsAttribute, FReplicationClientManager& InClientManager)
+		{
+			const auto GetText = [ContextObject, GetClientsAttribute, &InClientManager](FText Mixed, FText Specified, FText Realtime)
+			{
+				const FInlineClientArray& Clients = GetClientsAttribute.Get(); 
+				const TOptional<EConcertObjectReplicationMode> Mode = FrequencyUtils::FindSharedReplicationMode(ContextObject, Clients, InClientManager);
+				if (!Mode)
+				{
+					return Mixed;
+				}
+				switch (*Mode)
+				{
+				case EConcertObjectReplicationMode::SpecifiedRate: return Specified;
+				case EConcertObjectReplicationMode::Realtime: return Realtime;
+				default: return FText::GetEmpty();
+				}
+			};
+
+			constexpr bool bOpenSubMenuOnClick = false;
+			constexpr bool bShouldCloseWindowAfterMenuSelection = false;
+			MenuBuilder.AddSubMenu(
+				TAttribute<FText>::CreateLambda([GetText]()
+				{
+					const FText MixedLabel = LOCTEXT("Mixed.Label", "Mixed mode"); 
+					const FText SpecifiedLabel = LOCTEXT("Specified.Label", "Specified mode"); 
+					const FText RealtimeLabel = LOCTEXT("Realtime.Label", "Realtime mode"); 
+					return GetText(MixedLabel, SpecifiedLabel, RealtimeLabel);
+				}),
+				TAttribute<FText>::CreateLambda([GetText]()
+				{
+					const FText MixedTooltip = LOCTEXT("Mixed.Tooltip", "Multiple clients are replicating the object with different modes."); 
+					return GetText(MixedTooltip, GetSpecifiedRateTooltip(), GetRealtimeTooltip());
+				}),
+				FNewMenuDelegate::CreateLambda([ContextObject, GetClientsAttribute, &InClientManager](FMenuBuilder& MenuBuilder)
+				{
+					AppendReplicationModeToggleButtons(MenuBuilder, ContextObject, GetClientsAttribute, InClientManager);
+				}),
+				bOpenSubMenuOnClick, FSlateIcon(), bShouldCloseWindowAfterMenuSelection
+			);
+		}
+
+		static void SharedAppendFrequencyToMenu(FMenuBuilder& MenuBuilder, const FSoftObjectPath& ContextObject, TAttribute<FInlineClientArray> GetClientsAttribute, FReplicationClientManager& InClientManager)
+		{
+			MenuBuilder.BeginSection(NAME_None, LOCTEXT("Frequency", "Frequency"));
+			
+			// Change to Realtime / Specified Rate
+			AddReplicationModeSubMenu(MenuBuilder, ContextObject, GetClientsAttribute, InClientManager);
+
+			// Replication rate, only enabled if Specified Rate
+			constexpr bool bNoIndent = true;
+			constexpr bool bSearchable = true;
+			MenuBuilder.AddWidget(
+				SNew(Private::SFrequencyNumericBox, InClientManager)
+					.SelectedObject(ContextObject)
+					.Clients(MoveTemp(GetClientsAttribute)),
+				LOCTEXT("Frequency.Label", "Replication Rate"),
+				bNoIndent,
+				bSearchable,
+				GetEditBoxInstructionsTooltip()
+				);
+
+			MenuBuilder.EndSection();
+		}
+	}
+	
+	void AddFrequencyOptionsForSingleClient(FMenuBuilder& MenuBuilder, const FSoftObjectPath& ContextObject, const FGuid& ClientId, FReplicationClientManager& InClientManager)
+	{
+		Private::SharedAppendFrequencyToMenu(MenuBuilder, ContextObject, TArray{ ClientId }, InClientManager);
+	}
+
+	void AddFrequencyOptionsForMultipleClients(FMenuBuilder& MenuBuilder, const FSoftObjectPath& ContextObject, FReplicationClientManager& InClientManager)
+	{
+		using namespace Private;
+		TAttribute<FInlineClientArray> GetClientsAttribute = TAttribute<FInlineClientArray>::CreateLambda([ContextObject, &InClientManager]()
+		{
+			FInlineClientArray Result;
+			InClientManager.ForEachClient([&ContextObject, &Result](const FReplicationClient& Client)
+			{
+				// Only clients that have the object registered should be considered
+				if (Client.GetStreamSynchronizer().GetServerState().HasProperties(ContextObject))
+				{
+					Result.Add(Client.GetEndpointId());
+				}
+				return EBreakBehavior::Continue;
+			});
+			return Result;
+		});
+		
+		SharedAppendFrequencyToMenu(MenuBuilder, ContextObject, MoveTemp(GetClientsAttribute), InClientManager);
+	}
+}
+
+#undef LOCTEXT_NAMESPACE
\ No newline at end of file
diff --git a/Engine/Plugins/Developer/Concert/ConcertApp/MultiUserClient/Source/MultiUserClient/Private/Widgets/ActiveSession/Replication/Client/FrequencyContextMenuUtils.h b/Engine/Plugins/Developer/Concert/ConcertApp/MultiUserClient/Source/MultiUserClient/Private/Widgets/ActiveSession/Replication/Client/FrequencyContextMenuUtils.h
new file mode 100644
--- /dev/null
+++ b/Engine/Plugins/Developer/Concert/ConcertApp/MultiUserClient/Source/MultiUserClient/Private/Widgets/ActiveSession/Replication/Client/FrequencyContextMenuUtils.h
@@ -0,0 +1,39 @@
+// Copyright Epic Games, Inc. All Rights Reserved.
+
+#pragma once
+
+#include "Containers/ArrayView.h"
+
+class FMenuBuilder;
+struct FGuid;
+struct FSoftObjectPath;
+
+namespace UE::MultiUserClient
+{
+	class FReplicationClientManager;
+}
+
+namespace UE::MultiUserClient::FrequencyContextMenuUtils
+{
+	/** Adds an edit box for changing the selected object's frequencies. */
+	void AddFrequencyOptionsForSingleClient(FMenuBuilder& MenuBuilder, const FSoftObjectPath& ContextObjects, const FGuid& ClientId, FReplicationClientManager& InClientManager);
+	/** Adds an edit box for batch reassigning the select object's frequencies for all replicating clients. */
+	void AddFrequencyOptionsForMultipleClients(FMenuBuilder& MenuBuilder, const FSoftObjectPath& ContextObjects, FReplicationClientManager& InClientManager);
+	
+	/** Adds an edit box for changing the selected object's frequencies. */
+	inline void AddFrequencyOptionsIfOneContextObject_SingleClient(FMenuBuilder& MenuBuilder, TConstArrayView<FSoftObjectPath> ContextObjects, const FGuid& ClientId, FReplicationClientManager& InClientManager)
+	{
+		if (ContextObjects.Num() == 1)
+		{
+			AddFrequencyOptionsForSingleClient(MenuBuilder, ContextObjects[0], ClientId, InClientManager);
+		}
+	}
+	/** Adds an edit box for batch reassigning the select object's frequencies for all replicating clients. */
+	inline void AddFrequencyOptionsIfOneContextObject_MultiClient(FMenuBuilder& MenuBuilder, TConstArrayView<FSoftObjectPath> ContextObjects, FReplicationClientManager& InClientManager)
+	{
+		if (ContextObjects.Num() == 1)
+		{
+			AddFrequencyOptionsForMultipleClients(MenuBuilder, ContextObjects[0], InClientManager);
+		}
+	}
+}
diff --git a/Engine/Plugins/Developer/Concert/ConcertApp/MultiUserClient/Source/MultiUserClient/Private/Widgets/ActiveSession/Replication/Client/Multi/SMultiClientView.cpp b/Engine/Plugins/Developer/Concert/ConcertApp/MultiUserClient/Source/MultiUserClient/Private/Widgets/ActiveSession/Replication/Client/Multi/SMultiClientView.cpp
--- a/Engine/Plugins/Developer/Concert/ConcertApp/MultiUserClient/Source/MultiUserClient/Private/Widgets/ActiveSession/Replication/Client/Multi/SMultiClientView.cpp
+++ b/Engine/Plugins/Developer/Concert/ConcertApp/MultiUserClient/Source/MultiUserClient/Private/Widgets/ActiveSession/Replication/Client/Multi/SMultiClientView.cpp
@@ -16,6 +16,7 @@
 #include "Widgets/ActiveSession/Replication/Client/SClientToolbar.h"
 
 #include "Widgets/SBoxPanel.h"
+#include "Widgets/ActiveSession/Replication/Client/FrequencyContextMenuUtils.h"
 
 #define LOCTEXT_NAMESPACE "SMultiClientView"
 
@@ -91,6 +92,7 @@ namespace UE::MultiUserClient
 			.ViewerParams 
 			{
 				.SubobjectModel = CreateDefaultComponentHierarchySubobjectModel(), // This makes actors have children in the top view
+				.OnExtendObjectsContextMenu = FExtendObjectMenu::CreateSP(this, &SMultiClientView::ExtendObjectContextMenu),
 				.AdditionalObjectColumns =
 				{
 					MultiStreamColumns::ReplicationToggle(InConcertClient, ConsolidatedStreamModelAttribute, InClientManager),
@@ -155,6 +157,11 @@ namespace UE::MultiUserClient
 		StreamEditor->GetEditorBase().RequestObjectColumnResort(MultiStreamColumns::ReassignOwnershipColumnId);
 		StreamEditor->GetEditorBase().RequestPropertyColumnResort(MultiStreamColumns::AssignPropertyColumnId);
 	}
+
+	void SMultiClientView::ExtendObjectContextMenu(FMenuBuilder& MenuBuilder, TConstArrayView<FSoftObjectPath> ContextObjects) const
+	{
+		FrequencyContextMenuUtils::AddFrequencyOptionsIfOneContextObject_MultiClient(MenuBuilder, ContextObjects, *ClientManager);
+	}
 }
 
 #undef LOCTEXT_NAMESPACE
\ No newline at end of file
diff --git a/Engine/Plugins/Developer/Concert/ConcertApp/MultiUserClient/Source/MultiUserClient/Private/Widgets/ActiveSession/Replication/Client/Multi/SMultiClientView.h b/Engine/Plugins/Developer/Concert/ConcertApp/MultiUserClient/Source/MultiUserClient/Private/Widgets/ActiveSession/Replication/Client/Multi/SMultiClientView.h
--- a/Engine/Plugins/Developer/Concert/ConcertApp/MultiUserClient/Source/MultiUserClient/Private/Widgets/ActiveSession/Replication/Client/Multi/SMultiClientView.h
+++ b/Engine/Plugins/Developer/Concert/ConcertApp/MultiUserClient/Source/MultiUserClient/Private/Widgets/ActiveSession/Replication/Client/Multi/SMultiClientView.h
@@ -5,12 +5,8 @@
 #include "Widgets/DeclarativeSyntaxSupport.h"
 #include "Widgets/SCompoundWidget.h"
 
-namespace UE::MultiUserClient
-{
-	class IClientSelectionModel;
-}
-
 class IConcertClient;
+class FMenuBuilder;
 
 namespace UE::ConcertClientSharedSlate
 {
@@ -24,6 +20,7 @@ namespace UE::MultiUserClient
 	class FMultiStreamModel;
 	class FReplicationClient;
 	class FReplicationClientManager;
+	class IClientSelectionModel;
 	class SClientToolbar;
 
 	/** Displays a selection of clients. */
@@ -62,5 +59,8 @@ namespace UE::MultiUserClient
 		void RebuildClientSubscriptions();
 		void CleanClientSubscriptions();
 		void OnClientChanged(FGuid Guid);
+		
+		/** Adds additional entries to the context menu for the object tree view. */
+		void ExtendObjectContextMenu(FMenuBuilder& MenuBuilder, TConstArrayView<FSoftObjectPath> ContextObjects) const;
 	};
 }
diff --git a/Engine/Plugins/Developer/Concert/ConcertApp/MultiUserClient/Source/MultiUserClient/Private/Widgets/ActiveSession/Replication/Client/Single/SReplicationClientView.cpp b/Engine/Plugins/Developer/Concert/ConcertApp/MultiUserClient/Source/MultiUserClient/Private/Widgets/ActiveSession/Replication/Client/Single/SReplicationClientView.cpp
--- a/Engine/Plugins/Developer/Concert/ConcertApp/MultiUserClient/Source/MultiUserClient/Private/Widgets/ActiveSession/Replication/Client/Single/SReplicationClientView.cpp
+++ b/Engine/Plugins/Developer/Concert/ConcertApp/MultiUserClient/Source/MultiUserClient/Private/Widgets/ActiveSession/Replication/Client/Single/SReplicationClientView.cpp
@@ -15,6 +15,7 @@
 #include "Widgets/ActiveSession/Replication/Client/SClientToolbar.h"
 
 #include "Widgets/SBoxPanel.h"
+#include "Widgets/ActiveSession/Replication/Client/FrequencyContextMenuUtils.h"
 #include "Widgets/Layout/SBox.h"
 
 #define LOCTEXT_NAMESPACE "SReplicationClientView"
@@ -66,7 +67,7 @@ namespace UE::MultiUserClient
 		FGlobalAuthorityCache& AuthorityCache = ClientManager->GetAuthorityCache();
 		
 		const TAttribute<const IReplicationStreamViewer*> GetReplicationViewerAttribute =
-			TAttribute<const IReplicationStreamViewer*>::CreateLambda([this](){ return EditorView_TwoSectioned.Get(); });
+			TAttribute<const IReplicationStreamViewer*>::CreateLambda([this](){ return EditorView.Get(); });
 
 		// Add checkboxes in front of top level and subobject rows for changing authority
 		const FCreateEditorParams ReplicationEditorCreationParams
@@ -79,6 +80,7 @@ namespace UE::MultiUserClient
 			.ViewerParams =
 			{
 				.SubobjectModel = CreateDefaultComponentHierarchySubobjectModel(), // This makes actors have children in the top view
+				.OnExtendObjectsContextMenu = FExtendObjectMenu::CreateSP(this, &SReplicationClientView::ExtendObjectContextMenu),
 				.AdditionalObjectColumns =
 				{
 					SingleClientColumns::ToggleObjectAuthority(AuthorityTracker, SubmissionWorkflow),
@@ -93,21 +95,21 @@ namespace UE::MultiUserClient
 			}
 		};
 
-		EditorView_TwoSectioned = CreateDefaultStreamEditor(ReplicationEditorCreationParams);
-		return EditorView_TwoSectioned.ToSharedRef();
+		EditorView = CreateDefaultStreamEditor(ReplicationEditorCreationParams);
+		return EditorView.ToSharedRef();
+	}
+
+	void SReplicationClientView::ExtendObjectContextMenu(FMenuBuilder& MenuBuilder, TConstArrayView<FSoftObjectPath> ContextObjects) const
+	{
+		const FReplicationClient* ReplicationClient = GetReplicationClientAttribute.Get();
+		check(ReplicationClient);
+		
+		FrequencyContextMenuUtils::AddFrequencyOptionsIfOneContextObject_SingleClient(MenuBuilder, ContextObjects, ReplicationClient->GetEndpointId(), *ClientManager);
 	}
 
 	void SReplicationClientView::OnModelChanged() const
 	{
-		// Only one is valid at the same time
-		if (EditorView_TwoSectioned)
-		{
-			EditorView_TwoSectioned->Refresh();
-		}
-		if (EditorView_ThreeSectioned)
-		{
-			EditorView_ThreeSectioned->Refresh();
-		}
+		EditorView->Refresh();
 	}
 
 	void SReplicationClientView::EnumerateReplicatedObjects(TFunctionRef<void(const FSoftObjectPath&)> Consumer) const
diff --git a/Engine/Plugins/Developer/Concert/ConcertApp/MultiUserClient/Source/MultiUserClient/Private/Widgets/ActiveSession/Replication/Client/Single/SReplicationClientView.h b/Engine/Plugins/Developer/Concert/ConcertApp/MultiUserClient/Source/MultiUserClient/Private/Widgets/ActiveSession/Replication/Client/Single/SReplicationClientView.h
--- a/Engine/Plugins/Developer/Concert/ConcertApp/MultiUserClient/Source/MultiUserClient/Private/Widgets/ActiveSession/Replication/Client/Single/SReplicationClientView.h
+++ b/Engine/Plugins/Developer/Concert/ConcertApp/MultiUserClient/Source/MultiUserClient/Private/Widgets/ActiveSession/Replication/Client/Single/SReplicationClientView.h
@@ -6,6 +6,7 @@
 #include "Widgets/SCompoundWidget.h"
 
 class IConcertClient;
+class FMenuBuilder;
 
 namespace UE::ConcertClientSharedSlate
 {
@@ -40,14 +41,15 @@ namespace UE::MultiUserClient
 		FReplicationClientManager* ClientManager;
 		
 		/** The editor view of the replication content. */
-		TSharedPtr<ConcertClientSharedSlate::IReplicationStreamEditor> EditorView_TwoSectioned;
-		/** The editor view of the replication content. */
-		TSharedPtr<ConcertClientSharedSlate::IReplicationStreamEditor> EditorView_ThreeSectioned;
+		TSharedPtr<ConcertClientSharedSlate::IReplicationStreamEditor> EditorView;
 		
 		/** The client to depict. Should always return true. If the client is destroyed, so should this widget be. */
 		TAttribute<FReplicationClient*> GetReplicationClientAttribute;
-		
+
 		TSharedRef<SWidget> CreateContent(FReplicationClient& InReplicationClient);
+
+		/** Adds additional entries to the context menu for the object tree view. */
+		void ExtendObjectContextMenu(FMenuBuilder& MenuBuilder, TConstArrayView<FSoftObjectPath> ContextObjects) const;
 		
 		/** Called when any of the streams change. */
 		void OnModelChanged() const;
diff --git a/Engine/Plugins/Developer/Concert/ConcertSync/ConcertSyncCore/Source/ConcertSyncCore/Public/Replication/Messages/ChangeStream.h b/Engine/Plugins/Developer/Concert/ConcertSync/ConcertSyncCore/Source/ConcertSyncCore/Public/Replication/Messages/ChangeStream.h
--- a/Engine/Plugins/Developer/Concert/ConcertSync/ConcertSyncCore/Source/ConcertSyncCore/Public/Replication/Messages/ChangeStream.h
+++ b/Engine/Plugins/Developer/Concert/ConcertSync/ConcertSyncCore/Source/ConcertSyncCore/Public/Replication/Messages/ChangeStream.h
@@ -186,7 +186,7 @@ struct FConcertReplication_ChangeStream_Request
 	UPROPERTY()
 	TSet<FGuid> StreamsToRemove;
 
-	bool IsEmpty() const { return ObjectsToRemove.IsEmpty() && ObjectsToPut.IsEmpty() && StreamsToAdd.IsEmpty() && StreamsToRemove.IsEmpty(); }
+	bool IsEmpty() const { return ObjectsToRemove.IsEmpty() && ObjectsToPut.IsEmpty() && FrequencyChanges.IsEmpty() && StreamsToAdd.IsEmpty() && StreamsToRemove.IsEmpty(); }
 
 	friend bool operator==(const FConcertReplication_ChangeStream_Request& Left, const FConcertReplication_ChangeStream_Request& Right)
 	{
diff --git a/Engine/Plugins/Developer/Concert/ConcertUI/ConcertClientSharedSlate/Source/ConcertClientSharedSlate/Private/Replication/Editor/View/ObjectEditor/SBaseReplicationStreamEditor.cpp b/Engine/Plugins/Developer/Concert/ConcertUI/ConcertClientSharedSlate/Source/ConcertClientSharedSlate/Private/Replication/Editor/View/ObjectEditor/SBaseReplicationStreamEditor.cpp
--- a/Engine/Plugins/Developer/Concert/ConcertUI/ConcertClientSharedSlate/Source/ConcertClientSharedSlate/Private/Replication/Editor/View/ObjectEditor/SBaseReplicationStreamEditor.cpp
+++ b/Engine/Plugins/Developer/Concert/ConcertUI/ConcertClientSharedSlate/Source/ConcertClientSharedSlate/Private/Replication/Editor/View/ObjectEditor/SBaseReplicationStreamEditor.cpp
@@ -202,20 +202,24 @@ namespace UE::ConcertClientSharedSlate
 	TSharedPtr<SWidget> SBaseReplicationStreamEditor::OnObjectsContextMenuOpening()
 	{
 		FMenuBuilder MenuBuilder(true, nullptr);
+		
 		AddObjectSourceContextMenuOptions(MenuBuilder);
 		MenuBuilder.AddMenuEntry(
-				LOCTEXT("DeleteItems", "Delete"),
-				TAttribute<FText>::CreateLambda([this](){ return GetEditingDisabledText(); }),
-				FSlateIcon(),
-				FUIAction(
-					FExecuteAction::CreateSP(this, &SBaseReplicationStreamEditor::OnDeleteObjects_PassByValue, ReplicationViewer->GetSelectedOutlinerObjects()),
-					FCanExecuteAction::CreateLambda([this]() { return !IsEditingDisabled(); })
-					),
-				NAME_None,
-				EUserInterfaceActionType::Button
-			);
+			LOCTEXT("DeleteItems", "Delete"),
+			TAttribute<FText>::CreateLambda([this](){ return GetEditingDisabledText(); }),
+			FSlateIcon(),
+			FUIAction(
+				FExecuteAction::CreateSP(this, &SBaseReplicationStreamEditor::OnDeleteObjects_PassByValue, ReplicationViewer->GetSelectedOutlinerObjects()),
+				FCanExecuteAction::CreateLambda([this]() { return !IsEditingDisabled(); })
+				),
+			NAME_None,
+			EUserInterfaceActionType::Button
+		);
+
+		TArray<FSoftObjectPath> SelectedObjects;
+		Algo::Transform(ReplicationViewer->GetSelectedOutlinerObjects(), SelectedObjects, [](const TSharedPtr<FReplicatedObjectData>& Data){ return Data->GetObjectPath(); });
+		OnExtendObjectsContextMenuDelegate.ExecuteIfBound(MenuBuilder, SelectedObjects);
 		
-		OnExtendObjectsContextMenuDelegate.ExecuteIfBound(MenuBuilder);
 		return MenuBuilder.MakeWidget();
 	}
 
diff --git a/Engine/Plugins/Developer/Concert/ConcertUI/ConcertClientSharedSlate/Source/ConcertClientSharedSlate/Private/Replication/Editor/View/ObjectEditor/SBaseReplicationStreamEditor.h b/Engine/Plugins/Developer/Concert/ConcertUI/ConcertClientSharedSlate/Source/ConcertClientSharedSlate/Private/Replication/Editor/View/ObjectEditor/SBaseReplicationStreamEditor.h
--- a/Engine/Plugins/Developer/Concert/ConcertUI/ConcertClientSharedSlate/Source/ConcertClientSharedSlate/Private/Replication/Editor/View/ObjectEditor/SBaseReplicationStreamEditor.h
+++ b/Engine/Plugins/Developer/Concert/ConcertUI/ConcertClientSharedSlate/Source/ConcertClientSharedSlate/Private/Replication/Editor/View/ObjectEditor/SBaseReplicationStreamEditor.h
@@ -70,7 +70,7 @@ namespace UE::ConcertClientSharedSlate
 			SLATE_ARGUMENT(TSharedPtr<ISubobjectModel>, SubobjectModel)
 
 			/** Called to generate the context menu for objects. This extends the options already generated by this widget. */
-			SLATE_EVENT(FExtendMenu, OnExtendObjectsContextMenu)
+			SLATE_EVENT(FExtendObjectMenu, OnExtendObjectsContextMenu)
 
 			/** Optional widget to add to the left of the object list search bar. */
 			SLATE_NAMED_SLOT(FArguments, LeftOfObjectSearchBar)
@@ -123,7 +123,7 @@ namespace UE::ConcertClientSharedSlate
 		TAttribute<FText> EditingDisabledToolTipTextAttribute;
 		
 		/** Generates additional entries */
-		FExtendMenu OnExtendObjectsContextMenuDelegate;
+		FExtendObjectMenu OnExtendObjectsContextMenuDelegate;
 
 		/**
 		 * Set when the combo button is adding an object.
diff --git a/Engine/Plugins/Developer/Concert/ConcertUI/ConcertClientSharedSlate/Source/ConcertClientSharedSlate/Private/Replication/Editor/View/ObjectEditor/SDefaultReplicationStreamEditor.h b/Engine/Plugins/Developer/Concert/ConcertUI/ConcertClientSharedSlate/Source/ConcertClientSharedSlate/Private/Replication/Editor/View/ObjectEditor/SDefaultReplicationStreamEditor.h
--- a/Engine/Plugins/Developer/Concert/ConcertUI/ConcertClientSharedSlate/Source/ConcertClientSharedSlate/Private/Replication/Editor/View/ObjectEditor/SDefaultReplicationStreamEditor.h
+++ b/Engine/Plugins/Developer/Concert/ConcertUI/ConcertClientSharedSlate/Source/ConcertClientSharedSlate/Private/Replication/Editor/View/ObjectEditor/SDefaultReplicationStreamEditor.h
@@ -55,7 +55,7 @@ namespace UE::ConcertClientSharedSlate
 			SLATE_ARGUMENT(TSharedPtr<ISubobjectModel>, SubobjectModel)
 
 			/** Called to generate the context menu for objects. This extends the options already generated by this widget. */
-			SLATE_EVENT(FExtendMenu, OnExtendObjectsContextMenu)
+			SLATE_EVENT(FExtendObjectMenu, OnExtendObjectsContextMenu)
 
 			/** Optional. Used for determining the order in which properties are displayed. */
 			SLATE_EVENT(FSortPropertyPredicate, SortPropertyRowPredicate)
diff --git a/Engine/Plugins/Developer/Concert/ConcertUI/ConcertClientSharedSlate/Source/ConcertClientSharedSlate/Public/Replication/ReplicationWidgetDelegates.h b/Engine/Plugins/Developer/Concert/ConcertUI/ConcertClientSharedSlate/Source/ConcertClientSharedSlate/Public/Replication/ReplicationWidgetDelegates.h
--- a/Engine/Plugins/Developer/Concert/ConcertUI/ConcertClientSharedSlate/Source/ConcertClientSharedSlate/Public/Replication/ReplicationWidgetDelegates.h
+++ b/Engine/Plugins/Developer/Concert/ConcertUI/ConcertClientSharedSlate/Source/ConcertClientSharedSlate/Public/Replication/ReplicationWidgetDelegates.h
@@ -2,10 +2,12 @@
 
 #pragma once
 
+#include "Containers/ContainersFwd.h"
 #include "Delegates/Delegate.h"
 #include "Delegates/DelegateCombinations.h"
 
 class FMenuBuilder;
+struct FSoftObjectPath;
 
 namespace UE::ConcertClientSharedSlate
 {
@@ -14,6 +16,6 @@ namespace UE::ConcertClientSharedSlate
 	/** A predicate for determining Left < Right. */
 	DECLARE_DELEGATE_RetVal_TwoParams(bool, FSortPropertyPredicate, const FReplicatedPropertyData& Left, const FReplicatedPropertyData& Right);
 
-	/** Extends entries already added to a FMenuBuilder */
-	DECLARE_DELEGATE_OneParam(FExtendMenu, FMenuBuilder&);
+	/** Extends a context menu that is being built for a select of objects. */
+	DECLARE_DELEGATE_TwoParams(FExtendObjectMenu, FMenuBuilder&, TConstArrayView<FSoftObjectPath> ContextObjects);
 }
\ No newline at end of file
diff --git a/Engine/Plugins/Developer/Concert/ConcertUI/ConcertClientSharedSlate/Source/ConcertClientSharedSlate/Public/Replication/ReplicationWidgetFactories.h b/Engine/Plugins/Developer/Concert/ConcertUI/ConcertClientSharedSlate/Source/ConcertClientSharedSlate/Public/Replication/ReplicationWidgetFactories.h
--- a/Engine/Plugins/Developer/Concert/ConcertUI/ConcertClientSharedSlate/Source/ConcertClientSharedSlate/Public/Replication/ReplicationWidgetFactories.h
+++ b/Engine/Plugins/Developer/Concert/ConcertUI/ConcertClientSharedSlate/Source/ConcertClientSharedSlate/Public/Replication/ReplicationWidgetFactories.h
@@ -63,7 +63,7 @@ namespace UE::ConcertClientSharedSlate
 		TSharedPtr<ISubobjectModel> SubobjectModel;
 
 		/** Called to generate the context menu for objects. This extends the options already generated by this widget and is called at the end. */
-		FExtendMenu OnExtendObjectsContextMenu;
+		FExtendObjectMenu OnExtendObjectsContextMenu;
 		
 		/** Additional columns to add to the object view */
 		TArray<ReplicationColumns::FReplicationTopLevelObjectColumn> AdditionalObjectColumns;
