Make parallel request utils thread safe

diff --git a/Engine/Plugins/Developer/Concert/ConcertApp/MultiUserClient/Source/MultiUserClient/Private/Replication/Util/SynchronizedRequestUtils.cpp b/Engine/Plugins/Developer/Concert/ConcertApp/MultiUserClient/Source/MultiUserClient/Private/Replication/Util/SynchronizedRequestUtils.cpp
--- a/Engine/Plugins/Developer/Concert/ConcertApp/MultiUserClient/Source/MultiUserClient/Private/Replication/Util/SynchronizedRequestUtils.cpp
+++ b/Engine/Plugins/Developer/Concert/ConcertApp/MultiUserClient/Source/MultiUserClient/Private/Replication/Util/SynchronizedRequestUtils.cpp
@@ -10,6 +10,7 @@
 
 #include "Algo/AllOf.h"
 #include "Containers/UnrealString.h"
+#include <atomic>
 
 namespace UE::MultiUserClient
 {
@@ -143,16 +144,16 @@ namespace UE::MultiUserClient
 			{
 				LogStreamChangeError(ClientId, Response);
 				
-				--FindOperationByClient(ClientId)->NumOperationsLeft;
 				IntermediateResults.StreamResponses[ClientId] = MoveTemp(Response);
+				FindOperationByClient(ClientId)->bHasCompletedStreamChanges.store(true, std::memory_order_release);
 
 				FinishIfDone();
 			}
 
 			void OnCompleteAuthority(const FGuid& ClientId, FSubmitAuthorityChangesResponse&& Response)
 			{
-				--FindOperationByClient(ClientId)->NumOperationsLeft;
 				IntermediateResults.AuthorityResponses[ClientId] = MoveTemp(Response);
+				FindOperationByClient(ClientId)->bHasCompletedAuthority.store(true, std::memory_order_release);
 				
 				const bool bIsFailure = Response.ErrorCode == EAuthoritySubmissionResponseErrorCode::Timeout || (Response.Response && !Response.Response->RejectedObjects.IsEmpty());
 				UE_CLOG(bIsFailure, LogConcert, Warning, TEXT("Remote authority change to client %s failed"), *ClientId.ToString());
@@ -161,22 +162,37 @@ namespace UE::MultiUserClient
 
 			void OnSubmissionFailure(const FGuid& ClientId)
 			{
-				FindOperationByClient(ClientId)->NumOperationsLeft = 0;
+				FPendingOperation* PendingOperation = FindOperationByClient(ClientId);
+				PendingOperation->bHasCompletedStreamChanges.store(true, std::memory_order_release);
+				PendingOperation->bHasCompletedAuthority.store(true, std::memory_order_release);
 				FinishIfDone();
 			}
 
 		private:
 			
 			TPromise<FParallelExecutionResult> Promise;
-			bool bPromiseWasSet = false;
+			/** Must be atomic because FinishIfDone may be executed in parallel by the game thread (completing an operation) and the UDP (timing out an operation). */
+			std::atomic<bool> bPromiseWasSet = false;
 
 			struct FPendingOperation
 			{
 				/** Handles the submission. Unregisters from the client queue (if we're destroyed early). */
 				FDeferredSubmitter Submitter;
 			
-				/** Used to determine when the operation is done. */
-				int32 NumOperationsLeft = 2; // Stream + Authority = 2
+				/**
+				 * Whether the stream change has completed, i.e. OnCompleteStream has been called.
+				 * 
+				 * Must be atomic because FinishIfDone() may read while OnCompleteStream() / OnCompleteAuthority() are writing.
+				 * The (unlikely) scenario is that the game thread is finishing an operation while the UDP thread is timing out a task.
+				 */
+				std::atomic<bool> bHasCompletedStreamChanges = false;
+				/**
+				 * Whether the authority change has completed, i.e. OnCompleteAuthority has been called.
+				 * 
+				 * Must be atomic because FinishIfDone() may read while OnCompleteStream() / OnCompleteAuthority() are writing.
+				 * The (unlikely) scenario is that the game thread is finishing an operation while the UDP thread is timing out a task.
+				 */
+				std::atomic<bool> bHasCompletedAuthority = false;
 				
 				FPendingOperation(
 					FSyncOperation& InOwner,
@@ -187,8 +203,6 @@ namespace UE::MultiUserClient
 					)
 					: Submitter(InOwner, InClientManager, InSubmissionQueue, InClientId, MoveTemp(InSubmissionParams))
 				{}
-
-				bool IsDone() const { return NumOperationsLeft == 0; }
 			};
 			/**
 			 * Each latent tasks has its own, pre-allocated slot which effectively handles potentially concurrent writes.
@@ -206,12 +220,17 @@ namespace UE::MultiUserClient
 			{
 				const bool bAllDone = Algo::AllOf(SubOperations, [](const FPendingOperation& OperationInfo)
 				{
-					return OperationInfo.IsDone();
+					// This is called right after OnCompleteStream or OnCompleteAuthority.
+					// This may run concurrently: e.g. the game thread is completing an operation and the UDP is timing out an operation.
+					// In that case, both threads wrote to these bools with memory_order_release (in OnCompleteStream and OnCompleteAuthority),
+					// meaning the writes happen-before these reads.
+					return OperationInfo.bHasCompletedStreamChanges.load(std::memory_order_acquire)
+						&& OperationInfo.bHasCompletedAuthority.load(std::memory_order_acquire);
 				});
 
-				if (bAllDone && ensure(!bPromiseWasSet))
+				// In the above scenario, both game thread and UDP thread may have bAllDone. Using exchange, the first thread to write will Emplace the promise.
+				if (bAllDone && !bPromiseWasSet.exchange(true))
 				{
-					bPromiseWasSet = true;
 					Promise.EmplaceValue(MoveTemp(IntermediateResults));
 				}
 			}
