Document and ensure that public API functions are to be called on the game thread only

diff --git a/Engine/Plugins/Developer/Concert/ConcertApp/MultiUserClient/Source/MultiUserClient/Private/Replication/MultiUserReplicationManager.cpp b/Engine/Plugins/Developer/Concert/ConcertApp/MultiUserClient/Source/MultiUserClient/Private/Replication/MultiUserReplicationManager.cpp
--- a/Engine/Plugins/Developer/Concert/ConcertApp/MultiUserClient/Source/MultiUserClient/Private/Replication/MultiUserReplicationManager.cpp
+++ b/Engine/Plugins/Developer/Concert/ConcertApp/MultiUserClient/Source/MultiUserClient/Private/Replication/MultiUserReplicationManager.cpp
@@ -102,7 +102,7 @@ namespace UE::MultiUserClient
 
 	const FObjectReplicationMap* FMultiUserReplicationManager::FindReplicationMapForClient(const FGuid& ClientId) const
 	{
-		if (ConnectedState)
+		if (ConnectedState && ensureMsgf(IsInGameThread(), TEXT("To simplify implementation, only calls from game thread are allowed.")))
 		{
 			const FReplicationClient* ReplicationClient = ConnectedState->ClientManager.FindClient(ClientId);
 			return ReplicationClient
@@ -114,7 +114,7 @@ namespace UE::MultiUserClient
 
 	bool FMultiUserReplicationManager::IsReplicatingObject(const FGuid& ClientId, const FSoftObjectPath& ObjectPath) const 
 	{
-		if (ConnectedState)
+		if (ConnectedState && ensureMsgf(IsInGameThread(), TEXT("To simplify implementation, only calls from game thread are allowed.")))
 		{
 			const FReplicationClient* ReplicationClient = ConnectedState->ClientManager.FindClient(ClientId);
 			return ReplicationClient && ReplicationClient->GetAuthoritySynchronizer().HasAuthorityOver(ObjectPath);
@@ -124,16 +124,27 @@ namespace UE::MultiUserClient
 
 	void FMultiUserReplicationManager::RegisterReplicationDiscoverer(TSharedRef<IReplicationDiscoverer> Discoverer)
 	{
-		DiscoveryContainer.AddDiscoverer(Discoverer);
+		if (ensureMsgf(IsInGameThread(), TEXT("To simplify implementation, only calls from game thread are allowed.")))
+		{
+			DiscoveryContainer.AddDiscoverer(Discoverer);
+		}
 	}
 
 	void FMultiUserReplicationManager::RemoveReplicationDiscoverer(const TSharedRef<IReplicationDiscoverer>& Discoverer)
 	{
-		DiscoveryContainer.RemoveDiscoverer(Discoverer);
+		if (ensureMsgf(IsInGameThread(), TEXT("To simplify implementation, only calls from game thread are allowed.")))
+		{
+			DiscoveryContainer.RemoveDiscoverer(Discoverer);
+		}
 	}
 
 	TSharedRef<IClientChangeOperation> FMultiUserReplicationManager::EnqueueChanges(const FGuid& ClientId, TAttribute<FChangeClientReplicationRequest> SubmissionParams)
 	{
+		if (!ensureMsgf(IsInGameThread(), TEXT("To simplify implementation, only calls from game thread are allowed.")))
+		{
+			return FExternalClientChangeRequestHandler::MakeFailedOperation(EChangeStreamOperationResult::NotOnGameThread, EChangeAuthorityOperationResult::NotOnGameThread);
+		}
+		
 		if (ConnectedState)
 		{
 			FReplicationClient* ReplicationClient = ConnectedState->ClientManager.FindClient(ClientId);
diff --git a/Engine/Plugins/Developer/Concert/ConcertApp/MultiUserClient/Source/MultiUserClient/Private/Replication/Submission/External/ExternalClientChangeRequestHandler.cpp b/Engine/Plugins/Developer/Concert/ConcertApp/MultiUserClient/Source/MultiUserClient/Private/Replication/Submission/External/ExternalClientChangeRequestHandler.cpp
--- a/Engine/Plugins/Developer/Concert/ConcertApp/MultiUserClient/Source/MultiUserClient/Private/Replication/Submission/External/ExternalClientChangeRequestHandler.cpp
+++ b/Engine/Plugins/Developer/Concert/ConcertApp/MultiUserClient/Source/MultiUserClient/Private/Replication/Submission/External/ExternalClientChangeRequestHandler.cpp
@@ -65,6 +65,7 @@ namespace UE::MultiUserClient
 
 	TSharedRef<IClientChangeOperation> FExternalClientChangeRequestHandler::HandleRequest(TAttribute<FChangeClientReplicationRequest> SubmissionParams)
 	{
+		check(IsInGameThread());
 		TSharedRef<FClientChangeOperation> Result = FClientChangeOperation::StartOperation(
 			StreamId,
 			GetStreamContentDelegate,
diff --git a/Engine/Plugins/Developer/Concert/ConcertApp/MultiUserClient/Source/MultiUserClient/Public/Replication/ChangeOperationTypes.h b/Engine/Plugins/Developer/Concert/ConcertApp/MultiUserClient/Source/MultiUserClient/Public/Replication/ChangeOperationTypes.h
--- a/Engine/Plugins/Developer/Concert/ConcertApp/MultiUserClient/Source/MultiUserClient/Public/Replication/ChangeOperationTypes.h
+++ b/Engine/Plugins/Developer/Concert/ConcertApp/MultiUserClient/Source/MultiUserClient/Public/Replication/ChangeOperationTypes.h
@@ -44,6 +44,8 @@ namespace UE::MultiUserClient
 		NotInSession,
 		/** No submission took place because the client ID was invalid.  */
 		UnknownClient,
+		/** The request was not made on the game thread. */
+		NotOnGameThread,
 
 		/** Not an actual return code. Make sure it's always last. */
 		Count
@@ -77,6 +79,8 @@ namespace UE::MultiUserClient
 		NotInSession,
 		/** No submission took place because the client ID was invalid.  */
 		UnknownClient,
+		/** The request was not made on the game thread. */
+		NotOnGameThread,
 
 		/** Not an actual return code. Make sure it's always last. */
 		Count
diff --git a/Engine/Plugins/Developer/Concert/ConcertApp/MultiUserClient/Source/MultiUserClient/Public/Replication/IClientChangeOperation.h b/Engine/Plugins/Developer/Concert/ConcertApp/MultiUserClient/Source/MultiUserClient/Public/Replication/IClientChangeOperation.h
--- a/Engine/Plugins/Developer/Concert/ConcertApp/MultiUserClient/Source/MultiUserClient/Public/Replication/IClientChangeOperation.h
+++ b/Engine/Plugins/Developer/Concert/ConcertApp/MultiUserClient/Source/MultiUserClient/Public/Replication/IClientChangeOperation.h
@@ -9,8 +9,8 @@
 namespace UE::MultiUserClient
 {
 	/**
-	 * Abstracts an operation that is changing a client's stream, authority, or both.
-	 * They are the result of IMultiUserReplication::EnqueueChanges.
+	 * Exposes events that occur as part of changing a client's stream, authority, or both.
+	 * This is the result of IMultiUserReplication::EnqueueChanges.
 	 *
 	 * IClientChangeOperation wraps ISubmissionOperation, which is internal to Multi-User.
 	 * This allows the two interface to change independently.
@@ -22,21 +22,21 @@ namespace UE::MultiUserClient
 		/**
 		 * Completes when the operation of changing streams has completed.
 		 * @note This can be called at most once; subsequent calls result in an unset future.
-		 * @note This future can execute on any thread!
+		 * @note This future can complete on any thread.
 		 */
 		virtual TFuture<EChangeStreamOperationResult> OnChangeStream() = 0;
 		
 		/**
 		 * Completes when the operation of changing authority has completed.
 		 * @note This can be called at most once; subsequent calls result in an unset future.
-		 * @note This future can execute on any thread!
+		 * @note This future can complete on any thread.
 		 */
 		virtual TFuture<EChangeAuthorityOperationResult> OnChangeAuthority() = 0;
 		
 		/**
 		 * Completes all sub-operations have completed.
 		 * @note This can be called at most once; subsequent calls result in an unset future.
-		 * @note This future can execute on any thread!
+		 * @note This future can complete on any thread.
 		 */
 		virtual TFuture<FChangeClientReplicationResult> OnOperationCompleted() = 0;
 		
diff --git a/Engine/Plugins/Developer/Concert/ConcertApp/MultiUserClient/Source/MultiUserClient/Public/Replication/IMultiUserReplication.h b/Engine/Plugins/Developer/Concert/ConcertApp/MultiUserClient/Source/MultiUserClient/Public/Replication/IMultiUserReplication.h
--- a/Engine/Plugins/Developer/Concert/ConcertApp/MultiUserClient/Source/MultiUserClient/Public/Replication/IMultiUserReplication.h
+++ b/Engine/Plugins/Developer/Concert/ConcertApp/MultiUserClient/Source/MultiUserClient/Public/Replication/IMultiUserReplication.h
@@ -22,9 +22,15 @@ namespace UE::MultiUserClient
 
 		/**
 		 * @return Gets the last known server map of objects registered for replication for a given client.
-		 * This server state is regularly polled whilst the local client state should always be in synch.*/
+		 * 
+		 * This server state is regularly polled whilst the local client state should always be in synch.
+		 * This function must be called from the game thread.
+		 */
 		virtual const FObjectReplicationMap* FindReplicationMapForClient(const FGuid& ClientId) const = 0;
-		/** @return Whether the local editor instance thinks the client has authority over the properties it has registered to ObjectPath. */
+		/**
+		 * @return Whether the local editor instance thinks the client has authority over the properties it has registered to ObjectPath.
+		 * This function must be called from the game thread.
+		 */
 		virtual bool IsReplicatingObject(const FGuid& ClientId, const FSoftObjectPath& ObjectPath) const = 0;
 
 		/**
@@ -33,9 +39,11 @@ namespace UE::MultiUserClient
 		 * It is used to automatically configure UObjects for replication when appropriate:
 		 * - When an user adds an object via Add Actor button
 		 * - When an UObject is added to the world via a transaction (run on the client machine that adds the UObject)
+		 * This function must be called from the game thread.
+		 * 
 		 */
 		virtual void RegisterReplicationDiscoverer(TSharedRef<IReplicationDiscoverer> Discoverer) = 0;
-		/** Unregisters a previously registered discoverer */
+		/** Unregisters a previously registered discoverer. This function must be called from the game thread. */
 		virtual void RemoveReplicationDiscoverer(const TSharedRef<IReplicationDiscoverer>& Discoverer) = 0;
 
 		/**
@@ -46,6 +54,8 @@ namespace UE::MultiUserClient
 		 * The authority state specifies which of the registered objects should actually be sending data.
 		 * The stream change is requested first and is followed by the authority change.
 		 *
+		 * This function must be called from the game thread.
+		 *
 		 * @param ClientId The client for which to change authority
 		 * @param SubmissionParams Once the request is ready to be sent to the server, this attribute is used to generate the change request
 		 */
diff --git a/Engine/Plugins/Developer/Concert/ConcertApp/MultiUserClient/Source/MultiUserClientLibrary/Private/Replication/Async/ChangeClientBlueprintParams.cpp b/Engine/Plugins/Developer/Concert/ConcertApp/MultiUserClient/Source/MultiUserClientLibrary/Private/Replication/Async/ChangeClientBlueprintParams.cpp
--- a/Engine/Plugins/Developer/Concert/ConcertApp/MultiUserClient/Source/MultiUserClientLibrary/Private/Replication/Async/ChangeClientBlueprintParams.cpp
+++ b/Engine/Plugins/Developer/Concert/ConcertApp/MultiUserClient/Source/MultiUserClientLibrary/Private/Replication/Async/ChangeClientBlueprintParams.cpp
@@ -7,12 +7,12 @@ namespace UE::MultiUserClientLibrary
 #if WITH_CONCERT
 	EMultiUserChangeStreamOperationResult Transform(MultiUserClient::EChangeStreamOperationResult Data)
 	{
-		static_assert(static_cast<int32>(MultiUserClient::EChangeStreamOperationResult::Count) == 8, "Update this location if you changed the enum");
+		static_assert(static_cast<int32>(MultiUserClient::EChangeStreamOperationResult::Count) == 9, "Update EMultiUserChangeStreamOperationResult to have the equivalent enum entry you added to EChangeStreamOperationResult");
 		return static_cast<EMultiUserChangeStreamOperationResult>(Data);
 	}
 	EMultiUserChangeAuthorityOperationResult Transform(MultiUserClient::EChangeAuthorityOperationResult Data)
 	{
-		static_assert(static_cast<int32>(MultiUserClient::EChangeAuthorityOperationResult::Count) == 9, "Update this location if you changed the enum");
+		static_assert(static_cast<int32>(MultiUserClient::EChangeAuthorityOperationResult::Count) == 10, "Update EMultiUserChangeStreamOperationResult to have the equivalent enum entry you added to EChangeStreamOperationResult");
 		return static_cast<EMultiUserChangeAuthorityOperationResult>(Data);
 	}
 	FMultiUserChangeClientReplicationResult Transform(MultiUserClient::FChangeClientReplicationResult Data)
diff --git a/Engine/Plugins/Developer/Concert/ConcertApp/MultiUserClient/Source/MultiUserClientLibrary/Public/Replication/Async/ChangeClientBlueprintParams.h b/Engine/Plugins/Developer/Concert/ConcertApp/MultiUserClient/Source/MultiUserClientLibrary/Public/Replication/Async/ChangeClientBlueprintParams.h
--- a/Engine/Plugins/Developer/Concert/ConcertApp/MultiUserClient/Source/MultiUserClientLibrary/Public/Replication/Async/ChangeClientBlueprintParams.h
+++ b/Engine/Plugins/Developer/Concert/ConcertApp/MultiUserClient/Source/MultiUserClientLibrary/Public/Replication/Async/ChangeClientBlueprintParams.h
@@ -38,6 +38,8 @@ enum class EMultiUserChangeStreamOperationResult : uint8
 	NotInSession,
 	/** No submission took place because the client ID was invalid.  */
 	UnknownClient,
+	/** The request was not made on the game thread. */
+	NotOnGameThread,
 	/** The MultiUserClient module is not available. Usually when you're in a runtime build but the operation can only be run in editor builds. */
 	NotAvailable
 };
@@ -71,6 +73,8 @@ enum class EMultiUserChangeAuthorityOperationResult : uint8
 	NotInSession,
 	/** No submission took place because the client ID was invalid.  */
 	UnknownClient,
+	/** The request was not made on the game thread. */
+	NotOnGameThread,
 	/** The MultiUserClient module is not available. Usually when you're in a runtime build but the operation can only be run in editor builds. */
 	NotAvailable
 };
