Expose default replication settings in Multi User settings

diff --git a/Engine/Plugins/Developer/Concert/ConcertApp/MultiUserClient/Source/MultiUserClient/Private/MultiUserClientModule.cpp b/Engine/Plugins/Developer/Concert/ConcertApp/MultiUserClient/Source/MultiUserClient/Private/MultiUserClientModule.cpp
--- a/Engine/Plugins/Developer/Concert/ConcertApp/MultiUserClient/Source/MultiUserClient/Private/MultiUserClientModule.cpp
+++ b/Engine/Plugins/Developer/Concert/ConcertApp/MultiUserClient/Source/MultiUserClient/Private/MultiUserClientModule.cpp
@@ -15,7 +15,9 @@
 #include "IConcertSyncClient.h"
 #include "IConcertSyncClientModule.h"
 #include "MultiUserClientUtils.h"
+#include "MultiUserReplicationStyle.h"
 #include "Replication/MultiUserReplicationManager.h"
+#include "Settings/MultiUserReplicationSettings.h"
 
 #include "Misc/App.h"
 #include "Misc/AsyncTaskNotification.h"
@@ -50,8 +52,6 @@
 #include "DetailLayoutBuilder.h"
 #include "DetailCategoryBuilder.h"
 #include "DetailWidgetRow.h"
-#include "MultiUserReplicationSettings.h"
-#include "MultiUserReplicationStyle.h"
 
 
 #include "Widgets/SConcertBrowser.h"
diff --git a/Engine/Plugins/Developer/Concert/ConcertApp/MultiUserClient/Source/MultiUserClient/Private/MultiUserReplicationSettings.h b/Engine/Plugins/Developer/Concert/ConcertApp/MultiUserClient/Source/MultiUserClient/Private/MultiUserReplicationSettings.h
deleted file mode 100644
diff --git a/Engine/Plugins/Developer/Concert/ConcertApp/MultiUserClient/Source/MultiUserClient/Private/Replication/Authority/AuthorityChangeTracker.h b/Engine/Plugins/Developer/Concert/ConcertApp/MultiUserClient/Source/MultiUserClient/Private/Replication/Authority/AuthorityChangeTracker.h
--- a/Engine/Plugins/Developer/Concert/ConcertApp/MultiUserClient/Source/MultiUserClient/Private/Replication/Authority/AuthorityChangeTracker.h
+++ b/Engine/Plugins/Developer/Concert/ConcertApp/MultiUserClient/Source/MultiUserClient/Private/Replication/Authority/AuthorityChangeTracker.h
@@ -22,8 +22,12 @@ namespace UE::MultiUserClient
 	class FAuthorityChangeTracker
 	{
 	public:
-		
-		FAuthorityChangeTracker(const FGuid& InClientId, const IClientAuthoritySynchronizer& InAuthoritySynchronizer, FGlobalAuthorityCache& InAuthorityCache);
+		/**
+		 * @param InClientId The endpoint ID of the owning client
+		 * @param InAuthoritySynchronizer Used to get the authority state on the server. The caller ensures it outlives the constructed instance.
+		 * @param InAuthorityCache Used to predict authority conflicts. The caller ensures it outlives the constructed instance.
+		 */
+		FAuthorityChangeTracker(const FGuid& InClientId, const IClientAuthoritySynchronizer& InAuthoritySynchronizer UE_LIFETIMEBOUND, FGlobalAuthorityCache& InAuthorityCache UE_LIFETIMEBOUND);
 		~FAuthorityChangeTracker();
 		
 		/** Marks that the authority should be changed to bNewAuthorityState. */
@@ -34,9 +38,6 @@ namespace UE::MultiUserClient
 		/** Diffs NewAuthorityStates to the current authority states and removes entries. */
 		void RefreshChanges();
 		
-		void ClearChanges() { NewAuthorityStates.Reset(); }
-		bool HasChanges() const { return !NewAuthorityStates.IsEmpty(); }
-		
 		/** Gets the authority state the object will have if the changes are applied. */
 		bool GetAuthorityStateAfterApplied(const FSoftObjectPath& ObjectPath) const;
 		
diff --git a/Engine/Plugins/Developer/Concert/ConcertApp/MultiUserClient/Source/MultiUserClient/Private/Replication/Client/ReplicationClient.cpp b/Engine/Plugins/Developer/Concert/ConcertApp/MultiUserClient/Source/MultiUserClient/Private/Replication/Client/ReplicationClient.cpp
--- a/Engine/Plugins/Developer/Concert/ConcertApp/MultiUserClient/Source/MultiUserClient/Private/Replication/Client/ReplicationClient.cpp
+++ b/Engine/Plugins/Developer/Concert/ConcertApp/MultiUserClient/Source/MultiUserClient/Private/Replication/Client/ReplicationClient.cpp
@@ -10,6 +10,7 @@
 #include "Replication/Submission/Data/AuthoritySubmission.h"
 
 #include "Misc/CoreDelegates.h"
+#include "Settings/MultiUserReplicationSettings.h"
 
 namespace UE::MultiUserClient
 {
@@ -44,8 +45,9 @@ namespace UE::MultiUserClient
 			FStreamChangeTracker::FOnModifyReplicationMap::CreateLambda([this](){ ClientContentStorage->Stream->Modify(); })
 			)
 		, LocalAuthorityDiffer(EndpointId, *AuthoritySynchronizer, InAuthorityCache)
-		, ChangeRequestBuilder(EndpointId, InAuthorityCache, *StreamSynchronizer, LocalClientStreamDiffer, LocalAuthorityDiffer)
-		, AutoSubmissionPolicy(SubmissionQueue, ChangeRequestBuilder, LocalClientEditModel.Get(), LocalAuthorityDiffer)
+		, LocalFrequencyChangeTracker(*StreamSynchronizer)
+		, ChangeRequestBuilder(EndpointId, InAuthorityCache, *StreamSynchronizer, LocalClientStreamDiffer, LocalAuthorityDiffer, LocalFrequencyChangeTracker)
+		, AutoSubmissionPolicy(SubmissionQueue, ChangeRequestBuilder, LocalClientEditModel.Get(), LocalAuthorityDiffer, LocalFrequencyChangeTracker)
 	{
 		LocalClientEditModel->OnObjectsChanged().AddRaw(this, &FReplicationClient::OnObjectsChanged);
 		LocalClientEditModel->OnPropertiesChanged().AddRaw(this, &FReplicationClient::OnPropertiesChanged);
@@ -83,7 +85,6 @@ namespace UE::MultiUserClient
 	{
 		// Whenever this client's server state changes, the UI must be refreshed.
 		GetClientContent()->Stream->ReplicationMap = GetStreamSynchronizer().GetServerState();
-
 		DeferOnModelChanged();
 	}
 
@@ -116,6 +117,8 @@ namespace UE::MultiUserClient
 		TakeAuthorityOverNewlyAddedObjects(ChangeData);
 		// Refresh because local authority changes may no longer be valid after modifying the stream
 		LocalAuthorityDiffer.RefreshChanges();
+		// Check whether added objects are supposed to default to some special frequency settings
+		ApplyDefaultFrequencySettings(ChangeData);
 
 		// Finally, let everybody else know.
 		OnModelChangedDelegate.Broadcast();
@@ -138,6 +141,24 @@ namespace UE::MultiUserClient
 		LocalAuthorityDiffer.SetAuthorityIfAllowed(ObjectPaths, true);
 	}
 
+	void FReplicationClient::ApplyDefaultFrequencySettings(const FDeferredOnModelChangedData& ChangeData)
+	{
+		UMultiUserReplicationSettings* Settings = UMultiUserReplicationSettings::Get();
+		for (const TWeakObjectPtr<UObject>& AddedObject : ChangeData.AccumulatedAddedObjects)
+		{
+			if (!AddedObject.IsValid())
+			{
+				continue;
+			}
+			
+			const TOptional<FConcertObjectReplicationSettings> FrequencyOverride = Settings->DetermineObjectFrequencySettings(*AddedObject);
+			if (FrequencyOverride)
+			{
+				LocalFrequencyChangeTracker.AddOverride(AddedObject.Get(), *FrequencyOverride);
+			}
+		}
+	}
+
 	void FReplicationClient::OnAuthoritySubmissionCompleted(const FSubmitAuthorityChangesRequest& Request, const FSubmitAuthorityChangesResponse& Response)
 	{
 		if (!Response.Response)
diff --git a/Engine/Plugins/Developer/Concert/ConcertApp/MultiUserClient/Source/MultiUserClient/Private/Replication/Client/ReplicationClient.h b/Engine/Plugins/Developer/Concert/ConcertApp/MultiUserClient/Source/MultiUserClient/Private/Replication/Client/ReplicationClient.h
--- a/Engine/Plugins/Developer/Concert/ConcertApp/MultiUserClient/Source/MultiUserClient/Private/Replication/Client/ReplicationClient.h
+++ b/Engine/Plugins/Developer/Concert/ConcertApp/MultiUserClient/Source/MultiUserClient/Private/Replication/Client/ReplicationClient.h
@@ -4,6 +4,7 @@
 
 #include "Replication/Authority/AuthorityChangeTracker.h"
 #include "Replication/Authority/IClientAuthoritySynchronizer.h"
+#include "Replication/Frequency/FrequencyChangeTracker.h"
 #include "Replication/Stream/IClientStreamSynchronizer.h"
 #include "Replication/Stream/StreamChangeTracker.h"
 #include "Replication/Submission/AutoSubmissionPolicy.h"
@@ -48,9 +49,9 @@ namespace UE::MultiUserClient
 		 */
 		FReplicationClient(
 			const FGuid& EndpointId,
-			FReplicationDiscoveryContainer& InDiscoveryContainer,
-			FGlobalAuthorityCache& InAuthorityCache,
-			UMultiUserReplicationClientPreset& InSessionContent,
+			FReplicationDiscoveryContainer& InDiscoveryContainer UE_LIFETIMEBOUND,
+			FGlobalAuthorityCache& InAuthorityCache UE_LIFETIMEBOUND,
+			UMultiUserReplicationClientPreset& InSessionContent UE_LIFETIMEBOUND,
 			TUniquePtr<IClientStreamSynchronizer> InStreamSynchronizer,
 			TUniquePtr<IClientAuthoritySynchronizer> InAuthoritySynchronizer,
 			TUniquePtr<ISubmissionWorkflow> InSubmissionWorkflow
@@ -136,6 +137,8 @@ namespace UE::MultiUserClient
 		FStreamChangeTracker LocalClientStreamDiffer;
 		/** Tracks changes made to the client's authority state. */
 		FAuthorityChangeTracker LocalAuthorityDiffer;
+		/** Tracks enqueued local changes to object replication frequency. */
+		FFrequencyChangeTracker LocalFrequencyChangeTracker;
 		
 		/** Shared logic for building stream and authority change requests based on local change made. */
 		FChangeRequestBuilder ChangeRequestBuilder;
@@ -170,6 +173,9 @@ namespace UE::MultiUserClient
 		
 		/** Takes authority over newly added objects for better UX */
 		void TakeAuthorityOverNewlyAddedObjects(const FDeferredOnModelChangedData& ChangeData);
+
+		/** Applies default replication frequency settings specified in the Multi User settings */
+		void ApplyDefaultFrequencySettings(const FDeferredOnModelChangedData& ChangeData);
 		
 		/** Removes authority if request fails */
 		void OnAuthoritySubmissionCompleted(const FSubmitAuthorityChangesRequest& Request, const FSubmitAuthorityChangesResponse& Response);
diff --git a/Engine/Plugins/Developer/Concert/ConcertApp/MultiUserClient/Source/MultiUserClient/Private/Replication/Frequency/FrequencyChangeTracker.cpp b/Engine/Plugins/Developer/Concert/ConcertApp/MultiUserClient/Source/MultiUserClient/Private/Replication/Frequency/FrequencyChangeTracker.cpp
new file mode 100644
--- /dev/null
+++ b/Engine/Plugins/Developer/Concert/ConcertApp/MultiUserClient/Source/MultiUserClient/Private/Replication/Frequency/FrequencyChangeTracker.cpp
@@ -0,0 +1,69 @@
+// Copyright Epic Games, Inc. All Rights Reserved.
+
+#include "FrequencyChangeTracker.h"
+
+#include "Replication/Stream/IClientStreamSynchronizer.h"
+#include "Settings/MultiUserReplicationSettings.h"
+
+namespace UE::MultiUserClient
+{
+	FFrequencyChangeTracker::FFrequencyChangeTracker(
+		IClientStreamSynchronizer& InStreamSynchronizer
+		)
+		: StreamSynchronizer(InStreamSynchronizer)
+	{
+		// Refresh locally recorded changes when the server state changes, e.g. after a request has been served.
+		StreamSynchronizer.OnServerStateChanged().AddRaw(this, &FFrequencyChangeTracker::RefreshChanges);
+	}
+
+	FFrequencyChangeTracker::~FFrequencyChangeTracker()
+	{
+		StreamSynchronizer.OnServerStateChanged().RemoveAll(this);
+	}
+
+	void FFrequencyChangeTracker::AddOverride(FSoftObjectPath Object, FConcertObjectReplicationSettings NewSettings)
+	{
+		RecordedChanges.OverridesToAdd.Emplace(MoveTemp(Object), MoveTemp(NewSettings));
+	}
+	
+	FFrequencyChangelist FFrequencyChangeTracker::BuildForSubmission(const FStreamChangelist& ObjectChanges)
+	{
+		FFrequencyChangelist Result = RecordedChanges;
+
+		// For now we only support adding overrides.
+		const FGuid StreamId = StreamSynchronizer.GetStreamId();
+		const FObjectReplicationMap& ReplicationMap = StreamSynchronizer.GetServerState();
+		for (auto IteratorForAdds = Result.OverridesToAdd.CreateIterator(); IteratorForAdds; ++IteratorForAdds)
+		{
+			const TPair<FSoftObjectPath, FConcertObjectReplicationSettings>& Pair = *IteratorForAdds;
+			if (!ReplicationMap.HasProperties(Pair.Key) && !ObjectChanges.ObjectsToPut.Contains({ StreamId, Pair.Key }))
+			{
+				IteratorForAdds.RemoveCurrent();
+			}
+		}
+
+		// Set frequency defaults if either this is the first time we're submitting to the server or the user has changes the project settings since.
+		const FConcertObjectReplicationSettings& DefaultsToSet = UMultiUserReplicationSettings::Get()->FrequencyRules.DefaultObjectFrequencySettings;
+		if (DefaultsToSet != StreamSynchronizer.GetFrequencySettings().Defaults)
+		{
+			Result.NewDefaults = DefaultsToSet;
+		}
+		
+		return Result;
+	}
+
+	void FFrequencyChangeTracker::RefreshChanges()
+	{
+		const FObjectReplicationMap& ReplicationMap = StreamSynchronizer.GetServerState();
+		for (auto IteratorForAdds = RecordedChanges.OverridesToAdd.CreateIterator(); IteratorForAdds; ++IteratorForAdds)
+		{
+			const TPair<FSoftObjectPath, FConcertObjectReplicationSettings>& Pair = *IteratorForAdds;
+			if (!ReplicationMap.HasProperties(Pair.Key))
+			{
+				IteratorForAdds.RemoveCurrent();
+			}
+		}
+
+		// In the future we must also cleans RecordedChanges.OverridesToRemove if we add a function like RemoveOverride
+	}
+}
diff --git a/Engine/Plugins/Developer/Concert/ConcertApp/MultiUserClient/Source/MultiUserClient/Private/Replication/Frequency/FrequencyChangeTracker.h b/Engine/Plugins/Developer/Concert/ConcertApp/MultiUserClient/Source/MultiUserClient/Private/Replication/Frequency/FrequencyChangeTracker.h
new file mode 100644
--- /dev/null
+++ b/Engine/Plugins/Developer/Concert/ConcertApp/MultiUserClient/Source/MultiUserClient/Private/Replication/Frequency/FrequencyChangeTracker.h
@@ -0,0 +1,54 @@
+// Copyright Epic Games, Inc. All Rights Reserved.
+
+#pragma once
+
+#include "Replication/Util/StreamRequestUtils.h"
+#include "Templates/UnrealTemplate.h"
+#include "UObject/UnrealType.h"
+
+struct FConcertObjectReplicationSettings;
+struct FSoftObjectPath;
+
+namespace UE::MultiUserClient
+{
+	class IClientStreamSynchronizer;
+	class FStreamChangeTracker;
+	
+	/** Used to queue up local made changes to object frequency settings. */
+	class FFrequencyChangeTracker : public FNoncopyable
+    {
+    public:
+
+		FFrequencyChangeTracker(
+			IClientStreamSynchronizer& InStreamSynchronizer UE_LIFETIMEBOUND
+			);
+		~FFrequencyChangeTracker();
+
+		/** Adds a frequency override for the given object. */
+		void AddOverride(FSoftObjectPath Object, FConcertObjectReplicationSettings NewSettings);
+		
+		// In the future we could add RemoveOverride if needed
+
+		/** Compares the recorded changes to about to be submitted ObjectChanges and strips any invalid changes. */
+		FFrequencyChangelist BuildForSubmission(const FStreamChangelist& ObjectChanges);
+
+		/** Compares the local changes against the server state and optionally removes changes. */
+		void RefreshChanges();
+
+		DECLARE_MULTICAST_DELEGATE(FOnFrequencySettingsChanged);
+		/** Called when frequency settings have changed been recorded that need to be submitted to the server. */
+		FOnFrequencySettingsChanged& OnFrequencySettingsChanged() { return OnFrequencySettingsChangedDelegate; } 
+
+	private:
+
+		/** Tells us when the server state changes so RecordedChanges can be cleansed of applied changes. */
+		IClientStreamSynchronizer& StreamSynchronizer;
+
+		/** The changes recorded thus far. */
+		FFrequencyChangelist RecordedChanges;
+
+		/** Called when frequency settings have changed been recorded that need to be submitted to the server. */
+		FOnFrequencySettingsChanged OnFrequencySettingsChangedDelegate;
+    };
+}
+
diff --git a/Engine/Plugins/Developer/Concert/ConcertApp/MultiUserClient/Source/MultiUserClient/Private/Replication/Stream/Discovery/MultiUserStreamExtender.cpp b/Engine/Plugins/Developer/Concert/ConcertApp/MultiUserClient/Source/MultiUserClient/Private/Replication/Stream/Discovery/MultiUserStreamExtender.cpp
--- a/Engine/Plugins/Developer/Concert/ConcertApp/MultiUserClient/Source/MultiUserClient/Private/Replication/Stream/Discovery/MultiUserStreamExtender.cpp
+++ b/Engine/Plugins/Developer/Concert/ConcertApp/MultiUserClient/Source/MultiUserClient/Private/Replication/Stream/Discovery/MultiUserStreamExtender.cpp
@@ -2,11 +2,11 @@
 
 #include "MultiUserStreamExtender.h"
 
-#include "MultiUserReplicationSettings.h"
 #include "ReplicationDiscoveryContainer.h"
 #include "Replication/IReplicationDiscoveryContext.h"
 #include "Replication/Data/ConcertPropertySelection.h"
 #include "Replication/Editor/Model/Extension/IStreamExtensionContext.h"
+#include "Settings/MultiUserReplicationSettings.h"
 
 namespace UE::MultiUserClient
 {
diff --git a/Engine/Plugins/Developer/Concert/ConcertApp/MultiUserClient/Source/MultiUserClient/Private/Replication/Stream/StreamChangeTracker.h b/Engine/Plugins/Developer/Concert/ConcertApp/MultiUserClient/Source/MultiUserClient/Private/Replication/Stream/StreamChangeTracker.h
--- a/Engine/Plugins/Developer/Concert/ConcertApp/MultiUserClient/Source/MultiUserClient/Private/Replication/Stream/StreamChangeTracker.h
+++ b/Engine/Plugins/Developer/Concert/ConcertApp/MultiUserClient/Source/MultiUserClient/Private/Replication/Stream/StreamChangeTracker.h
@@ -21,7 +21,7 @@ namespace UE::MultiUserClient
 {
 	/**
 	 * Knows of the local client's registered replication streams and builds a changelist. The changelist tracks the
-	 * unconfirmed changes to the client's streams and is updates when the server confirms the change.
+	 * unconfirmed changes to the client's streams and is updated when the server confirms the change.
 	 */
 	class FStreamChangeTracker : public FNoncopyable
 	{
diff --git a/Engine/Plugins/Developer/Concert/ConcertApp/MultiUserClient/Source/MultiUserClient/Private/Replication/Submission/AutoSubmissionPolicy.cpp b/Engine/Plugins/Developer/Concert/ConcertApp/MultiUserClient/Source/MultiUserClient/Private/Replication/Submission/AutoSubmissionPolicy.cpp
--- a/Engine/Plugins/Developer/Concert/ConcertApp/MultiUserClient/Source/MultiUserClient/Private/Replication/Submission/AutoSubmissionPolicy.cpp
+++ b/Engine/Plugins/Developer/Concert/ConcertApp/MultiUserClient/Source/MultiUserClient/Private/Replication/Submission/AutoSubmissionPolicy.cpp
@@ -6,6 +6,7 @@
 #include "ISubmissionWorkflow.h"
 #include "Replication/Authority/AuthorityChangeTracker.h"
 #include "Replication/Editor/Model/IEditableReplicationStreamModel.h"
+#include "Replication/Frequency/FrequencyChangeTracker.h"
 #include "Replication/Submission/Queue/SubmissionQueue.h"
 
 namespace UE::MultiUserClient
@@ -14,17 +15,20 @@ namespace UE::MultiUserClient
 		FSubmissionQueue& InSubmissionQueue,
 		const FChangeRequestBuilder& InRequestBuilder,
 		ConcertClientSharedSlate::IEditableReplicationStreamModel& InStreamEditorModel,
-		FAuthorityChangeTracker& InAuthorityChangeTracker
+		FAuthorityChangeTracker& InAuthorityChangeTracker,
+		FFrequencyChangeTracker& InFrequencyChangeTracker
 		)
 		: FSelfUnregisteringDeferredSubmitter(InSubmissionQueue)
 		, SubmissionQueue(InSubmissionQueue)
 		, RequestBuilder(InRequestBuilder)
 		, StreamEditorModel(InStreamEditorModel)
 		, AuthorityChangeTracker(InAuthorityChangeTracker)
+		, FrequencyChangeTracker(InFrequencyChangeTracker)
 	{
 		StreamEditorModel.OnObjectsChanged().AddRaw(this, &FAutoSubmissionPolicy::OnObjectsChanged);
 		StreamEditorModel.OnPropertiesChanged().AddRaw(this, &FAutoSubmissionPolicy::OnChangesDetected);
 		AuthorityChangeTracker.OnAddedOwnedObjects().AddRaw(this, &FAutoSubmissionPolicy::OnChangesDetected);
+		FrequencyChangeTracker.OnFrequencySettingsChanged().AddRaw(this, &FAutoSubmissionPolicy::OnChangesDetected);
 	}
 
 	FAutoSubmissionPolicy::~FAutoSubmissionPolicy()
@@ -32,6 +36,7 @@ namespace UE::MultiUserClient
 		StreamEditorModel.OnObjectsChanged().RemoveAll(this);
 		StreamEditorModel.OnPropertiesChanged().RemoveAll(this);
 		AuthorityChangeTracker.OnAddedOwnedObjects().RemoveAll(this);
+		FrequencyChangeTracker.OnFrequencySettingsChanged().RemoveAll(this);
 	}
 
 	void FAutoSubmissionPolicy::ProcessAccumulatedChangesAndSubmit()
diff --git a/Engine/Plugins/Developer/Concert/ConcertApp/MultiUserClient/Source/MultiUserClient/Private/Replication/Submission/AutoSubmissionPolicy.h b/Engine/Plugins/Developer/Concert/ConcertApp/MultiUserClient/Source/MultiUserClient/Private/Replication/Submission/AutoSubmissionPolicy.h
--- a/Engine/Plugins/Developer/Concert/ConcertApp/MultiUserClient/Source/MultiUserClient/Private/Replication/Submission/AutoSubmissionPolicy.h
+++ b/Engine/Plugins/Developer/Concert/ConcertApp/MultiUserClient/Source/MultiUserClient/Private/Replication/Submission/AutoSubmissionPolicy.h
@@ -6,6 +6,11 @@
 #include "Replication/Submission/Queue/DeferredSubmitter.h"
 #include "Templates/UnrealTemplate.h"
 
+namespace UE::MultiUserClient
+{
+	class FFrequencyChangeTracker;
+}
+
 namespace UE::ConcertClientSharedSlate
 {
 	class IEditableReplicationStreamModel;
@@ -33,10 +38,11 @@ namespace UE::MultiUserClient
 	public:
 		
 		FAutoSubmissionPolicy(
-			FSubmissionQueue& InSubmissionQueue,
-			const FChangeRequestBuilder& InRequestBuilder,
-			ConcertClientSharedSlate::IEditableReplicationStreamModel& InStreamEditorModel,
-			FAuthorityChangeTracker& InAuthorityChangeTracker
+			FSubmissionQueue& InSubmissionQueue UE_LIFETIMEBOUND,
+			const FChangeRequestBuilder& InRequestBuilder UE_LIFETIMEBOUND,
+			ConcertClientSharedSlate::IEditableReplicationStreamModel& InStreamEditorModel UE_LIFETIMEBOUND,
+			FAuthorityChangeTracker& InAuthorityChangeTracker UE_LIFETIMEBOUND,
+			FFrequencyChangeTracker& InFrequencyChangeTracker UE_LIFETIMEBOUND
 			);
 		~FAutoSubmissionPolicy();
 
@@ -54,6 +60,8 @@ namespace UE::MultiUserClient
 		ConcertClientSharedSlate::IEditableReplicationStreamModel& StreamEditorModel;
 		/** Informs us when authority is changed by the user. */
 		FAuthorityChangeTracker& AuthorityChangeTracker;
+		/** Informs us when frequency settings are changed by the user. */
+		FFrequencyChangeTracker& FrequencyChangeTracker;
 
 		/** Whether any changes were made. */
 		bool bIsDirty = false;
diff --git a/Engine/Plugins/Developer/Concert/ConcertApp/MultiUserClient/Source/MultiUserClient/Private/Replication/Submission/ChangeRequestBuilder.cpp b/Engine/Plugins/Developer/Concert/ConcertApp/MultiUserClient/Source/MultiUserClient/Private/Replication/Submission/ChangeRequestBuilder.cpp
--- a/Engine/Plugins/Developer/Concert/ConcertApp/MultiUserClient/Source/MultiUserClient/Private/Replication/Submission/ChangeRequestBuilder.cpp
+++ b/Engine/Plugins/Developer/Concert/ConcertApp/MultiUserClient/Source/MultiUserClient/Private/Replication/Submission/ChangeRequestBuilder.cpp
@@ -14,13 +14,15 @@ namespace UE::MultiUserClient
 		const FGlobalAuthorityCache& InAuthorityCache,
 		IClientStreamSynchronizer& InStreamSynchronizer,
 		FStreamChangeTracker& InStreamChangeTracker,
-		FAuthorityChangeTracker& InAuthorityChangeTracker
+		FAuthorityChangeTracker& InAuthorityChangeTracker,
+		FFrequencyChangeTracker& InFrequencyChangeTracker
 		)
 		: LocalClientId(InLocalClientId)
 		, AuthorityCache(InAuthorityCache)
 		, StreamSynchronizer(InStreamSynchronizer)
 		, StreamChangeTracker(InStreamChangeTracker)
 		, AuthorityChangeTracker(InAuthorityChangeTracker)
+		, FrequencyChangeTracker(InFrequencyChangeTracker)
 	{}
 
 	TOptional<FConcertReplication_ChangeStream_Request> FChangeRequestBuilder::BuildStreamChange() const
@@ -32,9 +34,11 @@ namespace UE::MultiUserClient
 		const bool bIsChangelistEmpty = Changelist.ObjectsToPut.IsEmpty() && Changelist.ObjectsToRemove.IsEmpty();
 		if (!bIsChangelistEmpty)
 		{
+			const FFrequencyChangelist& FrequencyChanges = FrequencyChangeTracker.BuildForSubmission(Changelist);
+			
 			StreamRequest = StreamSynchronizer.GetServerState().ReplicatedObjects.IsEmpty()
-				? StreamRequestUtils::BuildChangeRequest_CreateNewStream(GetLocalClientStreamId(), Changelist)
-				: StreamRequestUtils::BuildChangeRequest_UpdateExistingStream(Changelist);
+				? StreamRequestUtils::BuildChangeRequest_CreateNewStream(GetLocalClientStreamId(), Changelist, FrequencyChanges)
+				: StreamRequestUtils::BuildChangeRequest_UpdateExistingStream(GetLocalClientStreamId(), Changelist, FrequencyChanges);
 				
 			// Predict conflicts in case the remote client's streams have changed. Also: while the UI highlights "bad" requests, it does not correct it.
 			AuthorityCache.CleanseConflictsFromStreamRequest(*StreamRequest, LocalClientId);
diff --git a/Engine/Plugins/Developer/Concert/ConcertApp/MultiUserClient/Source/MultiUserClient/Private/Replication/Submission/ChangeRequestBuilder.h b/Engine/Plugins/Developer/Concert/ConcertApp/MultiUserClient/Source/MultiUserClient/Private/Replication/Submission/ChangeRequestBuilder.h
--- a/Engine/Plugins/Developer/Concert/ConcertApp/MultiUserClient/Source/MultiUserClient/Private/Replication/Submission/ChangeRequestBuilder.h
+++ b/Engine/Plugins/Developer/Concert/ConcertApp/MultiUserClient/Source/MultiUserClient/Private/Replication/Submission/ChangeRequestBuilder.h
@@ -12,6 +12,7 @@ namespace UE::MultiUserClient
 {
 	class IClientStreamSynchronizer;
 	class FAuthorityChangeTracker;
+	class FFrequencyChangeTracker;
 	class FGlobalAuthorityCache;
 	class FStreamChangeTracker;
 	
@@ -22,13 +23,22 @@ namespace UE::MultiUserClient
 	class FChangeRequestBuilder
 	{
 	public:
-
+		
+		/**
+		 * @param InLocalClientId Endpoint ID of the client for which the request is being generated
+		 * @param InAuthorityCache Used to predict conflicts. The caller ensures it outlives the constructed instance.
+		 * @param InStreamSynchronizer Used to determine whether a new stream must be created. The caller ensures it outlives the constructed instance.
+		 * @param InStreamChangeTracker Builds object part of the stream request. The caller ensures it outlives the constructed instance.
+		 * @param InAuthorityChangeTracker Builds the authority changing request. The caller ensures it outlives the constructed instance.
+		 * @param InFrequencyChangeTracker Builds the object replication frequency part of the request. The caller ensures it outlives the constructed instance.
+		 */
 		FChangeRequestBuilder(
 			const FGuid& InLocalClientId,
-			const FGlobalAuthorityCache& InAuthorityCache,
-			IClientStreamSynchronizer& InStreamSynchronizer,
-			FStreamChangeTracker& InStreamChangeTracker,
-			FAuthorityChangeTracker& AuthorityChangeTracker
+			const FGlobalAuthorityCache& InAuthorityCache UE_LIFETIMEBOUND,
+			IClientStreamSynchronizer& InStreamSynchronizer UE_LIFETIMEBOUND,
+			FStreamChangeTracker& InStreamChangeTracker UE_LIFETIMEBOUND,
+			FAuthorityChangeTracker& InAuthorityChangeTracker UE_LIFETIMEBOUND,
+			FFrequencyChangeTracker& InFrequencyChangeTracker UE_LIFETIMEBOUND
 			);
 
 		/** @return Valid request that can be sent to the server, if there are any local changes. */
@@ -52,6 +62,8 @@ namespace UE::MultiUserClient
 		FStreamChangeTracker& StreamChangeTracker;
 		/** Informs us when authority is changed by the user. */
 		FAuthorityChangeTracker& AuthorityChangeTracker;
+		/** Used to build the object replication frequency requests. */
+		FFrequencyChangeTracker& FrequencyChangeTracker;
 
 		/** @return The ID MU uses for its single stream. */
 		FGuid GetLocalClientStreamId() const;
diff --git a/Engine/Plugins/Developer/Concert/ConcertApp/MultiUserClient/Source/MultiUserClient/Private/Replication/Submission/External/ClientChangeConversionUtils.cpp b/Engine/Plugins/Developer/Concert/ConcertApp/MultiUserClient/Source/MultiUserClient/Private/Replication/Submission/External/ClientChangeConversionUtils.cpp
--- a/Engine/Plugins/Developer/Concert/ConcertApp/MultiUserClient/Source/MultiUserClient/Private/Replication/Submission/External/ClientChangeConversionUtils.cpp
+++ b/Engine/Plugins/Developer/Concert/ConcertApp/MultiUserClient/Source/MultiUserClient/Private/Replication/Submission/External/ClientChangeConversionUtils.cpp
@@ -121,7 +121,7 @@ namespace UE::MultiUserClient::ClientChangeConversionUtils
 		const bool bNeedsToRegisterStream = ClientStreamContent.IsEmpty();
 		return bNeedsToRegisterStream
 			? StreamRequestUtils::BuildChangeRequest_CreateNewStream(ClientStreamId, StreamChangelist)
-			: StreamRequestUtils::BuildChangeRequest_UpdateExistingStream(MoveTemp(StreamChangelist));
+			: StreamRequestUtils::BuildChangeRequest_UpdateExistingStream(ClientStreamId, MoveTemp(StreamChangelist));
 	}
 
 	TOptional<FConcertReplication_ChangeAuthority_Request> Transform(
diff --git a/Engine/Plugins/Developer/Concert/ConcertApp/MultiUserClient/Source/MultiUserClient/Private/Replication/Submission/MultiEdit/ReassignObjectPropertiesLogic.cpp b/Engine/Plugins/Developer/Concert/ConcertApp/MultiUserClient/Source/MultiUserClient/Private/Replication/Submission/MultiEdit/ReassignObjectPropertiesLogic.cpp
--- a/Engine/Plugins/Developer/Concert/ConcertApp/MultiUserClient/Source/MultiUserClient/Private/Replication/Submission/MultiEdit/ReassignObjectPropertiesLogic.cpp
+++ b/Engine/Plugins/Developer/Concert/ConcertApp/MultiUserClient/Source/MultiUserClient/Private/Replication/Submission/MultiEdit/ReassignObjectPropertiesLogic.cpp
@@ -449,7 +449,7 @@ namespace UE::MultiUserClient
 			StreamId
 			);
 		FConcertReplication_ChangeStream_Request AssignedToClientRequest = bIsStreamRegistered
-			? StreamRequestUtils::BuildChangeRequest_UpdateExistingStream(Changelist)
+			? StreamRequestUtils::BuildChangeRequest_UpdateExistingStream(StreamId, Changelist)
 			: StreamRequestUtils::BuildChangeRequest_CreateNewStream(StreamId, Changelist);
 		
 		const TSharedPtr<ISubmissionOperation> SubmitOperation = Workflow.SubmitChanges({
diff --git a/Engine/Plugins/Developer/Concert/ConcertApp/MultiUserClient/Source/MultiUserClient/Private/Replication/Util/StreamRequestUtils.cpp b/Engine/Plugins/Developer/Concert/ConcertApp/MultiUserClient/Source/MultiUserClient/Private/Replication/Util/StreamRequestUtils.cpp
--- a/Engine/Plugins/Developer/Concert/ConcertApp/MultiUserClient/Source/MultiUserClient/Private/Replication/Util/StreamRequestUtils.cpp
+++ b/Engine/Plugins/Developer/Concert/ConcertApp/MultiUserClient/Source/MultiUserClient/Private/Replication/Util/StreamRequestUtils.cpp
@@ -7,9 +7,13 @@
 
 namespace UE::MultiUserClient::StreamRequestUtils
 {
-	FConcertReplication_ChangeStream_Request BuildChangeRequest_CreateNewStream(const FGuid& StreamId, const FStreamChangelist& FromChangelist)
+	FConcertReplication_ChangeStream_Request BuildChangeRequest_CreateNewStream(
+		const FGuid& StreamId,
+		const FStreamChangelist& ObjectChanges,
+		FFrequencyChangelist FrequencyChanges
+		)
 	{
-		const TMap<FObjectInStreamID, FConcertReplication_ChangeStream_PutObject>& ObjectsToPut = FromChangelist.ObjectsToPut;
+		const TMap<FObjectInStreamID, FConcertReplication_ChangeStream_PutObject>& ObjectsToPut = ObjectChanges.ObjectsToPut;
 		
 		FConcertReplication_ChangeStream_Request Request;
 		Request.StreamsToAdd.Emplace();
@@ -32,11 +36,36 @@ namespace UE::MultiUserClient::StreamRequestUtils
 			ReplicationMap.ReplicatedObjects.Add(PutObjectPair.Key.Object, *NewObjectInfo);
 		}
 		
+		FConcertStreamFrequencySettings& FrequencySettings = NewStream.BaseDescription.FrequencySettings;
+		if (FrequencyChanges.NewDefaults)
+		{
+			FrequencySettings.Defaults = MoveTemp(*FrequencyChanges.NewDefaults);
+		}
+		for (TPair<FSoftObjectPath, FConcertObjectReplicationSettings>& Overrides : FrequencyChanges.OverridesToAdd)
+		{
+			FrequencySettings.ObjectOverrides.Emplace(MoveTemp(Overrides.Key), MoveTemp(Overrides.Value));
+		}
+		
 		return Request;
 	}
 		
-	FConcertReplication_ChangeStream_Request BuildChangeRequest_UpdateExistingStream(FStreamChangelist FromChangelist)
+	FConcertReplication_ChangeStream_Request BuildChangeRequest_UpdateExistingStream(
+		const FGuid& StreamId,
+		FStreamChangelist ObjectChanges,
+		FFrequencyChangelist FrequencyChanges
+		)
 	{
-		return { MoveTemp(FromChangelist.ObjectsToRemove), MoveTemp(FromChangelist.ObjectsToPut) };
+		FConcertReplication_ChangeStream_Request Result { MoveTemp(ObjectChanges.ObjectsToRemove), MoveTemp(ObjectChanges.ObjectsToPut) };
+		
+		FConcertReplication_ChangeStream_Frequency& FrequencyChangeRequest = Result.FrequencyChanges.Add(StreamId);
+		if (FrequencyChanges.NewDefaults)
+		{
+			FrequencyChangeRequest.NewDefaults = MoveTemp(*FrequencyChanges.NewDefaults); 
+			FrequencyChangeRequest.Flags = EConcertReplicationChangeFrequencyFlags::SetDefaults; 
+		}
+		FrequencyChangeRequest.OverridesToAdd = MoveTemp(FrequencyChanges.OverridesToAdd);
+		FrequencyChangeRequest.OverridesToRemove = MoveTemp(FrequencyChanges.OverridesToRemove);
+
+		return Result;
 	}
 }
diff --git a/Engine/Plugins/Developer/Concert/ConcertApp/MultiUserClient/Source/MultiUserClient/Private/Replication/Util/StreamRequestUtils.h b/Engine/Plugins/Developer/Concert/ConcertApp/MultiUserClient/Source/MultiUserClient/Private/Replication/Util/StreamRequestUtils.h
--- a/Engine/Plugins/Developer/Concert/ConcertApp/MultiUserClient/Source/MultiUserClient/Private/Replication/Util/StreamRequestUtils.h
+++ b/Engine/Plugins/Developer/Concert/ConcertApp/MultiUserClient/Source/MultiUserClient/Private/Replication/Util/StreamRequestUtils.h
@@ -8,12 +8,20 @@ struct FConcertReplication_ChangeStream_Request;
 
 namespace UE::MultiUserClient
 {
-	/** Describes changes that MU client makes. */
+	/** Describes changes that MU clients make to the objects in a stream. */
 	struct FStreamChangelist
 	{
 		TSet<FObjectInStreamID> ObjectsToRemove;
 		TMap<FObjectInStreamID, FConcertReplication_ChangeStream_PutObject> ObjectsToPut;
 	};
+	
+	/** Describes changes that MU clients make to the frequency settings in a stream. */
+	struct FFrequencyChangelist
+	{
+		TSet<FSoftObjectPath> OverridesToRemove;
+		TMap<FSoftObjectPath, FConcertObjectReplicationSettings> OverridesToAdd;
+		TOptional<FConcertObjectReplicationSettings> NewDefaults;
+	};
 }
 
 namespace UE::MultiUserClient::StreamRequestUtils
@@ -21,13 +29,24 @@ namespace UE::MultiUserClient::StreamRequestUtils
 	/**
 	 * Builds a request for creating a new stream.
 	 * @param StreamId The stream that should be modified
-	 * @param FromChangelist The changes to be made
+	 * @param ObjectChanges The object changes to be made
+	 * @param FrequencyChanges The frequency changes to be made
 	 */
-	FConcertReplication_ChangeStream_Request BuildChangeRequest_CreateNewStream(const FGuid& StreamId, const FStreamChangelist& FromChangelist);
+	FConcertReplication_ChangeStream_Request BuildChangeRequest_CreateNewStream(
+		const FGuid& StreamId,
+		const FStreamChangelist& ObjectChanges,
+		FFrequencyChangelist FrequencyChanges = FFrequencyChangelist()
+		);
 	
 	/**
 	 * Builds a request for updating a preexisting stream.
-	 * @param FromChangelist The changes to be made
+	 * @param StreamId The stream that should be modified
+	 * @param ObjectChanges The object changes to be made
+	 * @param FrequencyChanges The frequency changes to be made
 	 */
-	FConcertReplication_ChangeStream_Request BuildChangeRequest_UpdateExistingStream(FStreamChangelist FromChangelist);
+	FConcertReplication_ChangeStream_Request BuildChangeRequest_UpdateExistingStream(
+		const FGuid& StreamId,
+		FStreamChangelist ObjectChanges,
+		FFrequencyChangelist FrequencyChanges = FFrequencyChangelist()
+		);
 }
diff --git a/Engine/Plugins/Developer/Concert/ConcertApp/MultiUserClient/Source/MultiUserClient/Private/Settings/FrequencyDefaultRules.cpp b/Engine/Plugins/Developer/Concert/ConcertApp/MultiUserClient/Source/MultiUserClient/Private/Settings/FrequencyDefaultRules.cpp
new file mode 100644
--- /dev/null
+++ b/Engine/Plugins/Developer/Concert/ConcertApp/MultiUserClient/Source/MultiUserClient/Private/Settings/FrequencyDefaultRules.cpp
@@ -0,0 +1,42 @@
+// Copyright Epic Games, Inc. All Rights Reserved.
+
+#include "FrequencyDefaultRules.h"
+
+#include "Algo/AnyOf.h"
+#include "GameFramework/Actor.h"
+
+bool FMultiUserActorFrequencyStackEntry::Matches(UObject& Object) const
+{
+	const FSoftClassPath ObjectClass = Object.GetClass();
+	const bool bMatchesAnyClass = Algo::AnyOf(MatchingClasses, [&ObjectClass](const FSoftClassPath& Class){ return ObjectClass == Class; });
+	if (bMatchesAnyClass || Object.IsA<AActor>())
+	{
+		return bMatchesAnyClass;
+	}
+
+	AActor* ObjectOwner = Object.GetTypedOuter<AActor>();
+	if (!ObjectOwner)
+	{
+		return false;
+	}
+
+	bool bFoundObject = false;
+	SubobjectMatchingRules.MatchSubobjectsRecursivelyBreakable(*ObjectOwner, [&Object, &bFoundObject](UObject& MatchedObject)
+	{
+		bFoundObject |= &MatchedObject == &Object;
+		return bFoundObject ? EBreakBehavior::Break : EBreakBehavior::Continue;
+	});
+	return bFoundObject;
+}
+
+TOptional<FConcertObjectReplicationSettings> FMultiUserFrequencyDefaultRules::FindOverrideSettings(UObject& Object) const
+{
+	for (const FMultiUserActorFrequencyStackEntry& Entry : FrequencyRuleStack)
+	{
+		if (Entry.Matches(Object))
+		{
+			return Entry.ReplicationSettings;
+		}
+	}
+	return {};
+}
diff --git a/Engine/Plugins/Developer/Concert/ConcertApp/MultiUserClient/Source/MultiUserClient/Private/Settings/FrequencyDefaultRules.h b/Engine/Plugins/Developer/Concert/ConcertApp/MultiUserClient/Source/MultiUserClient/Private/Settings/FrequencyDefaultRules.h
new file mode 100644
--- /dev/null
+++ b/Engine/Plugins/Developer/Concert/ConcertApp/MultiUserClient/Source/MultiUserClient/Private/Settings/FrequencyDefaultRules.h
@@ -0,0 +1,90 @@
+// Copyright Epic Games, Inc. All Rights Reserved.
+
+#pragma once
+
+#include "Replication/Data/ReplicationFrequencySettings.h"
+#include "Replication/Settings/ConcertStreamObjectAutoBindingRules.h"
+#include "Misc/Optional.h"
+#include "FrequencyDefaultRules.generated.h"
+
+class UObject;
+
+/**
+ * Used to build a stack of rules. The first stack entry that matches an object, is used to determine the object's frequency settings.
+ */
+USTRUCT()
+struct FMultiUserActorFrequencyStackEntry
+{
+	GENERATED_BODY()
+
+	/** The settings to apply to objects matched via this stack entry. */
+	UPROPERTY(Category = "Replication Settings|Frequency", EditAnywhere)
+	FConcertObjectReplicationSettings ReplicationSettings;
+
+	/**
+	 * If an object's class is contained in this set, ReplicationSettings are applied to it.
+	 * @see bAllowRecursiveClassSearch.
+	 */
+	UPROPERTY(Category = "Replication Settings|Frequency", EditAnywhere)
+	TSet<FSoftClassPath> MatchingClasses;
+
+	/**
+	 * These rules are applied only to objects that are not actors, e.g. components.
+	 * The matching rules are applied to the object's owning actor. If the object is found by these rules, ReplicationSettings are applied to it.
+	 */
+	UPROPERTY(Category = "Replication Settings|Frequency", EditAnywhere)
+	FConcertPerClassSubobjectMatchingRules SubobjectMatchingRules;
+
+	/**
+	 * Determines whether parent classes are included when matching against MatchingClasses.
+	 * 
+	 * Example: MatchingClasses = UStaticMeshComponent.
+	 * Object 1 is UStaticMeshComponent > matches.
+	 * Object 2 is UInstancedStaticMeshComponent > if bAllowRecursiveClassSearch == true, then UInstancedStaticMeshComponent is matched. Otherwise, not. 
+	 */
+	UPROPERTY(Category = "Replication Settings|Frequency", EditAnywhere, AdvancedDisplay)
+	bool bAllowRecursiveClassSearch = true;
+
+	/** @return Whether Object matches the these rules. */
+	bool Matches(UObject& Object) const;
+};
+
+/**
+ * Defines rules determining an object's frequency settings.
+ */
+USTRUCT()
+struct FMultiUserFrequencyDefaultRules
+{
+	GENERATED_BODY()
+
+	/**
+	 * Default frequency settings that are used by objects when they do not match any specific rules in FrequencyRuleStack.
+	 * These are the settings that are set for the Multi User stream's FSharedReplicationStreamDescription::FrequencySettings::Defaults member.
+	 */
+	UPROPERTY(Category = "Replication Settings|Frequency", EditAnywhere)
+	FConcertObjectReplicationSettings DefaultObjectFrequencySettings { EConcertObjectReplicationMode::SpecifiedRate, 30 };
+
+	/**
+	 * Use this to assign specific frequency settings to objects.
+	 * 
+	 * When an object is added to the replication stream, its frequency settings are looked up as follows:
+	 *	1. If the object is an actor,
+	 *		1.1 go through FrequencyRuleStack from top to bottom
+	 *		1.2 if the entry's MatchingClasses contains the actor's class, use the ReplicationSettings in that entry.
+	 *	2. If the object is not an actor,
+	*		2.1 go through FrequencyRuleStack from top to bottom
+	 *		2.2 if the entry's MatchingClasses contains the object's class, use the ReplicationSettings in that entry. 
+	 *		2.3 get the object's owning actor
+	 *		2.4 if the entry's SubobjectMatchingRules finds the object by searching the owning actor, use the ReplicationSettings in that entry.
+	 */
+	UPROPERTY(Category = "Replication Settings|Frequency", EditAnywhere)
+	TArray<FMultiUserActorFrequencyStackEntry> FrequencyRuleStack;
+
+	/**
+	 * Goes through FrequencyRuleStack and returns the FConcertObjectReplicationSettings of the entry that was matched.
+	 * 
+	 * If none was matched, the object should not have any frequency override and should default to the stream's
+	 * FSharedReplicationStreamDescription::FrequencySettings::Defaults, which should be set to DefaultObjectFrequencySettings. 
+	 */
+	TOptional<FConcertObjectReplicationSettings> FindOverrideSettings(UObject& Object) const;
+};
diff --git a/Engine/Plugins/Developer/Concert/ConcertApp/MultiUserClient/Source/MultiUserClient/Private/Settings/MultiUserReplicationSettings.h b/Engine/Plugins/Developer/Concert/ConcertApp/MultiUserClient/Source/MultiUserClient/Private/Settings/MultiUserReplicationSettings.h
new file mode 100644
--- /dev/null
+++ b/Engine/Plugins/Developer/Concert/ConcertApp/MultiUserClient/Source/MultiUserClient/Private/Settings/MultiUserReplicationSettings.h
@@ -0,0 +1,37 @@
+// Copyright Epic Games, Inc. All Rights Reserved.
+
+#pragma once
+
+#include "FrequencyDefaultRules.h"
+#include "Replication/Settings/ConcertStreamObjectAutoBindingRules.h"
+
+#include "Misc/Optional.h"
+#include "UObject/Object.h"
+
+#include "MultiUserReplicationSettings.generated.h"
+
+UCLASS(config=MultiUserClient)
+class UMultiUserReplicationSettings : public UObject
+{
+	GENERATED_BODY()
+public:
+
+	static UMultiUserReplicationSettings* Get() { return GetMutableDefault<UMultiUserReplicationSettings>(); }
+
+	/**
+	 * When you add an object via the stream editor, you may want to automatically bind properties and add additional subobjects.
+	 * Here you can specify the rules to achieve this.
+	 */
+	UPROPERTY(Config, Category = "Replication Settings", EditAnywhere)
+	FConcertStreamObjectAutoBindingRules ReplicationEditorSettings;
+
+	/** Determines the frequency settings applied to an object set by default when it is added to the stream. */
+	UPROPERTY(Config, Category = "Replication Settings|Frequency", EditAnywhere)
+	FMultiUserFrequencyDefaultRules FrequencyRules;
+
+	/**
+	 * Checks Object against these settings and returns override settings to set for this object.
+	 * @param Object An object that the user just added to the stream
+	 */
+	TOptional<FConcertObjectReplicationSettings> DetermineObjectFrequencySettings(UObject& Object) const { return FrequencyRules.FindOverrideSettings(Object); }
+};
\ No newline at end of file
diff --git a/Engine/Plugins/Developer/Concert/ConcertApp/MultiUserClient/Source/MultiUserClient/Private/Widgets/ActiveSession/Replication/Client/Multi/SMultiClientView.cpp b/Engine/Plugins/Developer/Concert/ConcertApp/MultiUserClient/Source/MultiUserClient/Private/Widgets/ActiveSession/Replication/Client/Multi/SMultiClientView.cpp
--- a/Engine/Plugins/Developer/Concert/ConcertApp/MultiUserClient/Source/MultiUserClient/Private/Widgets/ActiveSession/Replication/Client/Multi/SMultiClientView.cpp
+++ b/Engine/Plugins/Developer/Concert/ConcertApp/MultiUserClient/Source/MultiUserClient/Private/Widgets/ActiveSession/Replication/Client/Multi/SMultiClientView.cpp
@@ -4,7 +4,6 @@
 
 #include "IClientSelectionModel.h"
 #include "MultiStreamModel.h"
-#include "MultiUserReplicationSettings.h"
 #include "Replication/Client/ReplicationClient.h"
 #include "Replication/Client/ReplicationClientManager.h"
 #include "Replication/Editor/Model/Object/ActorSelectionSourceModel.h"
diff --git a/Engine/Plugins/Developer/Concert/ConcertApp/MultiUserClient/Source/MultiUserClient/Private/Widgets/ActiveSession/Replication/Client/Single/SReplicationClientView.cpp b/Engine/Plugins/Developer/Concert/ConcertApp/MultiUserClient/Source/MultiUserClient/Private/Widgets/ActiveSession/Replication/Client/Single/SReplicationClientView.cpp
--- a/Engine/Plugins/Developer/Concert/ConcertApp/MultiUserClient/Source/MultiUserClient/Private/Widgets/ActiveSession/Replication/Client/Single/SReplicationClientView.cpp
+++ b/Engine/Plugins/Developer/Concert/ConcertApp/MultiUserClient/Source/MultiUserClient/Private/Widgets/ActiveSession/Replication/Client/Single/SReplicationClientView.cpp
@@ -2,7 +2,6 @@
 
 #include "SReplicationClientView.h"
 
-#include "MultiUserReplicationSettings.h"
 #include "Replication/Client/ReplicationClient.h"
 #include "Replication/Editor/Model/IEditableReplicationStreamModel.h"
 #include "Replication/Editor/Model/Object/ActorSelectionSourceModel.h"
diff --git a/Engine/Plugins/Developer/Concert/ConcertUI/ConcertClientSharedSlate/Source/ConcertClientSharedSlate/Private/Replication/Editor/Model/StreamExtenderBySettings.cpp b/Engine/Plugins/Developer/Concert/ConcertUI/ConcertClientSharedSlate/Source/ConcertClientSharedSlate/Private/Replication/Editor/Model/StreamExtenderBySettings.cpp
--- a/Engine/Plugins/Developer/Concert/ConcertUI/ConcertClientSharedSlate/Source/ConcertClientSharedSlate/Private/Replication/Editor/Model/StreamExtenderBySettings.cpp
+++ b/Engine/Plugins/Developer/Concert/ConcertUI/ConcertClientSharedSlate/Source/ConcertClientSharedSlate/Private/Replication/Editor/Model/StreamExtenderBySettings.cpp
@@ -23,7 +23,7 @@ namespace UE::ConcertClientSharedSlate
 		{
 			Context.AddPropertyTo(ExtendedObject, MoveTemp(PropertyChain));
 		});
-		Settings->DefaultAddedSubobjectRules.MatchSubobjectsRecursivelyFor(ExtendedObject, [&ExtendedObject, &Context](UObject& AdditionalObject)
+		Settings->DefaultAddedSubobjectRules.MatchSubobjectsRecursively(ExtendedObject, [&ExtendedObject, &Context](UObject& AdditionalObject)
 		{
 			Context.AddAdditionalObject(AdditionalObject);
 		});
diff --git a/Engine/Plugins/Developer/Concert/ConcertUI/ConcertClientSharedSlate/Source/ConcertClientSharedSlate/Private/Replication/Settings/ConcertPerClassSubobjectMatchingRules.cpp b/Engine/Plugins/Developer/Concert/ConcertUI/ConcertClientSharedSlate/Source/ConcertClientSharedSlate/Private/Replication/Settings/ConcertPerClassSubobjectMatchingRules.cpp
--- a/Engine/Plugins/Developer/Concert/ConcertUI/ConcertClientSharedSlate/Source/ConcertClientSharedSlate/Private/Replication/Settings/ConcertPerClassSubobjectMatchingRules.cpp
+++ b/Engine/Plugins/Developer/Concert/ConcertUI/ConcertClientSharedSlate/Source/ConcertClientSharedSlate/Private/Replication/Settings/ConcertPerClassSubobjectMatchingRules.cpp
@@ -6,7 +6,7 @@
 
 namespace UE::ConcertClientSharedSlate::DefaultSubobjects
 {
-	static void InternalAddAdditionalObjectsFromSettings(UClass& StartClass, const FConcertPerClassSubobjectMatchingRules& Settings, const UObject& AddedObject, TFunctionRef<void(UObject&)> OnSubobjectMatched)
+	static void InternalAddAdditionalObjectsFromSettings(UClass& StartClass, const FConcertPerClassSubobjectMatchingRules& Settings, const UObject& AddedObject, TFunctionRef<EBreakBehavior(UObject&)> OnSubobjectMatched)
 	{
 		// Find the most specialized class properties
 		UClass* Current = &StartClass;
@@ -18,8 +18,18 @@ namespace UE::ConcertClientSharedSlate::DefaultSubobjects
 			{
 				continue;
 			}
+
+			EBreakBehavior BreakResult = EBreakBehavior::Continue;
+			DefaultSubobjects->MatchToSubobjectsBreakable(AddedObject, [&OnSubobjectMatched, &BreakResult](UObject& Object)
+			{
+				BreakResult = OnSubobjectMatched(Object);
+				return BreakResult;
+			});
+			if (BreakResult == EBreakBehavior::Break)
+			{
+				return;
+			}
 			
-			DefaultSubobjects->MatchToSubobjectsIn(AddedObject, OnSubobjectMatched);
 			// Recurse super structs
 			if (UClass* Parent = Current->GetSuperClass()
 				; Parent && DefaultSubobjects->bInheritFromBase)
@@ -30,7 +40,7 @@ namespace UE::ConcertClientSharedSlate::DefaultSubobjects
 	}
 }
 
-void FConcertPerClassSubobjectMatchingRules::MatchSubobjectsRecursivelyFor(const UObject& Object, TFunctionRef<void(UObject&)> OnSubobjectMatched) const
+void FConcertPerClassSubobjectMatchingRules::MatchSubobjectsRecursivelyBreakable(const UObject& Object, TFunctionRef<EBreakBehavior(UObject&)> OnSubobjectMatched) const
 {
 	checkf(!Object.IsA<UClass>(), TEXT("Pass in the UObject instanced directly, not its class!"));
 	UE::ConcertClientSharedSlate::DefaultSubobjects::InternalAddAdditionalObjectsFromSettings(*Object.GetClass(), *this, Object, OnSubobjectMatched);
diff --git a/Engine/Plugins/Developer/Concert/ConcertUI/ConcertClientSharedSlate/Source/ConcertClientSharedSlate/Private/Replication/Settings/ConcertSubobjectMatchingRules.cpp b/Engine/Plugins/Developer/Concert/ConcertUI/ConcertClientSharedSlate/Source/ConcertClientSharedSlate/Private/Replication/Settings/ConcertSubobjectMatchingRules.cpp
--- a/Engine/Plugins/Developer/Concert/ConcertUI/ConcertClientSharedSlate/Source/ConcertClientSharedSlate/Private/Replication/Settings/ConcertSubobjectMatchingRules.cpp
+++ b/Engine/Plugins/Developer/Concert/ConcertUI/ConcertClientSharedSlate/Source/ConcertClientSharedSlate/Private/Replication/Settings/ConcertSubobjectMatchingRules.cpp
@@ -4,6 +4,7 @@
 
 #include "Components/ActorComponent.h"
 #include "Internationalization/Regex.h"
+#include "Misc/EBreakBehavior.h"
 
 namespace UE::ConcertClientSharedSlate
 {
@@ -23,17 +24,17 @@ namespace UE::ConcertClientSharedSlate
 	}
 }
 
-void FConcertSubobjectMatchingRules::MatchToSubobjectsIn(const UObject& AddedObject, TFunctionRef<void(UObject&)> OnSubobjectMatched) const
+void FConcertSubobjectMatchingRules::MatchToSubobjectsBreakable(const UObject& AddedObject, TFunctionRef<EBreakBehavior(UObject&)> OnSubobjectMatched) const
 {
 	// Do not search recursively so FurtherObjectsCallback can decide to call AddAdditionalObjectsFromSettings again on the newly added objects.
 	constexpr bool bIncludeNested = false;
-	ForEachObjectWithOuter(&AddedObject, [this, &OnSubobjectMatched](UObject* Subobject)
+	ForEachObjectWithOuterBreakable(&AddedObject, [this, &OnSubobjectMatched](UObject* Subobject)
 	{
 		// Has exclusion regex?
 		const bool bShouldExclude = UE::ConcertClientSharedSlate::MatchesAnyRegex(Subobject->GetName(), ExcludeSubobjectRegex);
 		if (bShouldExclude)
 		{
-			return;
+			return true;
 		}
 
 		// Was told to add all UActorComponents?
@@ -42,8 +43,7 @@ void FConcertSubobjectMatchingRules::MatchToSubobjectsIn(const UObject& AddedObj
 			&& IncludeAllOption == EConcertIncludeAllSubobjectsType::AllComponents;
 		if (bIncludeAllSubobjects || bIncludeDueToComponent)
 		{
-			OnSubobjectMatched(*Subobject);
-			return;
+			return OnSubobjectMatched(*Subobject) == EBreakBehavior::Break;
 		}
 
 		// Has configured class?
@@ -52,8 +52,7 @@ void FConcertSubobjectMatchingRules::MatchToSubobjectsIn(const UObject& AddedObj
 		{
 			if (IncludeClasses.Contains(Subobject->GetClass()))
 			{
-				OnSubobjectMatched(*Subobject);
-				return;
+				return OnSubobjectMatched(*Subobject) == EBreakBehavior::Break;
 			}
 		}
 
@@ -61,7 +60,9 @@ void FConcertSubobjectMatchingRules::MatchToSubobjectsIn(const UObject& AddedObj
 		const bool bShouldIncludeByRegex = UE::ConcertClientSharedSlate::MatchesAnyRegex(Subobject->GetName(), IncludeSubobjectRegex);
 		if (bShouldIncludeByRegex)
 		{
-			OnSubobjectMatched(*Subobject);
+			return OnSubobjectMatched(*Subobject) == EBreakBehavior::Break;
 		}
+		
+		return true;
 	}, bIncludeNested);
 }
\ No newline at end of file
diff --git a/Engine/Plugins/Developer/Concert/ConcertUI/ConcertClientSharedSlate/Source/ConcertClientSharedSlate/Public/Replication/Settings/ConcertPerClassSubobjectMatchingRules.h b/Engine/Plugins/Developer/Concert/ConcertUI/ConcertClientSharedSlate/Source/ConcertClientSharedSlate/Public/Replication/Settings/ConcertPerClassSubobjectMatchingRules.h
--- a/Engine/Plugins/Developer/Concert/ConcertUI/ConcertClientSharedSlate/Source/ConcertClientSharedSlate/Public/Replication/Settings/ConcertPerClassSubobjectMatchingRules.h
+++ b/Engine/Plugins/Developer/Concert/ConcertUI/ConcertClientSharedSlate/Source/ConcertClientSharedSlate/Public/Replication/Settings/ConcertPerClassSubobjectMatchingRules.h
@@ -34,6 +34,14 @@ struct CONCERTCLIENTSHAREDSLATE_API FConcertPerClassSubobjectMatchingRules
 	UPROPERTY(EditAnywhere, Config, Category = "Replication|Editor")
 	TMap<FSoftClassPath, FConcertInheritableSubobjectMatchingRules> SubobjectMatchingRules;
 
-	/** Recursively walks up Object class hierarchy and tries to find subobjects  */
-	void MatchSubobjectsRecursivelyFor(const UObject& Object, TFunctionRef<void(UObject&)> OnSubobjectMatched) const;
+	/** Recursively walks up Object class hierarchy and tries to find subobjects */
+	void MatchSubobjectsRecursivelyBreakable(const UObject& Object, TFunctionRef<EBreakBehavior(UObject&)> OnSubobjectMatched) const;
+	void MatchSubobjectsRecursively(const UObject& Object, TFunctionRef<void(UObject&)> OnSubobjectMatched) const
+	{
+		return MatchSubobjectsRecursivelyBreakable(Object, [&OnSubobjectMatched](UObject& Object)
+		{
+			OnSubobjectMatched(Object);
+			return EBreakBehavior::Continue;
+		});
+	}
 };
\ No newline at end of file
diff --git a/Engine/Plugins/Developer/Concert/ConcertUI/ConcertClientSharedSlate/Source/ConcertClientSharedSlate/Public/Replication/Settings/ConcertSubobjectMatchingRules.h b/Engine/Plugins/Developer/Concert/ConcertUI/ConcertClientSharedSlate/Source/ConcertClientSharedSlate/Public/Replication/Settings/ConcertSubobjectMatchingRules.h
--- a/Engine/Plugins/Developer/Concert/ConcertUI/ConcertClientSharedSlate/Source/ConcertClientSharedSlate/Public/Replication/Settings/ConcertSubobjectMatchingRules.h
+++ b/Engine/Plugins/Developer/Concert/ConcertUI/ConcertClientSharedSlate/Source/ConcertClientSharedSlate/Public/Replication/Settings/ConcertSubobjectMatchingRules.h
@@ -2,6 +2,7 @@
 
 #pragma once
 
+#include "Misc/EBreakBehavior.h"
 #include "Templates/Function.h"
 #include "UObject/SoftObjectPath.h"
 #include "ConcertSubobjectMatchingRules.generated.h"
@@ -58,5 +59,13 @@ struct FConcertSubobjectMatchingRules
 	 * Looks for subobjects in AddedObject that match these rules.
 	 * This looks for direct subobjects only though you can call MatchToSubobjectsIn recursively on what you received through OnSubobjectMatched.
 	 */
-	void MatchToSubobjectsIn(const UObject& AddedObject, TFunctionRef<void(UObject&)> OnSubobjectMatched) const;
+	void MatchToSubobjectsBreakable(const UObject& AddedObject, TFunctionRef<EBreakBehavior(UObject&)> OnSubobjectMatched) const;
+	void MatchToSubobjects(const UObject& AddedObject, TFunctionRef<void(UObject&)> OnSubobjectMatched) const
+	{
+		return MatchToSubobjectsBreakable(AddedObject, [&OnSubobjectMatched](UObject& Object)
+		{
+			OnSubobjectMatched(Object);
+			return EBreakBehavior::Continue;
+		});
+	}
 };
