Reassigning objects to clients now transfers their frequency as well

diff --git a/Engine/Plugins/Developer/Concert/ConcertApp/MultiUserClient/Source/MultiUserClient/Private/Replication/Submission/MultiEdit/ReassignObjectPropertiesLogic.cpp b/Engine/Plugins/Developer/Concert/ConcertApp/MultiUserClient/Source/MultiUserClient/Private/Replication/Submission/MultiEdit/ReassignObjectPropertiesLogic.cpp
--- a/Engine/Plugins/Developer/Concert/ConcertApp/MultiUserClient/Source/MultiUserClient/Private/Replication/Submission/MultiEdit/ReassignObjectPropertiesLogic.cpp
+++ b/Engine/Plugins/Developer/Concert/ConcertApp/MultiUserClient/Source/MultiUserClient/Private/Replication/Submission/MultiEdit/ReassignObjectPropertiesLogic.cpp
@@ -17,14 +17,18 @@ namespace UE::MultiUserClient
 {
 	namespace ReassignObjectProperties::Private
 	{
+		using FClientId = FGuid;
+		
 		struct FObjectReassignment
 		{
 			/** Used to build the final changelist to apply to the target client. Maps ClientIds to the object data they had registered. */
-			TMap<FGuid, FObjectReplicationMap> OldRegisteredObjects;
+			TMap<FClientId, FObjectReplicationMap> OldRegisteredObjects;
+			/** The frequency settings clients when reassignment operation was started. */
+			TMap<FClientId, FConcertStreamFrequencySettings> OldFrequencies;
 			/** Changes to make to all the other clients */
-			TMap<FGuid, FConcertReplication_ChangeStream_Request> ReassignedClientRequests;
+			TMap<FClientId, FConcertReplication_ChangeStream_Request> ReassignedClientRequests;
 			/** Authority that the clients had before. This is used to determine whether to give the assigned to client authority once the stream changes have been made. */
-			TMap<FGuid, TArray<FSoftObjectPath>> ReassignedAuthority;
+			TMap<FClientId, TArray<FSoftObjectPath>> ReassignedAuthority;
 		};
 		
 		/** Builds all the changes that must be made to the clients we are transferring from.  */
@@ -32,6 +36,7 @@ namespace UE::MultiUserClient
 		{
 			FObjectReassignment Result;
 			TMap<FGuid, FObjectReplicationMap>& OldRegisteredObjects = Result.OldRegisteredObjects;
+			TMap<FGuid, FConcertStreamFrequencySettings>& OldFrequencies = Result.OldFrequencies;
 			TMap<FGuid, FConcertReplication_ChangeStream_Request>& ReassignedClientRequests = Result.ReassignedClientRequests;
 			TMap<FGuid, TArray<FSoftObjectPath>>& ReassignedAuthority = Result.ReassignedAuthority;
 		
@@ -40,47 +45,88 @@ namespace UE::MultiUserClient
 				const FObjectInStreamID ObjectId{ StreamId, ObjectPath };
 
 				const FGlobalAuthorityCache& AuthorityCache = ClientManager.GetAuthorityCache();
-				AuthorityCache.ForEachClientWithObjectInStream(ObjectPath, [&ClientManager, &OldRegisteredObjects, &ReassignedClientRequests, &ReassignedAuthority, &ObjectPath, &ObjectId, &AuthorityCache](const FGuid& ClientId)
-				{
-					// Get the property assignments of the client being reassigned from
-					const FReplicationClient* ClientToReassignFrom = ClientManager.FindClient(ClientId);
-					if (!ensureMsgf(ClientToReassignFrom, TEXT("Authority cache out of sync"))
-						|| !ClientToReassignFrom->AllowsEditing())
-					{
-						return EBreakBehavior::Continue;
-					}
-					const FReplicatedObjectInfo* ObjectInfo = ClientToReassignFrom->GetStreamSynchronizer().GetServerState().ReplicatedObjects.Find(ObjectId.Object);
-					if (!ensureMsgf(ObjectInfo, TEXT("ForEachClientWithObjectInStream lied")))
+				AuthorityCache.ForEachClientWithObjectInStream(ObjectPath,
+					[&ClientManager, &OldRegisteredObjects, &OldFrequencies, &ReassignedClientRequests, &ReassignedAuthority, &ObjectPath, &ObjectId, &AuthorityCache](const FGuid& ClientId)
 					{
+						// Get the property assignments of the client being reassigned from
+						const FReplicationClient* ClientToReassignFrom = ClientManager.FindClient(ClientId);
+						if (!ensureMsgf(ClientToReassignFrom, TEXT("Authority cache out of sync"))
+							|| !ClientToReassignFrom->AllowsEditing())
+						{
+							return EBreakBehavior::Continue;
+						}
+						
+						const IClientStreamSynchronizer& ReassignedStreamSynchronizer = ClientToReassignFrom->GetStreamSynchronizer();
+						const FReplicatedObjectInfo* ObjectInfo = ReassignedStreamSynchronizer.GetServerState().ReplicatedObjects.Find(ObjectId.Object);
+						if (!ensureMsgf(ObjectInfo, TEXT("ForEachClientWithObjectInStream lied")))
+						{
+							return EBreakBehavior::Continue;
+						}
+
+						// Transfer from client being reassigned from to the client being assigned to
+						OldRegisteredObjects.FindOrAdd(ClientId).ReplicatedObjects.Add(ObjectPath, *ObjectInfo);
+						// Remove the object completely from the reassigned client.
+						ReassignedClientRequests.FindOrAdd(ClientId).ObjectsToRemove.Add(ObjectId);
+						// Later we can determine whether to give the assigned to client authority
+						if (AuthorityCache.HasAuthorityOverObject(ObjectPath, ClientId))
+						{
+							ReassignedAuthority.FindOrAdd(ClientId).Add(ObjectPath);
+						}
+
+						// Keep track of this object's frequency settings so it can be transferred
+						const FConcertObjectReplicationSettings* OverrideFrequencySettings = ReassignedStreamSynchronizer.GetFrequencySettings().ObjectOverrides.Find(ObjectId.Object);
+						if (OverrideFrequencySettings)
+						{
+							OldFrequencies.FindOrAdd(ClientId).ObjectOverrides.Add(ObjectId.Object, *OverrideFrequencySettings);
+						}
+					
 						return EBreakBehavior::Continue;
-					}
+					});
+			}
 
-					// Transfer from client being reassigned from to the client being assigned to
-					OldRegisteredObjects.FindOrAdd(ClientId).ReplicatedObjects.Add(ObjectPath, *ObjectInfo);
-					// Remove the object completely from the reassigned client.
-					ReassignedClientRequests.FindOrAdd(ClientId).ObjectsToRemove.Add(ObjectId);
-					// Later we can determine whether to give the assigned to client authority
-					if (AuthorityCache.HasAuthorityOverObject(ObjectPath, ClientId))
-					{
-						ReassignedAuthority.FindOrAdd(ClientId).Add(ObjectPath);
-					}
-				
-					return EBreakBehavior::Continue;
-				});
+			// Need to set Defaults of the OldFrequencies because that's not been done, yet.
+			for (TPair<FGuid, FConcertStreamFrequencySettings>& OldFrequency : OldFrequencies)
+			{
+				const FReplicationClient* ClientToReassignFrom = ClientManager.FindClient(OldFrequency.Key);
+				check(ClientToReassignFrom);
+				OldFrequency.Value.Defaults = ClientToReassignFrom->GetStreamSynchronizer().GetFrequencySettings().Defaults;
 			}
 
 			return Result;
 		}
 
+		static FConcertObjectReplicationSettings GetHighestFrequency(const FSoftObjectPath& Object, const TMap<FClientId, FConcertStreamFrequencySettings>& OldFrequencies)
+		{
+			TOptional<FConcertObjectReplicationSettings> Settings;
+			for (const TPair<FClientId, FConcertStreamFrequencySettings>& Pair : OldFrequencies)
+			{
+				const FConcertObjectReplicationSettings& ClientSettings = Pair.Value.GetSettingsFor(Object);
+				if (!Settings)
+				{
+					Settings = ClientSettings;
+				}
+				else if (*Settings <= ClientSettings)
+				{
+					Settings = ClientSettings;
+				}
+			}
+			return Settings.Get(FConcertObjectReplicationSettings{});
+		}
+
 		/** Builds a changelist based on which remote clients we managed to change successfully. */
-		static FStreamChangelist MakeChangelistFromAppliedChanges(
+		static TPair<FStreamChangelist, FFrequencyChangelist> MakeChangelistFromAppliedChanges(
 			const TMap<FGuid, FObjectReplicationMap>& OldRegisteredObjects,
+			const TMap<FClientId, FConcertStreamFrequencySettings>& OldFrequencies,
 			const FParallelExecutionResult& ParallelExecutionResult,
-			const FObjectReplicationMap& CurrentTargetClientState,
-			const FGuid& TargetStreamId
+			const IClientStreamSynchronizer& AssignedToStreamSynchronizer
 			)
 		{
-			FStreamChangelist Result;
+			FStreamChangelist ObjectChanges;
+			FFrequencyChangelist FrequencyChanges;
+			
+			const FObjectReplicationMap& CurrentTargetClientState = AssignedToStreamSynchronizer.GetServerState();
+			const FConcertStreamFrequencySettings& CurrentFrequencySettings = AssignedToStreamSynchronizer.GetFrequencySettings();
+			const FGuid TargetStreamId = AssignedToStreamSynchronizer.GetStreamId();
 			for (const TPair<FGuid, FSubmitStreamChangesResponse>& ChangesRequestedOnRemote : ParallelExecutionResult.StreamResponses)
 			{
 				const TOptional<FCompletedChangeSubmission>& SubmissionInfo = ChangesRequestedOnRemote.Value.SubmissionInfo;
@@ -92,7 +138,8 @@ namespace UE::MultiUserClient
 				const FObjectReplicationMap& ObjectReplicationMap = OldRegisteredObjects[ChangesRequestedOnRemote.Key];
 				for (const TPair<FSoftObjectPath, FReplicatedObjectInfo>& ChangesToApply : ObjectReplicationMap.ReplicatedObjects)
 				{
-					const FObjectInStreamID ObjectId { TargetStreamId, ChangesToApply.Key };
+					const FSoftObjectPath& ObjectPath = ChangesToApply.Key;
+					const FObjectInStreamID ObjectId { TargetStreamId, ObjectPath };
 					TOptional<FConcertReplication_ChangeStream_PutObject> PutRequest = FConcertReplication_ChangeStream_PutObject::MakeFromInfo(ChangesToApply.Value);
 					if (!ensure(PutRequest))
 					{
@@ -110,10 +157,18 @@ namespace UE::MultiUserClient
 						}
 					}
 					
-					Result.ObjectsToPut.Add(ObjectId, MoveTemp(*PutRequest));
+					ObjectChanges.ObjectsToPut.Add(ObjectId, MoveTemp(*PutRequest));
+
+					// The reassigned to client will replicate at the highest rate of either its current rate or the highest from the source clients.
+					const FConcertObjectReplicationSettings& CurrentObjectFrequencySettings = CurrentFrequencySettings.GetSettingsFor(ObjectPath);
+					const FConcertObjectReplicationSettings HighestSourceFrequencySettings = GetHighestFrequency(ObjectPath, OldFrequencies);
+					if (CurrentObjectFrequencySettings < HighestSourceFrequencySettings)
+					{
+						FrequencyChanges.OverridesToAdd.Add(ObjectPath, HighestSourceFrequencySettings);
+					}
 				}
 			}
-			return Result;
+			return { ObjectChanges, FrequencyChanges };
 		}
 
 		/** Builds list of authority changes based on which clients we managed to change successfully. */
@@ -326,7 +381,7 @@ namespace UE::MultiUserClient
 
 		const IClientStreamSynchronizer& StreamSynchronizer = ClientToAssignTo->GetStreamSynchronizer();
 		const FGuid& TargetStreamId = StreamSynchronizer.GetStreamId();
-		auto[OldRegisteredObjects, ReassignedClientRequests, ReassignedAuthority]
+		auto[OldRegisteredObjects, OldFrequencies, ReassignedClientRequests, ReassignedAuthority]
 			= ReassignObjectProperties::Private::BuildChangesForTransferal(
 				ClientManager,
 				ObjectsToReassign,
@@ -340,7 +395,10 @@ namespace UE::MultiUserClient
 			return;
 		}
 		
-		InProgressOperation = { MoveTemp(ObjectsToReassign), ClientId, MoveTemp(OldRegisteredObjects), MoveTemp(ReassignedAuthority), ParallelSubmissionOperation.ToSharedRef() };
+		InProgressOperation = FOperationData{ MoveTemp(ObjectsToReassign), ClientId,
+			MoveTemp(OldRegisteredObjects),MoveTemp(OldFrequencies), MoveTemp(ReassignedAuthority),
+			ParallelSubmissionOperation.ToSharedRef()
+		};
 		ParallelSubmissionOperation->OnCompletedFuture_AnyThread()
 			.Next([this](FParallelExecutionResult&& Result)
 			{
@@ -442,15 +500,15 @@ namespace UE::MultiUserClient
 		const IClientStreamSynchronizer& StreamSynchronizer = ClientToAssignTo->GetStreamSynchronizer();
 		const FGuid& StreamId = StreamSynchronizer.GetStreamId();
 		const bool bIsStreamRegistered = !StreamSynchronizer.GetServerState().IsEmpty();
-		FStreamChangelist Changelist = ReassignObjectProperties::Private::MakeChangelistFromAppliedChanges(
+		auto[ObjectChanges, FrequencyChanges] = ReassignObjectProperties::Private::MakeChangelistFromAppliedChanges(
 			InProgressOperation->OldRegisteredObjects,
+			InProgressOperation->OldFrequencies,
 			ParallelResult,
-			StreamSynchronizer.GetServerState(),
-			StreamId
+			StreamSynchronizer
 			);
 		FConcertReplication_ChangeStream_Request AssignedToClientRequest = bIsStreamRegistered
-			? StreamRequestUtils::BuildChangeRequest_UpdateExistingStream(StreamId, Changelist)
-			: StreamRequestUtils::BuildChangeRequest_CreateNewStream(StreamId, Changelist);
+			? StreamRequestUtils::BuildChangeRequest_UpdateExistingStream(StreamId, ObjectChanges, MoveTemp(FrequencyChanges))
+			: StreamRequestUtils::BuildChangeRequest_CreateNewStream(StreamId, ObjectChanges, MoveTemp(FrequencyChanges));
 		
 		const TSharedPtr<ISubmissionOperation> SubmitOperation = Workflow.SubmitChanges({
 			MoveTemp(AssignedToClientRequest),
diff --git a/Engine/Plugins/Developer/Concert/ConcertApp/MultiUserClient/Source/MultiUserClient/Private/Replication/Submission/MultiEdit/ReassignObjectPropertiesLogic.h b/Engine/Plugins/Developer/Concert/ConcertApp/MultiUserClient/Source/MultiUserClient/Private/Replication/Submission/MultiEdit/ReassignObjectPropertiesLogic.h
--- a/Engine/Plugins/Developer/Concert/ConcertApp/MultiUserClient/Source/MultiUserClient/Private/Replication/Submission/MultiEdit/ReassignObjectPropertiesLogic.h
+++ b/Engine/Plugins/Developer/Concert/ConcertApp/MultiUserClient/Source/MultiUserClient/Private/Replication/Submission/MultiEdit/ReassignObjectPropertiesLogic.h
@@ -75,12 +75,16 @@ namespace UE::MultiUserClient
 
 		struct FOperationData
 		{
+			using FClientId = FGuid;
+			
 			const TArray<FSoftObjectPath> ObjectsBeingReassigned;
 			const FGuid AssignedToClient;
 			/** The objects we tried to remove from all the other clients. If the removal request worked, the (previously) assigned properties will be transferred to the target client. */
-			const TMap<FGuid, FObjectReplicationMap> OldRegisteredObjects;
+			const TMap<FClientId, FObjectReplicationMap> OldRegisteredObjects;
+			/** The frequency settings that clients had at the beginning of the reassign operation. */
+			const TMap<FClientId, FConcertStreamFrequencySettings> OldFrequencies; 
 			/** Clients to the objects they had authority over. Used to later to determine which objects to give the target client authority over (do not transfer authority from failed clients). */
-			const TMap<FGuid, TArray<FSoftObjectPath>> OldAuthority;
+			const TMap<FClientId, TArray<FSoftObjectPath>> OldAuthority;
 			
 			/**
 			 * The first step of reassignment is to remove all object bindings from all other clients.
diff --git a/Engine/Plugins/Developer/Concert/ConcertSync/ConcertSyncCore/Source/ConcertSyncCore/Public/Replication/Data/ReplicationFrequencySettings.h b/Engine/Plugins/Developer/Concert/ConcertSync/ConcertSyncCore/Source/ConcertSyncCore/Public/Replication/Data/ReplicationFrequencySettings.h
--- a/Engine/Plugins/Developer/Concert/ConcertSync/ConcertSyncCore/Source/ConcertSyncCore/Public/Replication/Data/ReplicationFrequencySettings.h
+++ b/Engine/Plugins/Developer/Concert/ConcertSync/ConcertSyncCore/Source/ConcertSyncCore/Public/Replication/Data/ReplicationFrequencySettings.h
@@ -32,7 +32,7 @@ struct CONCERTSYNCCORE_API FConcertObjectReplicationSettings
 	GENERATED_BODY()
 	
 	/** Controls how often this object should be replicated. */
-	UPROPERTY()
+	UPROPERTY(EditAnywhere, Category = "Concert")
 	EConcertObjectReplicationMode ReplicationMode = EConcertObjectReplicationMode::SpecifiedRate;
 	
 	/**
@@ -41,7 +41,7 @@ struct CONCERTSYNCCORE_API FConcertObjectReplicationSettings
 	 * The update in seconds is given by 1 / ReplicationRate.
 	 * The default rate of 30 results in an update interval every 0.033s
 	 */
-	UPROPERTY()
+	UPROPERTY(EditAnywhere, Category = "Concert")
 	uint8 ReplicationRate = 30;
 
 	bool IsValid() const { return ReplicationMode == EConcertObjectReplicationMode::Realtime || ReplicationRate > 0; }
@@ -84,6 +84,29 @@ struct CONCERTSYNCCORE_API FConcertObjectReplicationSettings
 	{
 		return !(Left == Right);
 	}
+
+	/** @return Whether Left replicates less frequently than Right. */
+	friend bool operator<(const FConcertObjectReplicationSettings& Left, const FConcertObjectReplicationSettings& Right)
+	{
+		const bool bLeftIsRealtime = Left.ReplicationMode == EConcertObjectReplicationMode::Realtime;
+		const bool bRightIsRealtime = Right.ReplicationMode == EConcertObjectReplicationMode::Realtime;
+		if (bLeftIsRealtime || bRightIsRealtime)
+		{
+			return !bLeftIsRealtime && bRightIsRealtime;
+		}
+		return Left.ReplicationRate < Right.ReplicationRate;
+	}
+	/** @return Whether Left replicates less or equally frequently than Right. */
+	friend bool operator<=(const FConcertObjectReplicationSettings& Left, const FConcertObjectReplicationSettings& Right)
+	{
+		const bool bLeftIsRealtime = Left.ReplicationMode == EConcertObjectReplicationMode::Realtime;
+		const bool bRightIsRealtime = Right.ReplicationMode == EConcertObjectReplicationMode::Realtime;
+		if (bLeftIsRealtime || bRightIsRealtime)
+		{
+			return bRightIsRealtime;
+		}
+		return Left.ReplicationRate <= Right.ReplicationRate;
+	}
 };
 
 /** Specifies replication frequency settings for an entire stream. */
diff --git a/Engine/Plugins/Developer/Concert/ConcertSync/ConcertSyncTest/Source/ConcertSyncTest/Private/Replication/Components/FrequencySettingsTests.cpp b/Engine/Plugins/Developer/Concert/ConcertSync/ConcertSyncTest/Source/ConcertSyncTest/Private/Replication/Components/FrequencySettingsTests.cpp
new file mode 100644
--- /dev/null
+++ b/Engine/Plugins/Developer/Concert/ConcertSync/ConcertSyncTest/Source/ConcertSyncTest/Private/Replication/Components/FrequencySettingsTests.cpp
@@ -0,0 +1,40 @@
+// Copyright Epic Games, Inc. All Rights Reserved.
+
+#include "Replication/Data/ReplicationFrequencySettings.h"
+
+#include "Misc/AutomationTest.h"
+#include "Misc/Paths.h"
+#include "HAL/FileManager.h"
+
+namespace UE::ConcertSyncTests
+{
+	/** Tests that FConcertObjectReplicationSettings's < and <= work as expected. */
+	IMPLEMENT_SIMPLE_AUTOMATION_TEST(FFrequencySettingsLessEqual, "Editor.Concert.Replication.Components.FrequencySettingsLessEqual", EAutomationTestFlags::EditorContext | EAutomationTestFlags::EngineFilter);
+	bool FFrequencySettingsLessEqual::RunTest(const FString& Parameters)
+	{
+		constexpr FConcertObjectReplicationSettings FPS_30 { EConcertObjectReplicationMode::SpecifiedRate, 30 }; 
+		constexpr FConcertObjectReplicationSettings FPS_60 { EConcertObjectReplicationMode::SpecifiedRate, 60 };
+		// Note for realtime ReplicationRate is ignored - we're only putting  30 and 60 to test that they are also ignored for comparisons
+		constexpr FConcertObjectReplicationSettings Realtime_20 { EConcertObjectReplicationMode::Realtime, 20 };
+		constexpr FConcertObjectReplicationSettings Realtime_60 { EConcertObjectReplicationMode::Realtime, 60 };
+
+		TestTrue(TEXT("FPS_30 <= FPS_30 == true"), FPS_30 <= FPS_30);
+		TestTrue(TEXT("FPS_30 < FPS_60 == true"), FPS_30 < FPS_60);
+		TestTrue(TEXT("FPS_30 <= FPS_60 == true"), FPS_30 <= FPS_60);
+		TestFalse(TEXT("FPS_60 < FPS_30 == false"), FPS_60 < FPS_30);
+		TestFalse(TEXT("FPS_60 <= FPS_30 == false"), FPS_60 <= FPS_30);
+		
+		TestTrue(TEXT("FPS_30 < FPS_60 == true"), FPS_30 < Realtime_20);
+		TestTrue(TEXT("FPS_30 <= FPS_60 == true"), FPS_30 <= Realtime_20);
+		TestFalse(TEXT("Realtime_20 < FPS_30 == false"), Realtime_20 < FPS_30);
+		TestFalse(TEXT("Realtime_20 <= FPS_30 == false"), Realtime_20 <= FPS_30);
+		
+		TestTrue(TEXT("Realtime_20 <= Realtime_20 == true"), Realtime_20 <= Realtime_20);
+		TestTrue(TEXT("Realtime_20 <= Realtime_60 == true"), Realtime_20 <= Realtime_60);
+		TestTrue(TEXT("Realtime_60 <= Realtime_20 == true"), Realtime_60 <= Realtime_20);
+		TestFalse(TEXT("Realtime_20 < Realtime_60 == false"), Realtime_20 < Realtime_60);
+		TestFalse(TEXT("Realtime_60 < Realtime_20 == false"), Realtime_60 < Realtime_20);
+		
+		return true;
+	}
+}
