Expose Multi User replication frequency to Blueprints
- Blueprints can edit the frequency
- Blueprints can read what the client thinks is the frequency rate of an object

diff --git a/Engine/Plugins/Developer/Concert/ConcertApp/MultiUserClient/Source/MultiUserClient/Private/Replication/MultiUserReplicationManager.cpp b/Engine/Plugins/Developer/Concert/ConcertApp/MultiUserClient/Source/MultiUserClient/Private/Replication/MultiUserReplicationManager.cpp
--- a/Engine/Plugins/Developer/Concert/ConcertApp/MultiUserClient/Source/MultiUserClient/Private/Replication/MultiUserReplicationManager.cpp
+++ b/Engine/Plugins/Developer/Concert/ConcertApp/MultiUserClient/Source/MultiUserClient/Private/Replication/MultiUserReplicationManager.cpp
@@ -112,6 +112,18 @@ namespace UE::MultiUserClient
 		return nullptr;
 	}
 
+	const FConcertStreamFrequencySettings* FMultiUserReplicationManager::FindReplicationFrequenciesForClient(const FGuid& ClientId) const
+	{
+		if (ConnectedState && ensureMsgf(IsInGameThread(), TEXT("To simplify implementation, only calls from game thread are allowed.")))
+		{
+			const FReplicationClient* ReplicationClient = ConnectedState->ClientManager.FindClient(ClientId);
+			return ReplicationClient
+				? &ReplicationClient->GetStreamSynchronizer().GetFrequencySettings()
+				: nullptr;
+		}
+		return nullptr;
+	}
+
 	bool FMultiUserReplicationManager::IsReplicatingObject(const FGuid& ClientId, const FSoftObjectPath& ObjectPath) const 
 	{
 		if (ConnectedState && ensureMsgf(IsInGameThread(), TEXT("To simplify implementation, only calls from game thread are allowed.")))
diff --git a/Engine/Plugins/Developer/Concert/ConcertApp/MultiUserClient/Source/MultiUserClient/Private/Replication/MultiUserReplicationManager.h b/Engine/Plugins/Developer/Concert/ConcertApp/MultiUserClient/Source/MultiUserClient/Private/Replication/MultiUserReplicationManager.h
--- a/Engine/Plugins/Developer/Concert/ConcertApp/MultiUserClient/Source/MultiUserClient/Private/Replication/MultiUserReplicationManager.h
+++ b/Engine/Plugins/Developer/Concert/ConcertApp/MultiUserClient/Source/MultiUserClient/Private/Replication/MultiUserReplicationManager.h
@@ -70,6 +70,7 @@ namespace UE::MultiUserClient
 
 		//~ Begin IMultiUserReplication Interface
 		virtual const FConcertObjectReplicationMap* FindReplicationMapForClient(const FGuid& ClientId) const override;
+		virtual const FConcertStreamFrequencySettings* FindReplicationFrequenciesForClient(const FGuid& ClientId) const override;
 		virtual bool IsReplicatingObject(const FGuid& ClientId, const FSoftObjectPath& ObjectPath) const override;
 		virtual void RegisterReplicationDiscoverer(TSharedRef<IReplicationDiscoverer> Discoverer) override;
 		virtual void RemoveReplicationDiscoverer(const TSharedRef<IReplicationDiscoverer>& Discoverer) override;
diff --git a/Engine/Plugins/Developer/Concert/ConcertApp/MultiUserClient/Source/MultiUserClient/Private/Replication/Submission/External/ClientChangeConversionUtils.cpp b/Engine/Plugins/Developer/Concert/ConcertApp/MultiUserClient/Source/MultiUserClient/Private/Replication/Submission/External/ClientChangeConversionUtils.cpp
--- a/Engine/Plugins/Developer/Concert/ConcertApp/MultiUserClient/Source/MultiUserClient/Private/Replication/Submission/External/ClientChangeConversionUtils.cpp
+++ b/Engine/Plugins/Developer/Concert/ConcertApp/MultiUserClient/Source/MultiUserClient/Private/Replication/Submission/External/ClientChangeConversionUtils.cpp
@@ -102,6 +102,16 @@ namespace UE::MultiUserClient::ClientChangeConversionUtils
 				);
 			}
 		}
+
+		static FFrequencyChangelist BuildFrequencyChanges(FConcertReplication_ChangeStream_Frequency FrequencyChanges)
+		{
+			FFrequencyChangelist Changelist;
+			Changelist.OverridesToAdd = MoveTemp(FrequencyChanges.OverridesToAdd);
+			Changelist.OverridesToRemove = MoveTemp(FrequencyChanges.OverridesToRemove);
+			Changelist.NewDefaults = EnumHasAnyFlags(FrequencyChanges.Flags, EConcertReplicationChangeFrequencyFlags::SetDefaults)
+				? MoveTemp(FrequencyChanges.NewDefaults) : TOptional<FConcertObjectReplicationSettings>{};
+			return Changelist;
+		}
 	}
 
 	TOptional<FConcertReplication_ChangeStream_Request> Transform(
@@ -113,15 +123,16 @@ namespace UE::MultiUserClient::ClientChangeConversionUtils
 		FStreamChangelist StreamChangelist;
 		Private::BuildStreamChanges(MoveTemp(Request.PropertyChanges), ClientStreamId, ClientStreamContent, StreamChangelist);
 		Private::BuildAuthorityChanges(MoveTemp(Request.ObjectsToRemove), ClientStreamId, StreamChangelist);
-		if (StreamChangelist.ObjectsToPut.IsEmpty() && StreamChangelist.ObjectsToRemove.IsEmpty())
+		FFrequencyChangelist FrequencyChangelist = Private::BuildFrequencyChanges(MoveTemp(Request.FrequencyChanges));
+		if (StreamChangelist.ObjectsToPut.IsEmpty() && StreamChangelist.ObjectsToRemove.IsEmpty() && FrequencyChangelist.IsEmpty())
 		{
 			return {};
 		}
 		
 		const bool bNeedsToRegisterStream = ClientStreamContent.IsEmpty();
 		return bNeedsToRegisterStream
-			? StreamRequestUtils::BuildChangeRequest_CreateNewStream(ClientStreamId, StreamChangelist)
-			: StreamRequestUtils::BuildChangeRequest_UpdateExistingStream(ClientStreamId, MoveTemp(StreamChangelist));
+			? StreamRequestUtils::BuildChangeRequest_CreateNewStream(ClientStreamId, StreamChangelist, MoveTemp(FrequencyChangelist))
+			: StreamRequestUtils::BuildChangeRequest_UpdateExistingStream(ClientStreamId, MoveTemp(StreamChangelist), MoveTemp(FrequencyChangelist));
 	}
 
 	TOptional<FConcertReplication_ChangeAuthority_Request> Transform(
diff --git a/Engine/Plugins/Developer/Concert/ConcertApp/MultiUserClient/Source/MultiUserClient/Private/Replication/Util/StreamRequestUtils.h b/Engine/Plugins/Developer/Concert/ConcertApp/MultiUserClient/Source/MultiUserClient/Private/Replication/Util/StreamRequestUtils.h
--- a/Engine/Plugins/Developer/Concert/ConcertApp/MultiUserClient/Source/MultiUserClient/Private/Replication/Util/StreamRequestUtils.h
+++ b/Engine/Plugins/Developer/Concert/ConcertApp/MultiUserClient/Source/MultiUserClient/Private/Replication/Util/StreamRequestUtils.h
@@ -21,6 +21,8 @@ namespace UE::MultiUserClient
 		TSet<FSoftObjectPath> OverridesToRemove;
 		TMap<FSoftObjectPath, FConcertObjectReplicationSettings> OverridesToAdd;
 		TOptional<FConcertObjectReplicationSettings> NewDefaults;
+
+		bool IsEmpty() const { return OverridesToRemove.IsEmpty() && OverridesToAdd.IsEmpty() && !NewDefaults.IsSet(); }
 	};
 }
 
diff --git a/Engine/Plugins/Developer/Concert/ConcertApp/MultiUserClient/Source/MultiUserClient/Public/Replication/ChangeOperationTypes.h b/Engine/Plugins/Developer/Concert/ConcertApp/MultiUserClient/Source/MultiUserClient/Public/Replication/ChangeOperationTypes.h
--- a/Engine/Plugins/Developer/Concert/ConcertApp/MultiUserClient/Source/MultiUserClient/Public/Replication/ChangeOperationTypes.h
+++ b/Engine/Plugins/Developer/Concert/ConcertApp/MultiUserClient/Source/MultiUserClient/Public/Replication/ChangeOperationTypes.h
@@ -3,6 +3,7 @@
 #pragma once
 
 #include "Replication/Data/ConcertPropertySelection.h"
+#include "Replication/Messages/ChangeStream.h"
 
 /**
  * This public, high-level Multi-User API for replication  effectively encapsulates the lower-level ConcertSyncCore replication API,
@@ -101,18 +102,23 @@ namespace UE::MultiUserClient
 	
 	struct FPropertyChange
 	{
+		/** The properties of the operation */
 		TArray<FConcertPropertyChain> Properties;
+		/** How to interpret Properties. */
 		EPropertyChangeType ChangeType;
 	};
 
 	/** Params for changing a client's stream. */
 	struct FChangeStreamRequest
 	{
-		/** Property changes to make to objects. */
+		/** Property changes to make to objects. This can be used to add new objects. */
 		TMap<UObject*, FPropertyChange> PropertyChanges;
 
 		/** Objects that should be unregistered (they will also stop replicating if added here) */
 		TSet<FSoftObjectPath> ObjectsToRemove;
+
+		/** Changes how often objects are supposed to be replicated. */
+		FConcertReplication_ChangeStream_Frequency FrequencyChanges;
 	};
 
 	/** Params for changing a client's authority (what they're replicating). */
diff --git a/Engine/Plugins/Developer/Concert/ConcertApp/MultiUserClient/Source/MultiUserClient/Public/Replication/IMultiUserReplication.h b/Engine/Plugins/Developer/Concert/ConcertApp/MultiUserClient/Source/MultiUserClient/Public/Replication/IMultiUserReplication.h
--- a/Engine/Plugins/Developer/Concert/ConcertApp/MultiUserClient/Source/MultiUserClient/Public/Replication/IMultiUserReplication.h
+++ b/Engine/Plugins/Developer/Concert/ConcertApp/MultiUserClient/Source/MultiUserClient/Public/Replication/IMultiUserReplication.h
@@ -27,6 +27,15 @@ namespace UE::MultiUserClient
 		 * This function must be called from the game thread.
 		 */
 		virtual const FConcertObjectReplicationMap* FindReplicationMapForClient(const FGuid& ClientId) const = 0;
+
+		/**
+		 * @return Gets the last known server object replication frequencies.
+		 * 
+		 * This server state is regularly polled whilst the local client state should always be in synch.
+		 * This function must be called from the game thread.
+		 */
+		virtual const FConcertStreamFrequencySettings* FindReplicationFrequenciesForClient(const FGuid& ClientId) const = 0;
+		
 		/**
 		 * @return Whether the local editor instance thinks the client has authority over the properties it has registered to ObjectPath.
 		 * This function must be called from the game thread.
diff --git a/Engine/Plugins/Developer/Concert/ConcertApp/MultiUserClient/Source/MultiUserClientLibrary/Private/Replication/Async/ChangeClientAsyncAction.cpp b/Engine/Plugins/Developer/Concert/ConcertApp/MultiUserClient/Source/MultiUserClientLibrary/Private/Replication/Async/ChangeClientAsyncAction.cpp
--- a/Engine/Plugins/Developer/Concert/ConcertApp/MultiUserClient/Source/MultiUserClientLibrary/Private/Replication/Async/ChangeClientAsyncAction.cpp
+++ b/Engine/Plugins/Developer/Concert/ConcertApp/MultiUserClient/Source/MultiUserClientLibrary/Private/Replication/Async/ChangeClientAsyncAction.cpp
@@ -34,7 +34,7 @@ void UChangeClientAsyncAction::Activate()
 			{
 				if (IsInGameThread())
 				{
-					// Do not call IsValid it outside the game thread!
+					// As per documentation, IsValid is unsafe to call outside of the game thread
 					if (WeakThis.IsValid())
 					{
 						WeakThis->OnCompleted.Broadcast(UE::MultiUserClientLibrary::Transform(MoveTemp(Result)));
diff --git a/Engine/Plugins/Developer/Concert/ConcertApp/MultiUserClient/Source/MultiUserClientLibrary/Private/Replication/Async/ChangeClientBlueprintParams.cpp b/Engine/Plugins/Developer/Concert/ConcertApp/MultiUserClient/Source/MultiUserClientLibrary/Private/Replication/Async/ChangeClientBlueprintParams.cpp
--- a/Engine/Plugins/Developer/Concert/ConcertApp/MultiUserClient/Source/MultiUserClientLibrary/Private/Replication/Async/ChangeClientBlueprintParams.cpp
+++ b/Engine/Plugins/Developer/Concert/ConcertApp/MultiUserClient/Source/MultiUserClientLibrary/Private/Replication/Async/ChangeClientBlueprintParams.cpp
@@ -41,6 +41,32 @@ namespace UE::MultiUserClientLibrary
 		return Result;
 	}
 	
+	EConcertObjectReplicationMode Transform(EMultiUserObjectReplicationMode Data)
+	{
+		static_assert(static_cast<int32>(EConcertObjectReplicationMode::Count) == 2, "Update this transform operation and its asserts.");
+		static_assert(static_cast<int32>(EConcertObjectReplicationMode::Realtime) == static_cast<int32>(EMultiUserObjectReplicationMode::Realtime));
+		static_assert(static_cast<int32>(EConcertObjectReplicationMode::SpecifiedRate) == static_cast<int32>(EMultiUserObjectReplicationMode::SpecifiedRate));
+		return static_cast<EConcertObjectReplicationMode>(Data);
+	}
+
+	EMultiUserObjectReplicationMode Transform(EConcertObjectReplicationMode Data)
+	{
+		static_assert(static_cast<int32>(EConcertObjectReplicationMode::Count) == 2, "Update this transform operation and its asserts.");
+		static_assert(static_cast<int32>(EConcertObjectReplicationMode::Realtime) == static_cast<int32>(EMultiUserObjectReplicationMode::Realtime));
+		static_assert(static_cast<int32>(EConcertObjectReplicationMode::SpecifiedRate) == static_cast<int32>(EMultiUserObjectReplicationMode::SpecifiedRate));
+		return static_cast<EMultiUserObjectReplicationMode>(Data);
+	}
+
+	FConcertObjectReplicationSettings Transform(const FMultiUserObjectReplicationSettings& Data)
+	{
+		return { Transform(Data.Mode), Data.ReplicationRate };
+	}
+	
+	FMultiUserObjectReplicationSettings Transform(const FConcertObjectReplicationSettings& Data)
+	{
+		return { Transform(Data.ReplicationMode), Data.ReplicationRate };
+	}
+	
 	MultiUserClient::FChangeStreamRequest Transform(FMultiUserChangeStreamRequest Data)
 	{
 		MultiUserClient::FChangeStreamRequest Result { .ObjectsToRemove = MoveTemp(Data.ObjectsToRemove) };
@@ -53,6 +79,19 @@ namespace UE::MultiUserClientLibrary
 				Result.PropertyChanges.Emplace(Change.Key, Transform(MoveTemp(Change.Value)));
 			}
 		}
+
+		FMultiUserFrequencyChangeRequest& InFrequencyChanges = Data.FrequencyChanges;
+		FConcertReplication_ChangeStream_Frequency& OutFrequencyChanges = Result.FrequencyChanges;
+		if (!InFrequencyChanges.IsEmpty())
+		{
+			OutFrequencyChanges.Flags = InFrequencyChanges.bChangeDefaults ? EConcertReplicationChangeFrequencyFlags::SetDefaults : EConcertReplicationChangeFrequencyFlags::None;
+			OutFrequencyChanges.NewDefaults = Transform(InFrequencyChanges.NewDefaults);
+			OutFrequencyChanges.OverridesToRemove = MoveTemp(InFrequencyChanges.OverridesToRemove);
+			Algo::Transform(InFrequencyChanges.OverridesToAdd, OutFrequencyChanges.OverridesToAdd, [](const TPair<FSoftObjectPath, FMultiUserObjectReplicationSettings>& Pair)
+			{
+				return TPair<FSoftObjectPath, FConcertObjectReplicationSettings>{ Pair.Key, Transform(Pair.Value) };
+			});
+		}
 		
 		return Result;
 	}
diff --git a/Engine/Plugins/Developer/Concert/ConcertApp/MultiUserClient/Source/MultiUserClientLibrary/Private/Replication/MultiUserReplicationSubsystem.cpp b/Engine/Plugins/Developer/Concert/ConcertApp/MultiUserClient/Source/MultiUserClientLibrary/Private/Replication/MultiUserReplicationSubsystem.cpp
--- a/Engine/Plugins/Developer/Concert/ConcertApp/MultiUserClient/Source/MultiUserClientLibrary/Private/Replication/MultiUserReplicationSubsystem.cpp
+++ b/Engine/Plugins/Developer/Concert/ConcertApp/MultiUserClient/Source/MultiUserClientLibrary/Private/Replication/MultiUserReplicationSubsystem.cpp
@@ -2,10 +2,13 @@
 
 #include "Replication/MultiUserReplicationSubsystem.h"
 
+#include "Replication/Async/ChangeClientBlueprintParams.h"
+
 #if WITH_CONCERT
 #include "UObjectAdapterReplicationDiscoverer.h"
 #include "IMultiUserClientModule.h"
 #include "Replication/Data/ObjectReplicationMap.h"
+#include "Replication/Data/ReplicationFrequencySettings.h"
 #include "Replication/IMultiUserReplication.h"
 
 #include "Algo/Transform.h"
@@ -23,6 +26,29 @@ bool UMultiUserReplicationSubsystem::IsReplicatingObject(const FGuid& ClientId,
 	return false;
 }
 
+bool UMultiUserReplicationSubsystem::GetObjectReplicationFrequency(const FGuid& ClientId, const FSoftObjectPath& ObjectPath, FMultiUserObjectReplicationSettings& OutFrequency)
+{
+#if WITH_CONCERT
+	const UE::MultiUserClient::IMultiUserReplication* ReplicationInterface = IMultiUserClientModule::Get().GetReplication();
+	if (ensureMsgf(ReplicationInterface, TEXT("We expected it to always be valid.")))
+	{
+		const FConcertObjectReplicationMap* ObjectMap = ReplicationInterface->FindReplicationMapForClient(ClientId);
+		if (!ObjectMap || !ObjectMap->HasProperties(ObjectPath))
+		{
+			return false;
+		}
+		
+		const FConcertStreamFrequencySettings* FrequencySettings = ReplicationInterface->FindReplicationFrequenciesForClient(ClientId);
+		if (FrequencySettings)
+		{
+			OutFrequency = UE::MultiUserClientLibrary::Transform(FrequencySettings->GetSettingsFor(ObjectPath));
+			return true;
+		}
+	}
+#endif
+	return false;
+}
+
 TArray<FConcertPropertyChainWrapper> UMultiUserReplicationSubsystem::GetPropertiesRegisteredToObject(const FGuid& ClientId, const FSoftObjectPath& ObjectPath) const
 {
 #if WITH_CONCERT  
diff --git a/Engine/Plugins/Developer/Concert/ConcertApp/MultiUserClient/Source/MultiUserClientLibrary/Public/Replication/Async/ChangeClientBlueprintParams.h b/Engine/Plugins/Developer/Concert/ConcertApp/MultiUserClient/Source/MultiUserClientLibrary/Public/Replication/Async/ChangeClientBlueprintParams.h
--- a/Engine/Plugins/Developer/Concert/ConcertApp/MultiUserClient/Source/MultiUserClientLibrary/Public/Replication/Async/ChangeClientBlueprintParams.h
+++ b/Engine/Plugins/Developer/Concert/ConcertApp/MultiUserClient/Source/MultiUserClientLibrary/Public/Replication/Async/ChangeClientBlueprintParams.h
@@ -96,14 +96,68 @@ struct FMultiUserPropertyChange
 	GENERATED_BODY()
 
 	/** The properties for the object. See ChangeType for how they are used */
-	UPROPERTY(BlueprintReadWrite, Category = "Multi-user")
+	UPROPERTY(EditAnywhere, BlueprintReadWrite, Category = "Multi-user")
 	TArray<FConcertPropertyChainWrapper> Properties;
 
 	/** How you want Properties to be applied to the object. */
-	UPROPERTY(BlueprintReadWrite, Category = "Multi-user")
+	UPROPERTY(EditAnywhere, BlueprintReadWrite, Category = "Multi-user")
 	EMultiUserPropertyChangeType ChangeType = EMultiUserPropertyChangeType::Add;
 };
 
+UENUM(BlueprintType)
+enum class EMultiUserObjectReplicationMode : uint8
+{
+	/** Replicate at the rate specified at FMultiUserObjectReplicationSettings::ReplicationRate */
+	SpecifiedRate,
+	/** Replicate the object as often as possible: every tick. */
+	Realtime
+};
+
+/** Frequency settings for a particular object */
+USTRUCT(BlueprintType)
+struct FMultiUserObjectReplicationSettings
+{
+	GENERATED_BODY()
+
+	/** Determines whether to send a replication event every tick (Realtime) or at the specified rate (SpecifiedRate). */
+	UPROPERTY(EditAnywhere, BlueprintReadWrite, Category = "Multi-user")
+	EMultiUserObjectReplicationMode Mode = EMultiUserObjectReplicationMode::SpecifiedRate;
+
+	/** If Mode == SpecifiedRate, then replicate this many times per second. */
+	UPROPERTY(EditAnywhere, BlueprintReadWrite, Category = "Multi-user")
+	uint8 ReplicationRate = 30;
+};
+
+/**
+ * Params for changing how often objects replicate.
+ * 
+ * All objects default to a global setting that can be changed by setting bChangeDefaults = true and changing NewDefaults to contain the new defaults.
+ * Objects can have overrides mutated via OverridesToAdd and OverridesToRemove. Overrides override the default behavior.
+ */
+USTRUCT(BlueprintType)
+struct FMultiUserFrequencyChangeRequest
+{
+	GENERATED_BODY()
+
+	/** Objects for which to remove overrides. */
+	UPROPERTY(EditAnywhere, BlueprintReadWrite, Category = "Multi-user")
+	TSet<FSoftObjectPath> OverridesToRemove;
+
+	/** Objects for which to add overrides */
+	UPROPERTY(EditAnywhere, BlueprintReadWrite, Category = "Multi-user")
+	TMap<FSoftObjectPath, FMultiUserObjectReplicationSettings> OverridesToAdd;
+
+	/** Set new frequency defaults for all objects that do not have any overrides. */
+	UPROPERTY(EditAnywhere, BlueprintReadWrite, Category = "Multi-user")
+	FMultiUserObjectReplicationSettings NewDefaults;
+
+	/** Whether to replace the defaults currently registered on the server with the ones specified in NewDefaults. */
+	UPROPERTY(EditAnywhere, BlueprintReadWrite, Category = "Multi-user")
+	bool bChangeDefaults = false;
+
+	bool IsEmpty() const { return OverridesToRemove.IsEmpty() && OverridesToAdd.IsEmpty() && !bChangeDefaults; }
+};
+
 /** Params for changing a client's stream. */
 USTRUCT(BlueprintType)
 struct FMultiUserChangeStreamRequest
@@ -111,14 +165,18 @@ struct FMultiUserChangeStreamRequest
 	GENERATED_BODY()
 	
 	/** Property changes to make to objects. */
-	UPROPERTY(BlueprintReadWrite, Category = "Multi-user")
+	UPROPERTY(EditAnywhere, BlueprintReadWrite, Category = "Multi-user")
 	TMap<UObject*, FMultiUserPropertyChange> PropertyChanges;
 
 	/** Objects that should be unregistered (they will also stop replicating if added here) */
-	UPROPERTY(BlueprintReadWrite, Category = "Multi-user")
+	UPROPERTY(EditAnywhere, BlueprintReadWrite, Category = "Multi-user")
 	TSet<FSoftObjectPath> ObjectsToRemove;
+
+	/** Change how often objects are supposed to be replicated. */
+	UPROPERTY(EditAnywhere, BlueprintReadWrite, Category = "Multi-user")
+	FMultiUserFrequencyChangeRequest FrequencyChanges;
 	
-	bool IsEmpty() const { return PropertyChanges.IsEmpty() && ObjectsToRemove.IsEmpty(); }
+	bool IsEmpty() const { return PropertyChanges.IsEmpty() && ObjectsToRemove.IsEmpty() && FrequencyChanges.IsEmpty(); }
 };
 
 /** Params for changing a client's authority (what they're replicating). */
@@ -128,11 +186,11 @@ struct FMultiUserChangeAuthorityRequest
 	GENERATED_BODY()
 	
 	/** Objects that should start replicating. The objects must previously have been registered. */
-	UPROPERTY(BlueprintReadWrite, Category = "Multi-user")
+	UPROPERTY(EditAnywhere, BlueprintReadWrite, Category = "Multi-user")
 	TSet<FSoftObjectPath> ObjectsToStartReplicating;
 
 	/** Objects that should stop replicating. */
-	UPROPERTY(BlueprintReadWrite, Category = "Multi-user")
+	UPROPERTY(EditAnywhere, BlueprintReadWrite, Category = "Multi-user")
 	TSet<FSoftObjectPath> ObjectToStopReplicating;
 
 	bool IsEmpty() const { return ObjectsToStartReplicating.IsEmpty() && ObjectToStopReplicating.IsEmpty(); }
@@ -183,7 +241,12 @@ namespace UE::MultiUserClientLibrary
 	
 	MULTIUSERCLIENTLIBRARY_API MultiUserClient::EPropertyChangeType Transform(EMultiUserPropertyChangeType Data);
 	MULTIUSERCLIENTLIBRARY_API MultiUserClient::FPropertyChange Transform(FMultiUserPropertyChange Data);
-	
+
+	MULTIUSERCLIENTLIBRARY_API EConcertObjectReplicationMode Transform(EMultiUserObjectReplicationMode Data);
+	MULTIUSERCLIENTLIBRARY_API EMultiUserObjectReplicationMode Transform(EConcertObjectReplicationMode Data);
+	MULTIUSERCLIENTLIBRARY_API FConcertObjectReplicationSettings Transform(const FMultiUserObjectReplicationSettings& Data);
+	MULTIUSERCLIENTLIBRARY_API FMultiUserObjectReplicationSettings Transform(const FConcertObjectReplicationSettings& Data);
+
 	MULTIUSERCLIENTLIBRARY_API MultiUserClient::FChangeStreamRequest Transform(FMultiUserChangeStreamRequest Data);
 	MULTIUSERCLIENTLIBRARY_API MultiUserClient::FChangeAuthorityRequest Transform(FMultiUserChangeAuthorityRequest Data);
 	MULTIUSERCLIENTLIBRARY_API MultiUserClient::FChangeClientReplicationRequest Transform(FMultiUserChangeClientReplicationRequest Data);
diff --git a/Engine/Plugins/Developer/Concert/ConcertApp/MultiUserClient/Source/MultiUserClientLibrary/Public/Replication/MultiUserReplicationSubsystem.h b/Engine/Plugins/Developer/Concert/ConcertApp/MultiUserClient/Source/MultiUserClientLibrary/Public/Replication/MultiUserReplicationSubsystem.h
--- a/Engine/Plugins/Developer/Concert/ConcertApp/MultiUserClient/Source/MultiUserClientLibrary/Public/Replication/MultiUserReplicationSubsystem.h
+++ b/Engine/Plugins/Developer/Concert/ConcertApp/MultiUserClient/Source/MultiUserClientLibrary/Public/Replication/MultiUserReplicationSubsystem.h
@@ -6,10 +6,9 @@
 #include "ConcertPropertyChainWrapper.h"
 #include "MultiUserReplicationSubsystem.generated.h"
 
-namespace UE::MultiUserClientLibrary
-{
-	class FUObjectAdapterReplicationDiscoverer;
-}
+struct FMultiUserObjectReplicationSettings;
+
+namespace UE::MultiUserClientLibrary { class FUObjectAdapterReplicationDiscoverer; }
 
 /** Exposes ways to interact with the Multi-user replication system via Blueprints. */
 UCLASS()
@@ -28,6 +27,14 @@ public:
 	UFUNCTION(BlueprintPure, Category = "Multi-user")
 	bool IsReplicatingObject(const FGuid& ClientId, const FSoftObjectPath& ObjectPath) const;
 
+	/**
+	 * @return Whether OutFrequency was modified.
+	 * @note An object can be registered but not replicated.
+	 * @see UMultiUserSubsystem::GetLocalClientId and UMultiUserSubsystem::GetRemoteClientIds
+	 */
+	UFUNCTION(BlueprintPure, Category = "Multi-user")
+	bool GetObjectReplicationFrequency(const FGuid& ClientId, const FSoftObjectPath& ObjectPath, FMultiUserObjectReplicationSettings& OutFrequency);
+
 	/**
 	 * @return The properties the client has registered for replication for the object.
 	 * @note An object can be registered but not replicated. Use IsReplicatingObject to find out whether the client is replicating the returned properties.
@@ -73,10 +80,12 @@ public:
 private:
 
 	/**
-	 * Registered when initialized.
+	 * This is used only for when the adds an object through the Add button in the UI.
 	 * 
 	 * This allows UObjects, the target being Blueprints, to implement the IConcertReplicationRegistration interface through which MU will use to
 	 * auto-add properties when registering an object to a client's replication stream.
+	 *
+	 * Registered when this subsystem is initialized.
 	 */
 	TSharedPtr<UE::MultiUserClientLibrary::FUObjectAdapterReplicationDiscoverer> UObjectAdapter;
 };
diff --git a/Engine/Plugins/Developer/Concert/ConcertSync/ConcertSyncCore/Source/ConcertSyncCore/Public/Replication/Data/ReplicationFrequencySettings.h b/Engine/Plugins/Developer/Concert/ConcertSync/ConcertSyncCore/Source/ConcertSyncCore/Public/Replication/Data/ReplicationFrequencySettings.h
--- a/Engine/Plugins/Developer/Concert/ConcertSync/ConcertSyncCore/Source/ConcertSyncCore/Public/Replication/Data/ReplicationFrequencySettings.h
+++ b/Engine/Plugins/Developer/Concert/ConcertSync/ConcertSyncCore/Source/ConcertSyncCore/Public/Replication/Data/ReplicationFrequencySettings.h
@@ -11,7 +11,12 @@ enum class EConcertObjectReplicationMode : uint8
 	/** Replicate at the rate specified at FConcertReplicatedObjectInfo::ReplicationRate */
 	SpecifiedRate,
 	/** Replicate the object as often as possible: every tick. */
-	Realtime
+	Realtime,
+
+	/***** ADD NEW ENTRIES ABOVE THIS LINE *****/
+	
+	/** Not a real mode. */
+	Count
 };
 
 namespace UE::ConcertSyncCore
