Fix submission futures not getting executed in the right order

diff --git a/Engine/Plugins/Developer/Concert/ConcertApp/MultiUserClient/Source/MultiUserClient/Private/Replication/Submission/Remote/SubmissionWorkflow_RemoteClient.cpp b/Engine/Plugins/Developer/Concert/ConcertApp/MultiUserClient/Source/MultiUserClient/Private/Replication/Submission/Remote/SubmissionWorkflow_RemoteClient.cpp
--- a/Engine/Plugins/Developer/Concert/ConcertApp/MultiUserClient/Source/MultiUserClient/Private/Replication/Submission/Remote/SubmissionWorkflow_RemoteClient.cpp
+++ b/Engine/Plugins/Developer/Concert/ConcertApp/MultiUserClient/Source/MultiUserClient/Private/Replication/Submission/Remote/SubmissionWorkflow_RemoteClient.cpp
@@ -27,7 +27,7 @@ namespace UE::MultiUserClient
 	{
 		ConcertSession->UnregisterCustomEventHandler<FMultiUser_ChangeRemote_StreamUpdatedEvent>(this);
 		ConcertSession->UnregisterCustomEventHandler<FMultiUser_ChangeRemote_AuthorityUpdatedEvent>(this);
-		CleanUpSubmissionProcess();
+		// InProgressOperation destructor chain execute any pending promises
 	}
 
 	TSharedPtr<ISubmissionOperation> FSubmissionWorkflow_RemoteClient::SubmitChanges(FSubmissionParams Params)
@@ -40,16 +40,8 @@ namespace UE::MultiUserClient
 		const bool bIsModifyingStream = !Params.IsStreamChangeEmpty();
 		const TSharedRef<FSingleClientSubmissionOperation> OperationResult = MakeShared<FSingleClientSubmissionOperation>(bIsModifyingStream);
 
-		// Since SendCustomRequest can return immediately, emplace relevant promises before 
-		if (Params.IsStreamChangeEmpty())
-		{
-			const FSubmitStreamChangesResponse StreamResponse { EStreamSubmissionErrorCode::NoChange };
-			OperationResult->EmplaceStreamPromise(StreamResponse);
-			StreamRequestCompletedDelegate.Broadcast(StreamResponse);
-		}
-		// Opposed to the local case, sending authority request is sent instantly in the remote case, so emplace the promise now.
-		OperationResult->EmplaceAuthorityRequestPromise({ EAuthoritySubmissionRequestErrorCode::Success, Params.AuthorityRequest });
-		
+		// Since SendCustomRequest can return immediately, emplace relevant promises before
+		EarlyCompletePromisesForUnchangedData(Params, OperationResult);
 		const FMultiUser_ChangeRemote_Request Request
 		{
 			Params.StreamRequest ? *Params.StreamRequest : FConcertReplication_ChangeStream_Request{},
@@ -91,7 +83,6 @@ namespace UE::MultiUserClient
 			HandleFailureResponse(Response.Error);
 		}
 
-		//
 		UE_LOG(LogConcert, Log, TEXT("Remote client %s accepted change request"), *GetRemoteClientName());
 	}
 	
@@ -99,27 +90,37 @@ namespace UE::MultiUserClient
 	{
 		UE_LOG(LogConcert, Warning, TEXT("Change request to client %s timed out"), *GetRemoteClientName());
 		
-		// SubmitChanges sets the
+		TimeoutStreamChangeIfUnset();
+		TimeoutAuthorityChangeIfUnset();
+		CleanUpSubmissionProcess();
+	}
+
+	void FSubmissionWorkflow_RemoteClient::TimeoutStreamChangeIfUnset()
+	{
 		if (!InProgressOperation->Parameters.IsStreamChangeEmpty())
 		{
 			const FSubmitStreamChangesResponse StreamResponse { EStreamSubmissionErrorCode::Timeout };
 			GetOperation().EmplaceStreamPromise(StreamResponse);
 			StreamRequestCompletedDelegate.Broadcast(StreamResponse);
 		}
+	}
 
-		// Sending authority request was successful.
-		const FSubmitAuthorityChangesRequest AuthorityRequest { EAuthoritySubmissionRequestErrorCode::Success, InProgressOperation->Parameters.AuthorityRequest };
-		const FSubmitAuthorityChangesResponse AuthorityResponse { EAuthoritySubmissionResponseErrorCode::Timeout };
-		GetOperation().EmplaceAuthorityResponsePromise(AuthorityResponse);
-		AuthorityRequestCompletedDelegate.Broadcast(AuthorityRequest, AuthorityResponse);
-			
-		CleanUpSubmissionProcess();
+	void FSubmissionWorkflow_RemoteClient::TimeoutAuthorityChangeIfUnset()
+	{
+		if (!InProgressOperation->ExposedOperation->HasSetAuthorityResponsePromise())
+		{
+			// Sending authority request was technically successful.
+			const FSubmitAuthorityChangesRequest AuthorityRequest { EAuthoritySubmissionRequestErrorCode::Success, InProgressOperation->Parameters.AuthorityRequest };
+			const FSubmitAuthorityChangesResponse AuthorityResponse { EAuthoritySubmissionResponseErrorCode::Timeout };
+			GetOperation().EmplaceAuthorityResponsePromise(AuthorityResponse);
+			AuthorityRequestCompletedDelegate.Broadcast(AuthorityRequest, AuthorityResponse);
+		}
 	}
 
 	void FSubmissionWorkflow_RemoteClient::HandleFailureResponse(EMultiUserChangeRemoteRequestError ErrorFlags)
 	{
 		UE_LOG(LogConcert, Warning, TEXT("Change request to client %s failed (error flags: %d)"), *GetRemoteClientName(), static_cast<int32>(ErrorFlags));
-
+		
 		switch (ErrorFlags)
 		{
 		case EMultiUserChangeRemoteRequestError::PredictedConflict:
@@ -127,6 +128,8 @@ namespace UE::MultiUserClient
 		case EMultiUserChangeRemoteRequestError::OtherInProgress:
 			// TODO DP UE-200924: Retry when client is busy serving another request
 		case EMultiUserChangeRemoteRequestError::RejectedChange:
+			TimeoutStreamChangeIfUnset();
+			TimeoutAuthorityChangeIfUnset();
 			CleanUpSubmissionProcess();
 			break;
 			
@@ -147,6 +150,32 @@ namespace UE::MultiUserClient
 			OnSubmitOperationCompletedDelegate.Broadcast();
 		}
 	}
+	
+	void FSubmissionWorkflow_RemoteClient::EarlyCompletePromisesForUnchangedData(const FSubmissionParams& Params, const TSharedRef<FSingleClientSubmissionOperation>& OperationResult) const
+	{
+		// Complete stream event if nothing is being changed
+		if (Params.IsStreamChangeEmpty())
+		{
+			const FSubmitStreamChangesResponse StreamResponse { EStreamSubmissionErrorCode::NoChange };
+			OperationResult->EmplaceStreamPromise(StreamResponse);
+			StreamRequestCompletedDelegate.Broadcast(StreamResponse);
+		}
+
+		// Complete the authority promise if nothing is being changed
+		if (Params.IsAuthorityChangeEmpty())
+		{
+			FSubmitAuthorityChangesRequest RequestEvent { EAuthoritySubmissionRequestErrorCode::NoChange, Params.AuthorityRequest };
+			FSubmitAuthorityChangesResponse ResponseEvent { EAuthoritySubmissionResponseErrorCode::NoChange };
+			OperationResult->EmplaceAuthorityRequestPromise(RequestEvent);
+			OperationResult->EmplaceAuthorityResponsePromise(ResponseEvent);
+			AuthorityRequestCompletedDelegate.Broadcast(RequestEvent, ResponseEvent);
+		}
+		else
+		{
+			// Opposed to the local case, sending authority request is sent instantly in the remote case, so emplace the promise now.
+			OperationResult->EmplaceAuthorityRequestPromise({ EAuthoritySubmissionRequestErrorCode::Success, Params.AuthorityRequest });
+		}
+	}
 
 	void FSubmissionWorkflow_RemoteClient::OnStreamRemoteChangeEvent(const FConcertSessionContext& Context, const FMultiUser_ChangeRemote_StreamUpdatedEvent& EventData)
 	{
diff --git a/Engine/Plugins/Developer/Concert/ConcertApp/MultiUserClient/Source/MultiUserClient/Private/Replication/Submission/Remote/SubmissionWorkflow_RemoteClient.h b/Engine/Plugins/Developer/Concert/ConcertApp/MultiUserClient/Source/MultiUserClient/Private/Replication/Submission/Remote/SubmissionWorkflow_RemoteClient.h
--- a/Engine/Plugins/Developer/Concert/ConcertApp/MultiUserClient/Source/MultiUserClient/Private/Replication/Submission/Remote/SubmissionWorkflow_RemoteClient.h
+++ b/Engine/Plugins/Developer/Concert/ConcertApp/MultiUserClient/Source/MultiUserClient/Private/Replication/Submission/Remote/SubmissionWorkflow_RemoteClient.h
@@ -63,6 +63,9 @@ namespace UE::MultiUserClient
 		
 		/** Handles timeout of submission response. */
 		void TimeoutAndClean();
+		void TimeoutStreamChangeIfUnset();
+		void TimeoutAuthorityChangeIfUnset();
+		
 		/** Handles a response from the client that indicates failure. */
 		void HandleFailureResponse(EMultiUserChangeRemoteRequestError ErrorFlags);
 		/** Cleans up side effects from last submission process */
@@ -70,6 +73,9 @@ namespace UE::MultiUserClient
 
 		FSingleClientSubmissionOperation& GetOperation() { return *InProgressOperation->ExposedOperation; }
 
+		/** Checks which data is not being changed by Params and completes the corresponding promises on OperationResult. */
+		void EarlyCompletePromisesForUnchangedData(const FSubmissionParams& Params, const TSharedRef<FSingleClientSubmissionOperation>& OperationResult) const;
+		
 		/** Called by remote client to update us about progress. */
 		void OnStreamRemoteChangeEvent(const FConcertSessionContext& Context, const FMultiUser_ChangeRemote_StreamUpdatedEvent& EventData);
 		void OnAuthorityRemoteChangeEvent(const FConcertSessionContext& Context, const FMultiUser_ChangeRemote_AuthorityUpdatedEvent& EventData);
diff --git a/Engine/Plugins/Developer/Concert/ConcertApp/MultiUserClient/Source/MultiUserClient/Private/Replication/Submission/SingleClientSubmissionOperation.h b/Engine/Plugins/Developer/Concert/ConcertApp/MultiUserClient/Source/MultiUserClient/Private/Replication/Submission/SingleClientSubmissionOperation.h
--- a/Engine/Plugins/Developer/Concert/ConcertApp/MultiUserClient/Source/MultiUserClient/Private/Replication/Submission/SingleClientSubmissionOperation.h
+++ b/Engine/Plugins/Developer/Concert/ConcertApp/MultiUserClient/Source/MultiUserClient/Private/Replication/Submission/SingleClientSubmissionOperation.h
@@ -2,8 +2,8 @@
 
 #pragma once
 
+#include "ConcertLogGlobal.h"
 #include "ISubmissionOperation.h"
-#include "Replication/Stream/IClientStreamSynchronizer.h"
 
 #include "Templates/UnrealTemplate.h"
 
@@ -44,21 +44,43 @@ namespace UE::MultiUserClient
 		{
 			bStreamPromiseWasSet = true;
 			StreamChangesPromise.EmplaceValue(MoveTemp(Result));
+			
+			UE_LOG(LogConcert, Log, TEXT("Submission operation: Completed stream request stage."));
 		}
 		void EmplaceAuthorityRequestPromise(FSubmitAuthorityChangesRequest Result)
 		{
 			bAuthorityRequestPromiseWasSet = true;
 			AuthorityChangeRequestPromise.EmplaceValue(MoveTemp(Result));
+			
+			UE_LOG(LogConcert, Log, TEXT("Submission operation: Complete authority request stage."));
 		}
 		void EmplaceAuthorityResponsePromise(FSubmitAuthorityChangesResponse Result)
 		{
 			bAuthorityResponsePromiseWasSet = true;
 			AuthorityChangeResponsePromise.EmplaceValue(MoveTemp(Result));
+			
+			UE_LOG(LogConcert, Log, TEXT("Submission operation: Completed authority response stage."));
 		}
 		void EmplaceCompleteOperationPromise(ESubmissionOperationCompletedCode Result)
 		{
+			// The following enforces the API contract that CompleteOperationPromise finishes as last future.
+			if (!ensureMsgf(bStreamPromiseWasSet, TEXT("All other stages should be explicitly completed before calling EmplaceCompleteOperationPromise. Check implementation logic.")))
+			{
+				EmplaceStreamPromise(FSubmitStreamChangesResponse{ EStreamSubmissionErrorCode::Cancelled });
+			}
+			if (!ensureMsgf(bAuthorityRequestPromiseWasSet, TEXT("All other stages should be explicitly completed before calling EmplaceCompleteOperationPromise. Check implementation logic.")))
+			{
+				EmplaceAuthorityRequestPromise(FSubmitAuthorityChangesRequest{ EAuthoritySubmissionRequestErrorCode::Cancelled });
+			}
+			if (!ensureMsgf(bAuthorityResponsePromiseWasSet, TEXT("All other stages should be explicitly completed before calling EmplaceCompleteOperationPromise. Check implementation logic.")))
+			{
+				EmplaceAuthorityResponsePromise(FSubmitAuthorityChangesResponse{ EAuthoritySubmissionResponseErrorCode::Cancelled });
+			}
+			
 			bCompleteOperationPromiseWasSet = true;
 			CompleteOperationPromise.EmplaceValue(Result);
+			
+			UE_LOG(LogConcert, Log, TEXT("Submission operation: Completed operation."));
 		}
 
 		bool HasSetStreamPromise() const { return bStreamPromiseWasSet; }
diff --git a/Engine/Plugins/Developer/Concert/ConcertApp/MultiUserClient/Source/MultiUserClient/Private/Replication/Submission/SubmissionWorkflow_LocalClient.cpp b/Engine/Plugins/Developer/Concert/ConcertApp/MultiUserClient/Source/MultiUserClient/Private/Replication/Submission/SubmissionWorkflow_LocalClient.cpp
--- a/Engine/Plugins/Developer/Concert/ConcertApp/MultiUserClient/Source/MultiUserClient/Private/Replication/Submission/SubmissionWorkflow_LocalClient.cpp
+++ b/Engine/Plugins/Developer/Concert/ConcertApp/MultiUserClient/Source/MultiUserClient/Private/Replication/Submission/SubmissionWorkflow_LocalClient.cpp
@@ -90,13 +90,7 @@ namespace UE::MultiUserClient
 		}
 		else
 		{
-			const FSubmitAuthorityChangesRequest Request{ EAuthoritySubmissionRequestErrorCode::CancelledDueToStreamUpdate };
-			const FSubmitAuthorityChangesResponse Response{ EAuthoritySubmissionResponseErrorCode::CancelledDueToStreamUpdate };
-			
-			Operation->EmplaceAuthorityRequestPromise(Request);
-			Operation->EmplaceAuthorityResponsePromise(Response);
-			AuthorityRequestCompletedDelegate.Broadcast(Request, Response);
-			
+			SkipAuthorityStage(EAuthoritySubmissionRequestErrorCode::CancelledDueToStreamUpdate, EAuthoritySubmissionResponseErrorCode::CancelledDueToStreamUpdate);
 			CleanUpSubmissionOperation();
 		}
 	}
@@ -107,8 +101,8 @@ namespace UE::MultiUserClient
 		IConcertClientReplicationManager* ReplicationManager = Client->GetReplicationManager();
 		if (!ensure(ReplicationManager))
 		{
-			AuthorityRequestCompletedDelegate.Broadcast({ EAuthoritySubmissionRequestErrorCode::Cancelled }, { EAuthoritySubmissionResponseErrorCode::Cancelled });
-			CleanUpSubmissionOperation(); // Automatically cancels the pending promises
+			SkipAuthorityStage(EAuthoritySubmissionRequestErrorCode::Cancelled, EAuthoritySubmissionResponseErrorCode::Cancelled);
+			CleanUpSubmissionOperation();
 			return;
 		}
 
@@ -116,12 +110,7 @@ namespace UE::MultiUserClient
 		const bool bHasNoChanges = !AuthorityChangeRequest || AuthorityChangeRequest->IsEmpty();
 		if (bHasNoChanges)
 		{
-			const FSubmitAuthorityChangesRequest Request{ EAuthoritySubmissionRequestErrorCode::NoChange };
-			const FSubmitAuthorityChangesResponse Response{ EAuthoritySubmissionResponseErrorCode::NoChange };
-			
-			Operation->EmplaceAuthorityRequestPromise(Request);
-			Operation->EmplaceAuthorityResponsePromise(Response);
-			AuthorityRequestCompletedDelegate.Broadcast(Request, Response);
+			SkipAuthorityStage(EAuthoritySubmissionRequestErrorCode::NoChange, EAuthoritySubmissionResponseErrorCode::NoChange);
 			CleanUpSubmissionOperation();
 			return;
 		}
@@ -144,6 +133,17 @@ namespace UE::MultiUserClient
 				}
 			});
 	}
+	
+	void FSubmissionWorkflow_LocalClient::SkipAuthorityStage(EAuthoritySubmissionRequestErrorCode RequestCode, EAuthoritySubmissionResponseErrorCode ResponseCode)
+	{
+		const TSharedRef<FSingleClientSubmissionOperation>& Operation = InProgressOperation->Operation;
+		const FSubmitAuthorityChangesRequest Request{ RequestCode };
+		const FSubmitAuthorityChangesResponse Response{ ResponseCode };
+			
+		Operation->EmplaceAuthorityRequestPromise(Request);
+		Operation->EmplaceAuthorityResponsePromise(Response);
+		AuthorityRequestCompletedDelegate.Broadcast(Request, Response);
+	}
 
 	void FSubmissionWorkflow_LocalClient::CleanUpSubmissionOperation()
 	{
diff --git a/Engine/Plugins/Developer/Concert/ConcertApp/MultiUserClient/Source/MultiUserClient/Private/Replication/Submission/SubmissionWorkflow_LocalClient.h b/Engine/Plugins/Developer/Concert/ConcertApp/MultiUserClient/Source/MultiUserClient/Private/Replication/Submission/SubmissionWorkflow_LocalClient.h
--- a/Engine/Plugins/Developer/Concert/ConcertApp/MultiUserClient/Source/MultiUserClient/Private/Replication/Submission/SubmissionWorkflow_LocalClient.h
+++ b/Engine/Plugins/Developer/Concert/ConcertApp/MultiUserClient/Source/MultiUserClient/Private/Replication/Submission/SubmissionWorkflow_LocalClient.h
@@ -58,7 +58,8 @@ namespace UE::MultiUserClient
 			TOptional<FConcertReplication_ChangeAuthority_Request> AuthorityChangeRequest
 			);
 		void HandlePendingAuthorityChangeRequest(TOptional<FConcertReplication_ChangeAuthority_Request> AuthorityChangeRequest);
-
+		
+		void SkipAuthorityStage(EAuthoritySubmissionRequestErrorCode RequestCode, EAuthoritySubmissionResponseErrorCode ResponseCode);
 		void CleanUpSubmissionOperation();
 	};
 }
