Add events to UMultiUserReplicationSubsystem for when a client's known server state for streams or authority changed:
- OnClientStreamServerStateChanged
- OnClientAuthorityServerStateChanged

diff --git a/Engine/Plugins/Developer/Concert/ConcertApp/MultiUserClient/Source/MultiUserClient/Private/Replication/MultiUserReplicationManager.cpp b/Engine/Plugins/Developer/Concert/ConcertApp/MultiUserClient/Source/MultiUserClient/Private/Replication/MultiUserReplicationManager.cpp
--- a/Engine/Plugins/Developer/Concert/ConcertApp/MultiUserClient/Source/MultiUserClient/Private/Replication/MultiUserReplicationManager.cpp
+++ b/Engine/Plugins/Developer/Concert/ConcertApp/MultiUserClient/Source/MultiUserClient/Private/Replication/MultiUserReplicationManager.cpp
@@ -86,6 +86,7 @@ namespace UE::MultiUserClient
 		if (bSuccess)
 		{
 			ConnectedState.Emplace(Client, DiscoveryContainer);
+			SetupClientConnectionEvents();
 			SetConnectionStateAndBroadcast(EMultiUserReplicationConnectionState::Connected);
 		}
 		else
@@ -100,6 +101,29 @@ namespace UE::MultiUserClient
 		OnReplicationConnectionStateChangedDelegate.Broadcast(ConnectionState);
 	}
 
+	void FMultiUserReplicationManager::SetupClientConnectionEvents()
+	{
+		FReplicationClientManager& ClientManager = ConnectedState->ClientManager;
+		SetupClientDelegates(ClientManager.GetLocalClient());
+		ClientManager.OnPostRemoteClientAdded().AddRaw(this, &FMultiUserReplicationManager::OnReplicationClientConnected);
+	}
+
+	void FMultiUserReplicationManager::OnClientStreamServerStateChanged(const FGuid EndpointId) const
+	{
+		OnStreamServerStateChangedDelegate.Broadcast(EndpointId);
+	}
+
+	void FMultiUserReplicationManager::OnClientAuthorityServerStateChanged(const FGuid EndpointId) const
+	{
+		OnAuthorityServerStateChangedDelegate.Broadcast(EndpointId);
+	}
+
+	void FMultiUserReplicationManager::SetupClientDelegates(FReplicationClient& InClient) const
+	{
+		InClient.GetStreamSynchronizer().OnServerStateChanged().AddRaw(this, &FMultiUserReplicationManager::OnClientStreamServerStateChanged, InClient.GetEndpointId());
+		InClient.GetStreamSynchronizer().OnServerStateChanged().AddRaw(this, &FMultiUserReplicationManager::OnClientStreamServerStateChanged, InClient.GetEndpointId());
+	}
+
 	const FConcertObjectReplicationMap* FMultiUserReplicationManager::FindReplicationMapForClient(const FGuid& ClientId) const
 	{
 		if (ConnectedState && ensureMsgf(IsInGameThread(), TEXT("To simplify implementation, only calls from game thread are allowed.")))
diff --git a/Engine/Plugins/Developer/Concert/ConcertApp/MultiUserClient/Source/MultiUserClient/Private/Replication/MultiUserReplicationManager.h b/Engine/Plugins/Developer/Concert/ConcertApp/MultiUserClient/Source/MultiUserClient/Private/Replication/MultiUserReplicationManager.h
--- a/Engine/Plugins/Developer/Concert/ConcertApp/MultiUserClient/Source/MultiUserClient/Private/Replication/MultiUserReplicationManager.h
+++ b/Engine/Plugins/Developer/Concert/ConcertApp/MultiUserClient/Source/MultiUserClient/Private/Replication/MultiUserReplicationManager.h
@@ -27,6 +27,8 @@ namespace UE::ConcertSharedSlate
 
 namespace UE::MultiUserClient
 {
+	class FReplicationClient;
+	
 	enum class EMultiUserReplicationConnectionState : uint8
 	{
 		Connecting,
@@ -75,6 +77,8 @@ namespace UE::MultiUserClient
 		virtual void RegisterReplicationDiscoverer(TSharedRef<IReplicationDiscoverer> Discoverer) override;
 		virtual void RemoveReplicationDiscoverer(const TSharedRef<IReplicationDiscoverer>& Discoverer) override;
 		virtual TSharedRef<IClientChangeOperation> EnqueueChanges(const FGuid& ClientId, TAttribute<FChangeClientReplicationRequest> SubmissionParams) override;
+		virtual FOnServerStateChanged& OnStreamServerStateChanged() override { return OnStreamServerStateChangedDelegate; }
+		virtual FOnServerStateChanged& OnAuthorityServerStateChanged() override { return OnAuthorityServerStateChangedDelegate; }
 		//~ End IMultiUserReplication Interface
 
 	private:
@@ -108,6 +112,11 @@ namespace UE::MultiUserClient
 
 		/** Called when ConnectionState changes. */
 		FOnReplicationConnectionStateChanged OnReplicationConnectionStateChangedDelegate;
+		
+		/** Triggers when a client's known server state has changed. */
+		FOnServerStateChanged OnStreamServerStateChangedDelegate;
+		/** Triggers when a client's known server state has changed. */
+		FOnServerStateChanged OnAuthorityServerStateChangedDelegate;
 
 		/** Callback into Concert for when client connection has changed. */
 		void OnSessionConnectionChanged(IConcertClientSession& ConcertClientSession, EConcertConnectionStatus ConcertConnectionStatus);
@@ -116,8 +125,15 @@ namespace UE::MultiUserClient
 
 		/** Handles server response for joining replication session */
 		void HandleReplicationSessionJoined(const ConcertSyncClient::Replication::FJoinReplicatedSessionResult& JoinSessionResult);
-		
+		/** Sets the current connection state and triggers OnReplicationConnectionStateChangedDelegate. */
 		void SetConnectionStateAndBroadcast(EMultiUserReplicationConnectionState NewState);
+
+		/** Sets up delegates for implementing the broadcasting of OnStreamServerStateChangedDelegate and OnAuthorityServerStateChangedDelegate. */
+		void SetupClientConnectionEvents();
+		void OnClientStreamServerStateChanged(const FGuid EndpointId) const;
+		void OnClientAuthorityServerStateChanged(const FGuid EndpointId) const;
+		void OnReplicationClientConnected(FRemoteReplicationClient& RemoteClient) const { SetupClientDelegates(RemoteClient); }
+		void SetupClientDelegates(FReplicationClient& InClient) const;
 	};
 }
 
diff --git a/Engine/Plugins/Developer/Concert/ConcertApp/MultiUserClient/Source/MultiUserClient/Public/Replication/IMultiUserReplication.h b/Engine/Plugins/Developer/Concert/ConcertApp/MultiUserClient/Source/MultiUserClient/Public/Replication/IMultiUserReplication.h
--- a/Engine/Plugins/Developer/Concert/ConcertApp/MultiUserClient/Source/MultiUserClient/Public/Replication/IMultiUserReplication.h
+++ b/Engine/Plugins/Developer/Concert/ConcertApp/MultiUserClient/Source/MultiUserClient/Public/Replication/IMultiUserReplication.h
@@ -70,6 +70,12 @@ namespace UE::MultiUserClient
 		 */
 		virtual TSharedRef<IClientChangeOperation> EnqueueChanges(const FGuid& ClientId, TAttribute<FChangeClientReplicationRequest> SubmissionParams) = 0;
 
+		DECLARE_MULTICAST_DELEGATE_OneParam(FOnServerStateChanged, const FGuid& /*EndpointId*/);
+		/** @return Delegate that triggers when the given client's known server state has changed. */
+		virtual FOnServerStateChanged& OnStreamServerStateChanged() = 0;
+		/** @return Delegate that triggers when the given client's known server state has changed. */
+		virtual FOnServerStateChanged& OnAuthorityServerStateChanged() = 0;
+
 		virtual ~IMultiUserReplication() = default;
 	};
 }
diff --git a/Engine/Plugins/Developer/Concert/ConcertApp/MultiUserClient/Source/MultiUserClientLibrary/Private/Replication/MultiUserReplicationSubsystem.cpp b/Engine/Plugins/Developer/Concert/ConcertApp/MultiUserClient/Source/MultiUserClientLibrary/Private/Replication/MultiUserReplicationSubsystem.cpp
--- a/Engine/Plugins/Developer/Concert/ConcertApp/MultiUserClient/Source/MultiUserClientLibrary/Private/Replication/MultiUserReplicationSubsystem.cpp
+++ b/Engine/Plugins/Developer/Concert/ConcertApp/MultiUserClient/Source/MultiUserClientLibrary/Private/Replication/MultiUserReplicationSubsystem.cpp
@@ -134,6 +134,9 @@ void UMultiUserReplicationSubsystem::Initialize(FSubsystemCollectionBase& Collec
 	{
 		UObjectAdapter = MakeShared<UE::MultiUserClientLibrary::FUObjectAdapterReplicationDiscoverer>();
 		ReplicationInterface->RegisterReplicationDiscoverer(UObjectAdapter.ToSharedRef());
+
+		ReplicationInterface->OnStreamServerStateChanged().AddUObject(this, &UMultiUserReplicationSubsystem::OnClientStreamsChanged);
+		ReplicationInterface->OnAuthorityServerStateChanged().AddUObject(this, &UMultiUserReplicationSubsystem::OnClientStreamsChanged);
 	}
 #endif
 }
@@ -150,6 +153,9 @@ void UMultiUserReplicationSubsystem::Deinitialize()
 		{
 			ReplicationInterface->RemoveReplicationDiscoverer(UObjectAdapter.ToSharedRef());
 			UObjectAdapter.Reset();
+			
+			ReplicationInterface->OnStreamServerStateChanged().RemoveAll(this);
+			ReplicationInterface->OnAuthorityServerStateChanged().RemoveAll(this);
 		}
 	}
 #endif
diff --git a/Engine/Plugins/Developer/Concert/ConcertApp/MultiUserClient/Source/MultiUserClientLibrary/Public/Replication/MultiUserReplicationSubsystem.h b/Engine/Plugins/Developer/Concert/ConcertApp/MultiUserClient/Source/MultiUserClientLibrary/Public/Replication/MultiUserReplicationSubsystem.h
--- a/Engine/Plugins/Developer/Concert/ConcertApp/MultiUserClient/Source/MultiUserClientLibrary/Public/Replication/MultiUserReplicationSubsystem.h
+++ b/Engine/Plugins/Developer/Concert/ConcertApp/MultiUserClient/Source/MultiUserClientLibrary/Public/Replication/MultiUserReplicationSubsystem.h
@@ -19,6 +19,8 @@ public:
 	
 	// This would be the right place to expose additional MU specific replication functions in the future
 	
+	DECLARE_DYNAMIC_MULTICAST_DELEGATE_OneParam(FOnServerStateChanged, const FGuid&, EndpointId);
+	
 	/**
 	 * @return Whether the client is replicating the object.
 	 * @note An object can be registered but not replicated.
@@ -78,6 +80,18 @@ public:
 	//~ End USubsystem Interface
 
 private:
+	
+	/**
+	 * Event triggered when the following changes about a client:
+	 * - The registered object to properties bindings
+	 * - The registered replication frequency setting of an object
+	 */
+	UPROPERTY(BlueprintAssignable, Category = "Multi-user")
+	FOnServerStateChanged OnClientStreamServerStateChanged;
+	
+	/** Event triggered a client changes the objects it is replicating. */
+	UPROPERTY(BlueprintAssignable, Category = "Multi-user")
+	FOnServerStateChanged OnClientAuthorityServerStateChanged;
 
 	/**
 	 * This is used only for when the adds an object through the Add button in the UI.
@@ -88,4 +102,7 @@ private:
 	 * Registered when this subsystem is initialized.
 	 */
 	TSharedPtr<UE::MultiUserClientLibrary::FUObjectAdapterReplicationDiscoverer> UObjectAdapter;
+
+	void OnClientStreamsChanged(const FGuid& EndpointId) const { OnClientStreamServerStateChanged.Broadcast(EndpointId); }
+	void OnClientAuthorityChanged(const FGuid& EndpointId) const { OnClientAuthorityServerStateChanged.Broadcast(EndpointId); }
 };
