Fixed authority changed event being improperly set up and thus never being called.

diff --git a/Engine/Plugins/Developer/Concert/ConcertApp/MultiUserClient/Source/MultiUserClient/Private/Replication/MultiUserReplicationManager.cpp b/Engine/Plugins/Developer/Concert/ConcertApp/MultiUserClient/Source/MultiUserClient/Private/Replication/MultiUserReplicationManager.cpp
--- a/Engine/Plugins/Developer/Concert/ConcertApp/MultiUserClient/Source/MultiUserClient/Private/Replication/MultiUserReplicationManager.cpp
+++ b/Engine/Plugins/Developer/Concert/ConcertApp/MultiUserClient/Source/MultiUserClient/Private/Replication/MultiUserReplicationManager.cpp
@@ -121,7 +121,7 @@ namespace UE::MultiUserClient
 	void FMultiUserReplicationManager::SetupClientDelegates(FReplicationClient& InClient) const
 	{
 		InClient.GetStreamSynchronizer().OnServerStateChanged().AddRaw(this, &FMultiUserReplicationManager::OnClientStreamServerStateChanged, InClient.GetEndpointId());
-		InClient.GetStreamSynchronizer().OnServerStateChanged().AddRaw(this, &FMultiUserReplicationManager::OnClientStreamServerStateChanged, InClient.GetEndpointId());
+		InClient.GetAuthoritySynchronizer().OnServerStateChanged().AddRaw(this, &FMultiUserReplicationManager::OnClientAuthorityServerStateChanged, InClient.GetEndpointId());
 	}
 
 	const FConcertObjectReplicationMap* FMultiUserReplicationManager::FindReplicationMapForClient(const FGuid& ClientId) const
