Fix BP replication API reporting request as successful even though it was rejected

diff --git a/Engine/Plugins/Developer/Concert/ConcertApp/MultiUserClient/Source/MultiUserClient/Private/Replication/Submission/External/ClientChangeConversionUtils.cpp b/Engine/Plugins/Developer/Concert/ConcertApp/MultiUserClient/Source/MultiUserClient/Private/Replication/Submission/External/ClientChangeConversionUtils.cpp
--- a/Engine/Plugins/Developer/Concert/ConcertApp/MultiUserClient/Source/MultiUserClient/Private/Replication/Submission/External/ClientChangeConversionUtils.cpp
+++ b/Engine/Plugins/Developer/Concert/ConcertApp/MultiUserClient/Source/MultiUserClient/Private/Replication/Submission/External/ClientChangeConversionUtils.cpp
@@ -158,7 +158,11 @@ namespace UE::MultiUserClient::ClientChangeConversionUtils
 	{
 		switch (Response.ErrorCode)
 		{
-		case EStreamSubmissionErrorCode::Success: return EChangeStreamOperationResult::Success;
+		case EStreamSubmissionErrorCode::Success:
+			// EStreamSubmissionErrorCode::Success the op was processed but the EChangeStreamOperationResult distinguishes between successful or rejected ops.
+			return !ensureMsgf(Response.SubmissionInfo.IsSet(), TEXT("Success implies set response. Investigate.")) || Response.SubmissionInfo->Response.IsFailure()
+				? EChangeStreamOperationResult::Rejected
+				: EChangeStreamOperationResult::Success;
 		case EStreamSubmissionErrorCode::NoChange: return EChangeStreamOperationResult::NoChanges; 
 		case EStreamSubmissionErrorCode::Timeout: return EChangeStreamOperationResult::Timeout; 
 		case EStreamSubmissionErrorCode::Cancelled: return EChangeStreamOperationResult::Cancelled;
@@ -172,7 +176,11 @@ namespace UE::MultiUserClient::ClientChangeConversionUtils
 	{
 		switch (Response.ErrorCode)
 		{
-		case EAuthoritySubmissionResponseErrorCode::Success: return EChangeAuthorityOperationResult::Success;
+		case EAuthoritySubmissionResponseErrorCode::Success:
+			// EAuthoritySubmissionResponseErrorCode::Success the op was processed but the EChangeStreamOperationResult distinguishes between successful or rejected ops.
+			return !ensureMsgf(Response.Response.IsSet(), TEXT("Success implies set response. Investigate.")) || !Response.Response->RejectedObjects.IsEmpty()
+				? EChangeAuthorityOperationResult::RejectedFullyOrPartially
+				: EChangeAuthorityOperationResult::Success;
 		case EAuthoritySubmissionResponseErrorCode::NoChange: return EChangeAuthorityOperationResult::NoChanges; 
 		case EAuthoritySubmissionResponseErrorCode::Timeout: return EChangeAuthorityOperationResult::Timeout; 
 		case EAuthoritySubmissionResponseErrorCode::CancelledDueToStreamUpdate: return EChangeAuthorityOperationResult::CancelledDueToStreamUpdate; 
diff --git a/Engine/Plugins/Developer/Concert/ConcertApp/MultiUserClient/Source/MultiUserClient/Public/Replication/ChangeOperationTypes.h b/Engine/Plugins/Developer/Concert/ConcertApp/MultiUserClient/Source/MultiUserClient/Public/Replication/ChangeOperationTypes.h
--- a/Engine/Plugins/Developer/Concert/ConcertApp/MultiUserClient/Source/MultiUserClient/Public/Replication/ChangeOperationTypes.h
+++ b/Engine/Plugins/Developer/Concert/ConcertApp/MultiUserClient/Source/MultiUserClient/Public/Replication/ChangeOperationTypes.h
@@ -63,7 +63,7 @@ namespace UE::MultiUserClient
 		// Realistic failures
 
 		/** The changes were rejected by the server. */
-		Rejected,
+		RejectedFullyOrPartially,
 		/** Failed because the request timed out. */
 		Timeout,
 		/** The authority change was not submitted because the stream change was unsuccessful. */
diff --git a/Engine/Plugins/Developer/Concert/ConcertApp/MultiUserClient/Source/MultiUserClientLibrary/Public/Replication/Async/ChangeClientBlueprintParams.h b/Engine/Plugins/Developer/Concert/ConcertApp/MultiUserClient/Source/MultiUserClientLibrary/Public/Replication/Async/ChangeClientBlueprintParams.h
--- a/Engine/Plugins/Developer/Concert/ConcertApp/MultiUserClient/Source/MultiUserClientLibrary/Public/Replication/Async/ChangeClientBlueprintParams.h
+++ b/Engine/Plugins/Developer/Concert/ConcertApp/MultiUserClient/Source/MultiUserClientLibrary/Public/Replication/Async/ChangeClientBlueprintParams.h
@@ -55,8 +55,8 @@ enum class EMultiUserChangeAuthorityOperationResult : uint8
 
 	// Realistic failures
 
-	/** The changes were rejected by the server. */
-	Rejected,
+	/** Some of the authority changes were rejected by the server. */
+	RejectedFullyOrPartially,
 	/** Failed because the request timed out. */
 	Timeout,
 	/** The authority change was not submitted because the stream change was unsuccessful. */
