Remove "Experimental Click Through" on Replication tab launch

diff --git a/Engine/Plugins/Developer/Concert/ConcertApp/MultiUserClient/Source/MultiUserClient/Private/Widgets/ActiveSession/Replication/SReplicationTabWithWarningOverlay.cpp b/Engine/Plugins/Developer/Concert/ConcertApp/MultiUserClient/Source/MultiUserClient/Private/Widgets/ActiveSession/Replication/SReplicationTabWithWarningOverlay.cpp
deleted file mode 100644
diff --git a/Engine/Plugins/Developer/Concert/ConcertApp/MultiUserClient/Source/MultiUserClient/Private/Widgets/ActiveSession/Replication/SReplicationTabWithWarningOverlay.h b/Engine/Plugins/Developer/Concert/ConcertApp/MultiUserClient/Source/MultiUserClient/Private/Widgets/ActiveSession/Replication/SReplicationTabWithWarningOverlay.h
deleted file mode 100644
diff --git a/Engine/Plugins/Developer/Concert/ConcertApp/MultiUserClient/Source/MultiUserClient/Private/Widgets/ActiveSession/SActiveSessionRoot.cpp b/Engine/Plugins/Developer/Concert/ConcertApp/MultiUserClient/Source/MultiUserClient/Private/Widgets/ActiveSession/SActiveSessionRoot.cpp
--- a/Engine/Plugins/Developer/Concert/ConcertApp/MultiUserClient/Source/MultiUserClient/Private/Widgets/ActiveSession/SActiveSessionRoot.cpp
+++ b/Engine/Plugins/Developer/Concert/ConcertApp/MultiUserClient/Source/MultiUserClient/Private/Widgets/ActiveSession/SActiveSessionRoot.cpp
@@ -3,7 +3,7 @@
 #include "SActiveSessionRoot.h"
 
 #include "Overview/SActiveSessionOverviewTab.h"
-#include "Replication/SReplicationTabWithWarningOverlay.h"
+#include "Replication/SReplicationRootWidget.h"
 #include "SActiveSessionToolbar.h"
 
 #include "Framework/Docking/TabManager.h"
@@ -107,7 +107,7 @@ namespace UE::MultiUserClient
 			.Label(LOCTEXT("ReplicationTab.Label", "Replication"))
 			.ToolTipText(LOCTEXT("ReplicationTab.Tooltip", "Manage real-time object replication"))
 			[
-				SNew(SReplicationTabWithWarningOverlay, InReplicationManager, ConcertSyncClient.ToSharedRef())
+				SNew(SReplicationRootWidget, InReplicationManager, ConcertSyncClient.ToSharedRef())
 			];
 	}
 }
