Set Default Replication on Rate of Cine Cameras to 60FPS

diff --git a/Engine/Plugins/Developer/Concert/ConcertApp/MultiUserClient/Config/BaseMultiUserClient.ini b/Engine/Plugins/Developer/Concert/ConcertApp/MultiUserClient/Config/BaseMultiUserClient.ini
--- a/Engine/Plugins/Developer/Concert/ConcertApp/MultiUserClient/Config/BaseMultiUserClient.ini
+++ b/Engine/Plugins/Developer/Concert/ConcertApp/MultiUserClient/Config/BaseMultiUserClient.ini
@@ -1,2 +1,3 @@
 [/Script/MultiUserClient.MultiUserReplicationSettings]
-ReplicationEditorSettings=(DefaultPropertySelection=(("/Script/Engine.SceneComponent", (DefaultSelectedProperties=("RelativeLocation","RelativeLocation.X","RelativeLocation.Y","RelativeLocation.Z","RelativeRotation","RelativeRotation.Yaw","RelativeRotation.Pitch","RelativeRotation.Roll"))),("/Script/CinematicCamera.CineCameraComponent", ()),("/Script/Engine.CameraComponent", ())),DefaultAddedSubobjectRules=(SubobjectMatchingRules=(("/Script/Engine.Actor", (IncludeClasses=("/Script/Engine.StaticMeshComponent","/Script/Engine.CameraComponent","/Script/CinematicCamera.CineCameraComponent"))))))
\ No newline at end of file
+ReplicationEditorSettings=(DefaultPropertySelection=(("/Script/Engine.SceneComponent", (DefaultSelectedProperties=("RelativeLocation","RelativeLocation.X","RelativeLocation.Y","RelativeLocation.Z","RelativeRotation","RelativeRotation.Yaw","RelativeRotation.Pitch","RelativeRotation.Roll"))),("/Script/CinematicCamera.CineCameraComponent", ()),("/Script/Engine.CameraComponent", ())),DefaultAddedSubobjectRules=(SubobjectMatchingRules=(("/Script/Engine.Actor", (IncludeClasses=("/Script/Engine.StaticMeshComponent","/Script/Engine.CameraComponent","/Script/CinematicCamera.CineCameraComponent"))))))
+FrequencyRules=(DefaultObjectFrequencySettings=(ReplicationMode=SpecifiedRate,ReplicationRate=30),FrequencyRuleStack=((ReplicationSettings=(ReplicationRate=60),SubobjectMatchingRules=(SubobjectMatchingRules=(("/Script/Engine.Actor", (IncludeClasses=("/Script/CinematicCamera.CineCameraComponent"))))))))
\ No newline at end of file
