In the multi client replication view, add an option to "Toggle this & children" in the hiearchy. Checking the row's checkbox defaults to "Toggle this & children" (changed from changing only the row's object)

diff --git a/Engine/Plugins/Developer/Concert/ConcertApp/MultiUserClient/Source/MultiUserClient/Private/Replication/Authority/AuthorityChangeTracker.cpp b/Engine/Plugins/Developer/Concert/ConcertApp/MultiUserClient/Source/MultiUserClient/Private/Replication/Authority/AuthorityChangeTracker.cpp
--- a/Engine/Plugins/Developer/Concert/ConcertApp/MultiUserClient/Source/MultiUserClient/Private/Replication/Authority/AuthorityChangeTracker.cpp
+++ b/Engine/Plugins/Developer/Concert/ConcertApp/MultiUserClient/Source/MultiUserClient/Private/Replication/Authority/AuthorityChangeTracker.cpp
@@ -44,7 +44,7 @@ namespace UE::MultiUserClient
 
 		if (bAddedAnything)
 		{
-			OnAddedOwnedObjectsDelegate.Broadcast();
+			OnChangedOwnedObjectsDelegate.Broadcast();
 		}
 	}
 
diff --git a/Engine/Plugins/Developer/Concert/ConcertApp/MultiUserClient/Source/MultiUserClient/Private/Replication/Authority/AuthorityChangeTracker.h b/Engine/Plugins/Developer/Concert/ConcertApp/MultiUserClient/Source/MultiUserClient/Private/Replication/Authority/AuthorityChangeTracker.h
--- a/Engine/Plugins/Developer/Concert/ConcertApp/MultiUserClient/Source/MultiUserClient/Private/Replication/Authority/AuthorityChangeTracker.h
+++ b/Engine/Plugins/Developer/Concert/ConcertApp/MultiUserClient/Source/MultiUserClient/Private/Replication/Authority/AuthorityChangeTracker.h
@@ -51,7 +51,7 @@ namespace UE::MultiUserClient
 
 		/** Called entries are added to NewAuthorityStates. NOT called when objects are removed. */
 		DECLARE_MULTICAST_DELEGATE(FOnAuthorityChangeMade);
-		FOnAuthorityChangeMade& OnAddedOwnedObjects() { return OnAddedOwnedObjectsDelegate; }
+		FOnAuthorityChangeMade& OnChangedOwnedObjects() { return OnChangedOwnedObjectsDelegate; }
 		
 	private:
 
@@ -67,7 +67,7 @@ namespace UE::MultiUserClient
 		TMap<FSoftObjectPath, bool> NewAuthorityStates;
 
 		/** Called when NewAuthorityStates is updated. */
-		FOnAuthorityChangeMade OnAddedOwnedObjectsDelegate;
+		FOnAuthorityChangeMade OnChangedOwnedObjectsDelegate;
 		
 		void OnClientChanged(const FGuid& Guid) { RefreshChanges(); }
 	};
diff --git a/Engine/Plugins/Developer/Concert/ConcertApp/MultiUserClient/Source/MultiUserClient/Private/Replication/Client/ReplicationClient.cpp b/Engine/Plugins/Developer/Concert/ConcertApp/MultiUserClient/Source/MultiUserClient/Private/Replication/Client/ReplicationClient.cpp
--- a/Engine/Plugins/Developer/Concert/ConcertApp/MultiUserClient/Source/MultiUserClient/Private/Replication/Client/ReplicationClient.cpp
+++ b/Engine/Plugins/Developer/Concert/ConcertApp/MultiUserClient/Source/MultiUserClient/Private/Replication/Client/ReplicationClient.cpp
@@ -56,7 +56,7 @@ namespace UE::MultiUserClient
 	{
 		LocalClientEditModel->OnObjectsChanged().AddRaw(this, &FReplicationClient::OnObjectsChanged);
 		LocalClientEditModel->OnPropertiesChanged().AddRaw(this, &FReplicationClient::OnPropertiesChanged);
-		LocalAuthorityDiffer.OnAddedOwnedObjects().AddRaw(this, &FReplicationClient::DeferOnModelChanged);
+		LocalAuthorityDiffer.OnChangedOwnedObjects().AddRaw(this, &FReplicationClient::DeferOnModelChanged);
 
 		SubmissionWorkflow->OnAuthorityRequestCompleted_AnyThread().AddRaw(this, &FReplicationClient::OnAuthoritySubmissionCompleted);
 		StreamSynchronizer->OnServerStateChanged().AddRaw(this, &FReplicationClient::OnServerStateChanged);
diff --git a/Engine/Plugins/Developer/Concert/ConcertApp/MultiUserClient/Source/MultiUserClient/Private/Replication/Submission/AutoSubmissionPolicy.cpp b/Engine/Plugins/Developer/Concert/ConcertApp/MultiUserClient/Source/MultiUserClient/Private/Replication/Submission/AutoSubmissionPolicy.cpp
--- a/Engine/Plugins/Developer/Concert/ConcertApp/MultiUserClient/Source/MultiUserClient/Private/Replication/Submission/AutoSubmissionPolicy.cpp
+++ b/Engine/Plugins/Developer/Concert/ConcertApp/MultiUserClient/Source/MultiUserClient/Private/Replication/Submission/AutoSubmissionPolicy.cpp
@@ -27,7 +27,7 @@ namespace UE::MultiUserClient
 	{
 		StreamEditorModel.OnObjectsChanged().AddRaw(this, &FAutoSubmissionPolicy::OnObjectsChanged);
 		StreamEditorModel.OnPropertiesChanged().AddRaw(this, &FAutoSubmissionPolicy::OnChangesDetected);
-		AuthorityChangeTracker.OnAddedOwnedObjects().AddRaw(this, &FAutoSubmissionPolicy::OnChangesDetected);
+		AuthorityChangeTracker.OnChangedOwnedObjects().AddRaw(this, &FAutoSubmissionPolicy::OnChangesDetected);
 		FrequencyChangeTracker.OnFrequencySettingsChanged().AddRaw(this, &FAutoSubmissionPolicy::OnChangesDetected);
 	}
 
@@ -35,7 +35,7 @@ namespace UE::MultiUserClient
 	{
 		StreamEditorModel.OnObjectsChanged().RemoveAll(this);
 		StreamEditorModel.OnPropertiesChanged().RemoveAll(this);
-		AuthorityChangeTracker.OnAddedOwnedObjects().RemoveAll(this);
+		AuthorityChangeTracker.OnChangedOwnedObjects().RemoveAll(this);
 		FrequencyChangeTracker.OnFrequencySettingsChanged().RemoveAll(this);
 	}
 
diff --git a/Engine/Plugins/Developer/Concert/ConcertApp/MultiUserClient/Source/MultiUserClient/Private/Widgets/ActiveSession/Replication/Client/Multi/Columns/MultiStreamColumns.h b/Engine/Plugins/Developer/Concert/ConcertApp/MultiUserClient/Source/MultiUserClient/Private/Widgets/ActiveSession/Replication/Client/Multi/Columns/MultiStreamColumns.h
--- a/Engine/Plugins/Developer/Concert/ConcertApp/MultiUserClient/Source/MultiUserClient/Private/Widgets/ActiveSession/Replication/Client/Multi/Columns/MultiStreamColumns.h
+++ b/Engine/Plugins/Developer/Concert/ConcertApp/MultiUserClient/Source/MultiUserClient/Private/Widgets/ActiveSession/Replication/Client/Multi/Columns/MultiStreamColumns.h
@@ -17,6 +17,7 @@ namespace UE::MultiUserClient
 
 namespace UE::ConcertSharedSlate
 {
+	class IObjectHierarchyModel;
 	class IReplicationStreamModel;
 	class IMultiReplicationStreamEditor;
 	class IReplicationStreamViewer;
@@ -40,7 +41,7 @@ namespace UE::MultiUserClient::MultiStreamColumns
 	 * Toggles replication for all clients assigned to the object (and optionally all children).
 	 * 
 	 * @param ConcertClient Used to look up client names
-	 * @param ConsolidatedStreamModelAttribute Used to get child objects
+	 * @param ObjectHierarchyModelAttribute Used to get child objects
 	 * @param ClientManager Used to access all clients for toggling authority
 	 * @param ColumnsSortPriority The order relative to the other columns
 	 * 
@@ -48,7 +49,7 @@ namespace UE::MultiUserClient::MultiStreamColumns
 	 */
 	ConcertSharedSlate::ReplicationColumns::FReplicationTopLevelObjectColumn ReplicationToggle(
 		TSharedRef<IConcertClient> ConcertClient,
-		TAttribute<ConcertSharedSlate::IReplicationStreamModel*> ConsolidatedStreamModelAttribute,
+		TAttribute<ConcertSharedSlate::IObjectHierarchyModel*> ObjectHierarchyModelAttribute,
 		FReplicationClientManager& ClientManager,
 		const int32 ColumnsSortPriority = static_cast<int32>(EColumnSortOrder::ReplicationToggle)
 		);
@@ -57,7 +58,8 @@ namespace UE::MultiUserClient::MultiStreamColumns
 	 * A combo box which displays all current owners for an object and allows bulk reassigning properties.
 	 *
 	 * @param ConcertClient Used to look up client names
-	 * @param MultiStreamModelAttribute Used to get child objects from the consolidated model and for requesting resorting the column
+	 * @param MultiStreamModelAttribute Used for requesting resorting the column
+	 * @param ObjectHierarchyModelAttribute Used to get child objects
 	 * @param ReassignmentLogic Performs the act of reassigning
 	 * @param ClientManager Used to access all clients for display in the combo box drop-down
 	 * @param ColumnsSortPriority The order relative to the other columns
@@ -65,6 +67,7 @@ namespace UE::MultiUserClient::MultiStreamColumns
 	ConcertSharedSlate::ReplicationColumns::FReplicationTopLevelObjectColumn ReassignOwnership(
 		TSharedRef<IConcertClient> ConcertClient,
 		TAttribute<TSharedPtr<ConcertSharedSlate::IMultiReplicationStreamEditor>> MultiStreamModelAttribute,
+		TAttribute<ConcertSharedSlate::IObjectHierarchyModel*> ObjectHierarchyModelAttribute,
 		FReassignObjectPropertiesLogic& ReassignmentLogic,
 		const FReplicationClientManager& ClientManager,
 		int32 ColumnsSortPriority = static_cast<int32>(EColumnSortOrder::ReassignOwnership)
diff --git a/Engine/Plugins/Developer/Concert/ConcertApp/MultiUserClient/Source/MultiUserClient/Private/Widgets/ActiveSession/Replication/Client/Multi/Columns/MultiStreamColumns_ReassignOwnership.cpp b/Engine/Plugins/Developer/Concert/ConcertApp/MultiUserClient/Source/MultiUserClient/Private/Widgets/ActiveSession/Replication/Client/Multi/Columns/MultiStreamColumns_ReassignOwnership.cpp
--- a/Engine/Plugins/Developer/Concert/ConcertApp/MultiUserClient/Source/MultiUserClient/Private/Widgets/ActiveSession/Replication/Client/Multi/Columns/MultiStreamColumns_ReassignOwnership.cpp
+++ b/Engine/Plugins/Developer/Concert/ConcertApp/MultiUserClient/Source/MultiUserClient/Private/Widgets/ActiveSession/Replication/Client/Multi/Columns/MultiStreamColumns_ReassignOwnership.cpp
@@ -17,6 +17,7 @@ namespace UE::MultiUserClient::MultiStreamColumns
 	ConcertSharedSlate::ReplicationColumns::FReplicationTopLevelObjectColumn ReassignOwnership(
 		TSharedRef<IConcertClient> ConcertClient,
 		TAttribute<TSharedPtr<ConcertSharedSlate::IMultiReplicationStreamEditor>> MultiStreamModelAttribute,
+		TAttribute<ConcertSharedSlate::IObjectHierarchyModel*> ObjectHierarchyModelAttribute,
 		FReassignObjectPropertiesLogic& ReassignmentLogic,
 		const FReplicationClientManager& ClientManager,
 		const int32 ColumnsSortPriority
@@ -25,18 +26,12 @@ namespace UE::MultiUserClient::MultiStreamColumns
 		using namespace ConcertSharedSlate::ReplicationColumns;
 
 		auto MakeWidget =
-			[ConcertClient, MultiStreamModelAttribute = MoveTemp(MultiStreamModelAttribute), &ReassignmentLogic, &ClientManager]
+			[ConcertClient, MultiStreamModelAttribute = MoveTemp(MultiStreamModelAttribute), ObjectHierarchyModelAttribute = MoveTemp(ObjectHierarchyModelAttribute), &ReassignmentLogic, &ClientManager]
 			(const FReplicationTopLevelObjectColumn::FBuildArgs& InArgs)
 			{
 				return SNew(SReassignObjectComboBox, ConcertClient, ReassignmentLogic, ClientManager)
 					.ManagedObject(InArgs.RowData.GetObjectPath())
-					.ConsolidatedModel_Lambda([MultiStreamModelAttribute]()
-					{
-						const TSharedPtr<ConcertSharedSlate::IMultiReplicationStreamEditor> Editor = MultiStreamModelAttribute.Get();
-						return Editor
-							? &MultiStreamModelAttribute.Get()->GetConsolidatedModel()
-							: nullptr;
-					})
+					.ObjectHierarchyModel(ObjectHierarchyModelAttribute)
 					.HighlightText(InArgs.HighlightText)
 					.OnReassignAllOptionClicked_Lambda([MultiStreamModelAttribute](auto)
 					{
diff --git a/Engine/Plugins/Developer/Concert/ConcertApp/MultiUserClient/Source/MultiUserClient/Private/Widgets/ActiveSession/Replication/Client/Multi/Columns/MultiStreamColumns_ReplicationToggle.cpp b/Engine/Plugins/Developer/Concert/ConcertApp/MultiUserClient/Source/MultiUserClient/Private/Widgets/ActiveSession/Replication/Client/Multi/Columns/MultiStreamColumns_ReplicationToggle.cpp
--- a/Engine/Plugins/Developer/Concert/ConcertApp/MultiUserClient/Source/MultiUserClient/Private/Widgets/ActiveSession/Replication/Client/Multi/Columns/MultiStreamColumns_ReplicationToggle.cpp
+++ b/Engine/Plugins/Developer/Concert/ConcertApp/MultiUserClient/Source/MultiUserClient/Private/Widgets/ActiveSession/Replication/Client/Multi/Columns/MultiStreamColumns_ReplicationToggle.cpp
@@ -16,7 +16,7 @@ namespace UE::MultiUserClient::MultiStreamColumns
 	
 	ConcertSharedSlate::ReplicationColumns::FReplicationTopLevelObjectColumn ReplicationToggle(
 		TSharedRef<IConcertClient> ConcertClient,
-		TAttribute<ConcertSharedSlate::IReplicationStreamModel*> ConsolidatedStreamModelAttribute,
+		TAttribute<ConcertSharedSlate::IObjectHierarchyModel*> ObjectHierarchyModelAttribute,
 		FReplicationClientManager& ClientManager,
 		const int32 ColumnsSortPriority
 		)
@@ -25,14 +25,14 @@ namespace UE::MultiUserClient::MultiStreamColumns
 		
 		return FReplicationTopLevelObjectColumn(
 			FReplicationTopLevelObjectColumn::FArguments()
-				.GenerateWidgetColumn_Lambda([ConcertClient = MoveTemp(ConcertClient), ConsolidatedStreamModelAttribute = MoveTemp(ConsolidatedStreamModelAttribute), &ClientManager](const FReplicationTopLevelObjectColumn::FBuildArgs& InArgs)
+				.GenerateWidgetColumn_Lambda([ConcertClient = MoveTemp(ConcertClient), ObjectHierarchyModelAttribute = MoveTemp(ObjectHierarchyModelAttribute), &ClientManager](const FReplicationTopLevelObjectColumn::FBuildArgs& InArgs)
 				{
 					return SNew(SBox)
 					.HAlign(HAlign_Left) // Warning icon is sometimes collapsed - we don't want the widget to be centered
 					[
 						SNew(SReplicationMultiToggleCheckbox, ClientManager, ConcertClient)
 						.Object(InArgs.RowData.GetObjectPath())
-						.ConsolidatedStreamModelAttribute(ConsolidatedStreamModelAttribute)
+						.ObjectHierarchyModel(ObjectHierarchyModelAttribute)
 					];
 				})
 				.ColumnSortOrder(ColumnsSortPriority),
diff --git a/Engine/Plugins/Developer/Concert/ConcertApp/MultiUserClient/Source/MultiUserClient/Private/Widgets/ActiveSession/Replication/Client/Multi/Columns/SReassignObjectComboBox.cpp b/Engine/Plugins/Developer/Concert/ConcertApp/MultiUserClient/Source/MultiUserClient/Private/Widgets/ActiveSession/Replication/Client/Multi/Columns/SReassignObjectComboBox.cpp
--- a/Engine/Plugins/Developer/Concert/ConcertApp/MultiUserClient/Source/MultiUserClient/Private/Widgets/ActiveSession/Replication/Client/Multi/Columns/SReassignObjectComboBox.cpp
+++ b/Engine/Plugins/Developer/Concert/ConcertApp/MultiUserClient/Source/MultiUserClient/Private/Widgets/ActiveSession/Replication/Client/Multi/Columns/SReassignObjectComboBox.cpp
@@ -13,6 +13,7 @@
 #include "Framework/MultiBox/MultiBoxBuilder.h"
 #include "Framework/Commands/UIAction.h"
 #include "Framework/Commands/UICommandInfo.h"
+#include "Replication/Editor/Model/Object/IObjectHierarchyModel.h"
 #include "Textures/SlateIcon.h"
 #include "Widgets/Images/SThrobber.h"
 #include "Widgets/Input/SComboBox.h"
@@ -82,8 +83,8 @@ namespace UE::MultiUserClient
 		HighlightText = InArgs._HighlightText;
 
 		ManagedObject = InArgs._ManagedObject;
-		ConsolidatedStreamModelAttribute = InArgs._ConsolidatedModel;
-		check(ConsolidatedStreamModelAttribute.IsSet() || ConsolidatedStreamModelAttribute.IsBound());
+		ObjectHierarchyModelAttribute = InArgs._ObjectHierarchyModel;
+		check(ObjectHierarchyModelAttribute.IsSet() || ObjectHierarchyModelAttribute.IsBound());
 
 		OnReassignAllOptionClickedDelegate = InArgs._OnReassignAllOptionClicked;
 		
@@ -204,8 +205,7 @@ namespace UE::MultiUserClient
 
 	SReassignObjectComboBox::FInlineObjectPathArray SReassignObjectComboBox::GetChildrenOfManagedObject() const
 	{
-		const ConcertSharedSlate::IReplicationStreamModel* Model = ConsolidatedStreamModelAttribute.Get();
-		return ensure(Model) ? Model->GetSubobjects<FInlineAllocator>(ManagedObject) : FInlineObjectPathArray{};
+		return ObjectHierarchyModelAttribute.Get()->GetChildrenRecursive<FInlineAllocator>(ManagedObject);
 	}
 }
 
diff --git a/Engine/Plugins/Developer/Concert/ConcertApp/MultiUserClient/Source/MultiUserClient/Private/Widgets/ActiveSession/Replication/Client/Multi/Columns/SReassignObjectComboBox.h b/Engine/Plugins/Developer/Concert/ConcertApp/MultiUserClient/Source/MultiUserClient/Private/Widgets/ActiveSession/Replication/Client/Multi/Columns/SReassignObjectComboBox.h
--- a/Engine/Plugins/Developer/Concert/ConcertApp/MultiUserClient/Source/MultiUserClient/Private/Widgets/ActiveSession/Replication/Client/Multi/Columns/SReassignObjectComboBox.h
+++ b/Engine/Plugins/Developer/Concert/ConcertApp/MultiUserClient/Source/MultiUserClient/Private/Widgets/ActiveSession/Replication/Client/Multi/Columns/SReassignObjectComboBox.h
@@ -10,7 +10,7 @@ class IConcertClient;
 class IConcertClientSession;
 
 namespace UE::ConcertClientSharedSlate { class SHorizontalClientList; }
-namespace UE::ConcertSharedSlate { class IReplicationStreamModel; }
+namespace UE::ConcertSharedSlate { class IObjectHierarchyModel; }
 
 namespace UE::MultiUserClient
 {
@@ -46,7 +46,7 @@ namespace UE::MultiUserClient
 			SLATE_ARGUMENT(FSoftObjectPath, ManagedObject)
 			SLATE_ARGUMENT(TSharedPtr<FText>, HighlightText)
 			/** Used to to figure out child objects. */
-			SLATE_ATTRIBUTE(ConcertSharedSlate::IReplicationStreamModel*, ConsolidatedModel)
+			SLATE_ATTRIBUTE(ConcertSharedSlate::IObjectHierarchyModel*, ObjectHierarchyModel)
 
 			/** Called when a valid client ID is selected for reassignment. */
 			SLATE_EVENT(FOnReassignAllOptionClicked, OnReassignAllOptionClicked)
@@ -73,7 +73,7 @@ namespace UE::MultiUserClient
 		/** The object being reassigned by this widget. */
 		FSoftObjectPath ManagedObject;
 		/** Used to to figure out child objects. */
-		TAttribute<ConcertSharedSlate::IReplicationStreamModel*> ConsolidatedStreamModelAttribute;
+		TAttribute<ConcertSharedSlate::IObjectHierarchyModel*> ObjectHierarchyModelAttribute;
 		
 		/** Displayed as content of combo button */
 		TSharedPtr<ConcertClientSharedSlate::SHorizontalClientList> ComboClientList;
diff --git a/Engine/Plugins/Developer/Concert/ConcertApp/MultiUserClient/Source/MultiUserClient/Private/Widgets/ActiveSession/Replication/Client/Multi/Columns/SReplicationMultiToggleCheckbox.cpp b/Engine/Plugins/Developer/Concert/ConcertApp/MultiUserClient/Source/MultiUserClient/Private/Widgets/ActiveSession/Replication/Client/Multi/Columns/SReplicationMultiToggleCheckbox.cpp
--- a/Engine/Plugins/Developer/Concert/ConcertApp/MultiUserClient/Source/MultiUserClient/Private/Widgets/ActiveSession/Replication/Client/Multi/Columns/SReplicationMultiToggleCheckbox.cpp
+++ b/Engine/Plugins/Developer/Concert/ConcertApp/MultiUserClient/Source/MultiUserClient/Private/Widgets/ActiveSession/Replication/Client/Multi/Columns/SReplicationMultiToggleCheckbox.cpp
@@ -9,6 +9,7 @@
 #include "Algo/AnyOf.h"
 #include "Framework/Commands/UIAction.h"
 #include "Framework/MultiBox/MultiBoxBuilder.h"
+#include "Replication/Editor/Model/Object/IObjectHierarchyModel.h"
 #include "Styling/AppStyle.h"
 #include "Styling/SlateTypes.h"
 #include "Textures/SlateIcon.h"
@@ -29,8 +30,8 @@ namespace UE::MultiUserClient
 		)
 	{
 		Object = InArgs._Object;
-		ConsolidatedStreamModelAttribute = InArgs._ConsolidatedStreamModelAttribute;
-		check(ConsolidatedStreamModelAttribute.IsSet() || ConsolidatedStreamModelAttribute.IsBound());
+		ObjectHierarchyModelAttribute = InArgs._ObjectHierarchyModel;
+		check(ObjectHierarchyModelAttribute.IsSet() || ObjectHierarchyModelAttribute.IsBound());
 		
 		ClientManager = &InClientManager;
 		ConcertClient = MoveTemp(InConcertClient);
@@ -51,9 +52,9 @@ namespace UE::MultiUserClient
 				.AutoWidth()
 				[
 					SNew(SCheckBox)
-					.IsChecked(this, &SReplicationMultiToggleCheckbox::GetCheckboxStateForObject, Object)
-					.IsEnabled(this, &SReplicationMultiToggleCheckbox::IsCheckboxEnabledForObject, Object)
-					.OnCheckStateChanged(this, &SReplicationMultiToggleCheckbox::OnCheckboxStateChangedForObject, Object)
+					.IsChecked(this, &SReplicationMultiToggleCheckbox::GetCheckboxStateForThisAndChildren)
+					.IsEnabled(this, &SReplicationMultiToggleCheckbox::CanToggleThisOrChildren)
+					.OnCheckStateChanged(this, &SReplicationMultiToggleCheckbox::OnCheckboxStateChanged)
 				]
 			
 				+SHorizontalBox::Slot()
@@ -91,45 +92,48 @@ namespace UE::MultiUserClient
 		}
 	}
 
-	ECheckBoxState SReplicationMultiToggleCheckbox::GetCheckboxStateForObject(FSoftObjectPath InObject) const
+	ECheckBoxState SReplicationMultiToggleCheckbox::GetCheckboxStateForObjects(TConstArrayView<FSoftObjectPath> InObjects) const
 	{
 		// The checkbox shows
-		// - checked if all editable clients with the object registered have authority
-		// - unchecked if all editable clients with the object registered do not have authority
+		// - checked if all editable clients with the objects registered have authority
+		// - unchecked if all editable clients with the objects registered do not have authority
 		// - undetermined otherwise
-		const TArray<FGuid> ClientsWithAuthority = ClientManager->GetAuthorityCache().GetClientsWithAuthorityOverObject(InObject);
 
-		// TODO UE-200496 Predict conflicts
-		
 		bool bHadAnyEditableClient = false;
-		ECheckBoxState ConsolidatedState = ECheckBoxState::Undetermined;
-		ClientManager->GetAuthorityCache().ForEachClientWithObjectInStream(InObject, [this, &ClientsWithAuthority, &bHadAnyEditableClient, &ConsolidatedState](const FGuid& ClientId)
+		TOptional<ECheckBoxState> ConsolidatedState;
+		for (const FSoftObjectPath& InObject : InObjects)
 		{
-			// The state of the checkbox always skips non-editable clients
-			const FReplicationClient* Client = ClientManager->FindClient(ClientId);
-			if (!ensure(Client) || !Client->AllowsEditing())
+			const TArray<FGuid> ClientsWithAuthority = ClientManager->GetAuthorityCache().GetClientsWithAuthorityOverObject(InObject);
+
+			// TODO UE-200496 Predict conflicts
+			ClientManager->GetAuthorityCache().ForEachClientWithObjectInStream(InObject, [this, &ClientsWithAuthority, &bHadAnyEditableClient, &ConsolidatedState](const FGuid& ClientId)
 			{
-				return EBreakBehavior::Continue;
-			}
+				// The state of the checkbox always skips non-editable clients
+				const FReplicationClient* Client = ClientManager->FindClient(ClientId);
+				if (!ensure(Client) || !Client->AllowsEditing())
+				{
+					return EBreakBehavior::Continue;
+				}
 			
-			bHadAnyEditableClient = true;
-			const bool bHasAuthority = ClientsWithAuthority.Contains(Client->GetEndpointId());
-			const ECheckBoxState ClientState = bHasAuthority ? ECheckBoxState::Checked : ECheckBoxState::Unchecked;
+				bHadAnyEditableClient = true;
+				const bool bHasAuthority = ClientsWithAuthority.Contains(Client->GetEndpointId());
+				const ECheckBoxState ClientState = bHasAuthority ? ECheckBoxState::Checked : ECheckBoxState::Unchecked;
 			
-			if (ConsolidatedState == ECheckBoxState::Undetermined)
-			{
-				ConsolidatedState = ClientState;
-			}
-			else if (ConsolidatedState != ClientState)
-			{
-				ConsolidatedState = ECheckBoxState::Undetermined;
-				return EBreakBehavior::Break;
-			}
+				if (!ConsolidatedState.IsSet())
+				{
+					ConsolidatedState = ClientState;
+				}
+				else if (*ConsolidatedState != ClientState)
+				{
+					ConsolidatedState = ECheckBoxState::Undetermined;
+					return EBreakBehavior::Break;
+				}
 			
-			return EBreakBehavior::Continue;
-		});
+				return EBreakBehavior::Continue;
+			});
+		}
 		
-		return bHadAnyEditableClient ? ConsolidatedState : ECheckBoxState::Unchecked;
+		return bHadAnyEditableClient ? ConsolidatedState.Get(ECheckBoxState::Unchecked) : ECheckBoxState::Unchecked;
 	}
 
 	bool SReplicationMultiToggleCheckbox::IsCheckboxEnabledForObject(FSoftObjectPath InObject) const
@@ -146,19 +150,18 @@ namespace UE::MultiUserClient
 		return bHasAtLeastOneEditableClient;
 	}
 
-	void SReplicationMultiToggleCheckbox::OnCheckboxStateChangedForObject(ECheckBoxState NewState, FSoftObjectPath InObject) const
+	void SReplicationMultiToggleCheckbox::OnCheckboxStateChanged(ECheckBoxState NewState) const
 	{
 		const bool bShouldHaveAuthority = NewState == ECheckBoxState::Checked;
-		ClientManager->GetAuthorityCache().ForEachClientWithObjectInStream(InObject, [this, InObject, bShouldHaveAuthority](const FGuid& ClientId)
-		{
-			FReplicationClient* Client = ClientManager->FindClient(ClientId);
-			if (ensure(Client) && Client->AllowsEditing())
+		
+		SetAuthorityForObject(bShouldHaveAuthority, Object);
+		ObjectHierarchyModelAttribute.Get()->ForEachChildRecursive(
+			Object, 
+		[this, bShouldHaveAuthority](const FSoftObjectPath&, const FSoftObjectPath& ChildObject, ConcertSharedSlate::EChildRelationship)
 			{
-				Client->GetAuthorityDiffer().SetAuthorityIfAllowed({ InObject }, bShouldHaveAuthority);
-			}
-			
-			return EBreakBehavior::Continue;
-		});
+				SetAuthorityForObject(bShouldHaveAuthority, ChildObject);
+				return EBreakBehavior::Continue;
+			});
 	}
 
 	TSharedRef<SWidget> SReplicationMultiToggleCheckbox::GetDropDownMenuContent()
@@ -166,17 +169,12 @@ namespace UE::MultiUserClient
 		FMenuBuilder MenuBuilder(true, nullptr);
 
 		MenuBuilder.AddMenuEntry(
-			LOCTEXT("OnlyThis", "Toggle object only"),
+			LOCTEXT("OnlyThis", "Toggle this only"),
 			FText::GetEmpty(),
 			FSlateIcon(),
 			FUIAction(
-				FExecuteAction::CreateLambda([this]()
-				{
-					const ECheckBoxState OldCheckboxState = GetCheckboxStateForObject(Object);
-					const bool bGiveAuthority = OldCheckboxState == ECheckBoxState::Unchecked;
-					OnCheckboxStateChangedForObject(bGiveAuthority ? ECheckBoxState::Checked : ECheckBoxState::Unchecked, Object);
-				}),
-				FCanExecuteAction::CreateLambda([this](){ return IsCheckboxEnabledForObject(Object); })
+				FExecuteAction::CreateSP(this, &SReplicationMultiToggleCheckbox::ToggleThis),
+				FCanExecuteAction::CreateSP(this, &SReplicationMultiToggleCheckbox::CanToggleThis)
 				)
 			);
 		MenuBuilder.AddMenuEntry(
@@ -188,20 +186,29 @@ namespace UE::MultiUserClient
 				FCanExecuteAction::CreateSP(this, &SReplicationMultiToggleCheckbox::CanToggleChildren)
 				)
 			);
+		MenuBuilder.AddMenuEntry(
+			LOCTEXT("OnlyThis", "Toggle this & children"),
+			FText::GetEmpty(),
+			FSlateIcon(),
+			FUIAction(
+				FExecuteAction::CreateSP(this, &SReplicationMultiToggleCheckbox::ToggleThisAndChildren),
+				FCanExecuteAction::CreateSP(this, &SReplicationMultiToggleCheckbox::CanToggleThisOrChildren)
+				)
+			);
 		
 		return MenuBuilder.MakeWidget();
 	}
 
-	void SReplicationMultiToggleCheckbox::ToggleChildren() const
+	TArray<FSoftObjectPath> SReplicationMultiToggleCheckbox::GetThisAndChildren() const
 	{
-		const TArray<FSoftObjectPath> Subobjects = ConsolidatedStreamModelAttribute.Get()->GetSubobjects(Object);
-		ToggleObjects(Subobjects);
+		TArray<FSoftObjectPath> AllObjects = GetChildObjects();
+		AllObjects.Add(Object);
+		return AllObjects;
 	}
 
-	bool SReplicationMultiToggleCheckbox::CanToggleChildren() const
+	TArray<FSoftObjectPath> SReplicationMultiToggleCheckbox::GetChildObjects() const
 	{
-		const TArray<FSoftObjectPath> Subobjects = ConsolidatedStreamModelAttribute.Get()->GetSubobjects(Object);
-		return CanToggleObjects(Subobjects);
+		return ObjectHierarchyModelAttribute.Get()->GetChildrenRecursive(Object);
 	}
 
 	void SReplicationMultiToggleCheckbox::ToggleObjects(TConstArrayView<FSoftObjectPath> Objects) const
@@ -224,9 +231,10 @@ namespace UE::MultiUserClient
 		}
 
 		const ECheckBoxState StateToSet = ConsolidatedCheckboxState == ECheckBoxState::Unchecked ? ECheckBoxState::Checked : ECheckBoxState::Unchecked;
+		const bool bShouldHaveAuthority = StateToSet == ECheckBoxState::Checked;
 		for (const FSoftObjectPath& ObjectPath : Objects)
 		{
-			OnCheckboxStateChangedForObject(StateToSet, ObjectPath);
+			SetAuthorityForObject(bShouldHaveAuthority, ObjectPath);
 		}
 	}
 
@@ -238,6 +246,20 @@ namespace UE::MultiUserClient
 		});
 	}
 
+	void SReplicationMultiToggleCheckbox::SetAuthorityForObject(bool bShouldHaveAuthority, const FSoftObjectPath& InObject) const
+	{
+		ClientManager->GetAuthorityCache().ForEachClientWithObjectInStream(InObject, [this, InObject, bShouldHaveAuthority](const FGuid& ClientId)
+		{
+			FReplicationClient* Client = ClientManager->FindClient(ClientId);
+			if (ensure(Client) && Client->AllowsEditing())
+			{
+				Client->GetAuthorityDiffer().SetAuthorityIfAllowed({ InObject }, bShouldHaveAuthority);
+			}
+			
+			return EBreakBehavior::Continue;
+		});
+	}
+
 	EVisibility SReplicationMultiToggleCheckbox::GetWarningVisibility() const
 	{
 		const bool bContainsNonEditableClients = !GetNonEditableClients().IsEmpty();
diff --git a/Engine/Plugins/Developer/Concert/ConcertApp/MultiUserClient/Source/MultiUserClient/Private/Widgets/ActiveSession/Replication/Client/Multi/Columns/SReplicationMultiToggleCheckbox.h b/Engine/Plugins/Developer/Concert/ConcertApp/MultiUserClient/Source/MultiUserClient/Private/Widgets/ActiveSession/Replication/Client/Multi/Columns/SReplicationMultiToggleCheckbox.h
--- a/Engine/Plugins/Developer/Concert/ConcertApp/MultiUserClient/Source/MultiUserClient/Private/Widgets/ActiveSession/Replication/Client/Multi/Columns/SReplicationMultiToggleCheckbox.h
+++ b/Engine/Plugins/Developer/Concert/ConcertApp/MultiUserClient/Source/MultiUserClient/Private/Widgets/ActiveSession/Replication/Client/Multi/Columns/SReplicationMultiToggleCheckbox.h
@@ -9,18 +9,12 @@
 class IConcertClient;
 enum class ECheckBoxState : uint8;
 
-namespace UE::ConcertSharedSlate
-{
-	class IReplicationStreamModel;
-}
+namespace UE::ConcertSharedSlate { class IObjectHierarchyModel; }
 
 namespace UE::MultiUserClient
 {
 	class FReplicationClientManager;
-}
-
-namespace UE::MultiUserClient
-{
+	
 	/**
 	 * Checkbox in a combo button.
 	 * 
@@ -36,7 +30,7 @@ namespace UE::MultiUserClient
 		SLATE_BEGIN_ARGS(SReplicationMultiToggleCheckbox)
 		{}
 		    SLATE_ARGUMENT(FSoftObjectPath, Object)
-			SLATE_ATTRIBUTE(ConcertSharedSlate::IReplicationStreamModel*, ConsolidatedStreamModelAttribute)
+			SLATE_ATTRIBUTE(ConcertSharedSlate::IObjectHierarchyModel*, ObjectHierarchyModel)
 		SLATE_END_ARGS()
 
 		void Construct(
@@ -53,7 +47,7 @@ namespace UE::MultiUserClient
 		/** Used to access all clients for toggling authority. */
 		FReplicationClientManager* ClientManager = nullptr;
 		/** Used to get children of Object */
-		TAttribute<ConcertSharedSlate::IReplicationStreamModel*> ConsolidatedStreamModelAttribute;
+		TAttribute<ConcertSharedSlate::IObjectHierarchyModel*> ObjectHierarchyModelAttribute;
 		
 		/** Used to look up client display names in case of conflicts. */
 		TSharedPtr<IConcertClient> ConcertClient;
@@ -62,20 +56,36 @@ namespace UE::MultiUserClient
 		FText GetRootToolTipText() const;
 		
 		/** @return Goes through all clients that can be edited and returns whether all of them have authority. */
-		ECheckBoxState GetCheckboxStateForObject(FSoftObjectPath InObject) const;
+		ECheckBoxState GetCheckboxStateForObject(const FSoftObjectPath& InObject) const { return GetCheckboxStateForObjects({ InObject }); }
+		ECheckBoxState GetCheckboxStateForThisAndChildren() const
+		{
+			// This algorithm could be improved but it will be performant enough for small hierarchies.
+			// The growth is linear in the number of children but every checkbox evaluates itself constantly instead of asking child boxes for cached state.
+			return GetCheckboxStateForObjects(GetThisAndChildren());
+		}
+		ECheckBoxState GetCheckboxStateForObjects(TConstArrayView<FSoftObjectPath> InObjects) const;
 		/** @return Whether there are any clients for which the authority can be changed. */
 		bool IsCheckboxEnabledForObject(FSoftObjectPath InObject) const;
-		/** Gives all editable clients which have something registered to Object authority or takes it away. */
-		void OnCheckboxStateChangedForObject(ECheckBoxState NewState, FSoftObjectPath InObject) const;
+		/** Sets the authority state of this and all children based on NewState. */
+		void OnCheckboxStateChanged(ECheckBoxState NewState) const;
 
 		/** @return Menu widget that with the options to remove or give authority to this and other objects. */
 		TSharedRef<SWidget> GetDropDownMenuContent();
-		
-		void ToggleChildren() const;
-		bool CanToggleChildren() const;
+
+		TArray<FSoftObjectPath> GetThisAndChildren() const;
+		TArray<FSoftObjectPath> GetChildObjects() const;
+
+		void ToggleThis() const { ToggleObjects({ Object }); };
+		void ToggleChildren() const { ToggleObjects(GetChildObjects()); }
+		void ToggleThisAndChildren() const { ToggleObjects(GetThisAndChildren()); }
+
+		bool CanToggleThis() const { return CanToggleObjects({ Object }); }
+		bool CanToggleChildren() const { return CanToggleObjects(GetChildObjects()); }
+		bool CanToggleThisOrChildren() const { return CanToggleObjects(GetThisAndChildren()); }
 
 		void ToggleObjects(TConstArrayView<FSoftObjectPath> Objects) const;
 		bool CanToggleObjects(TConstArrayView<FSoftObjectPath> Objects) const;
+		void SetAuthorityForObject(bool bShouldHaveAuthority, const FSoftObjectPath& InObject) const;
 		
 		EVisibility GetWarningVisibility() const;
 		FText GetWarningToolTipText() const;
diff --git a/Engine/Plugins/Developer/Concert/ConcertApp/MultiUserClient/Source/MultiUserClient/Private/Widgets/ActiveSession/Replication/Client/Multi/SMultiClientView.cpp b/Engine/Plugins/Developer/Concert/ConcertApp/MultiUserClient/Source/MultiUserClient/Private/Widgets/ActiveSession/Replication/Client/Multi/SMultiClientView.cpp
--- a/Engine/Plugins/Developer/Concert/ConcertApp/MultiUserClient/Source/MultiUserClient/Private/Widgets/ActiveSession/Replication/Client/Multi/SMultiClientView.cpp
+++ b/Engine/Plugins/Developer/Concert/ConcertApp/MultiUserClient/Source/MultiUserClient/Private/Widgets/ActiveSession/Replication/Client/Multi/SMultiClientView.cpp
@@ -72,10 +72,10 @@ namespace UE::MultiUserClient
 		   {
 			   return StreamEditor;
 		   });
-		const TAttribute<IReplicationStreamModel*> ConsolidatedStreamModelAttribute =
-		   TAttribute<IReplicationStreamModel*>::CreateLambda([this]()
+		const TAttribute<IObjectHierarchyModel*> ObjecHierarchyAttribute =
+		   TAttribute<IObjectHierarchyModel*>::CreateLambda([this]()
 		   {
-			   return &StreamEditor->GetConsolidatedModel();
+			   return ObjectHierarchy.Get();
 		   });
 		FGetAutoAssignTarget GetAutoAssignTargetDelegate = FGetAutoAssignTarget::CreateLambda([this, &InClientManager](TConstArrayView<UObject*>)
 		{
@@ -83,6 +83,8 @@ namespace UE::MultiUserClient
 			return StreamModel->GetEditableStreams().Contains(LocalStream) ? LocalStream.ToSharedPtr() : nullptr;
 		});
 		
+		ObjectHierarchy = ConcertClientSharedSlate::CreateObjectHierarchyForComponentHierarchy();
+		
 		FCreateMultiStreamEditorParams Params
 		{
 			.MultiStreamModel = StreamModel.ToSharedRef(),
@@ -92,13 +94,13 @@ namespace UE::MultiUserClient
 			.GetAutoAssignToStreamDelegate = MoveTemp(GetAutoAssignTargetDelegate),
 			.ViewerParams 
 			{
-				.ObjectHierarchy = ConcertClientSharedSlate::CreateObjectHierarchyForComponentHierarchy(), // This makes actors have children in the top view
+				.ObjectHierarchy = ObjectHierarchy, // This makes actors have children in the top view
 				.NameModel = ConcertClientSharedSlate::CreateEditorObjectNameModel(), // This makes actors use their labels, and components use the names given in the BP editor
 				.OnExtendObjectsContextMenu = FExtendObjectMenu::CreateSP(this, &SMultiClientView::ExtendObjectContextMenu),
 				.AdditionalObjectColumns =
 				{
-					MultiStreamColumns::ReplicationToggle(InConcertClient, ConsolidatedStreamModelAttribute, InClientManager),
-					MultiStreamColumns::ReassignOwnership(InConcertClient, MultiStreamEditorAttribute, InClientManager.GetReassignmentLogic(), InClientManager)
+					MultiStreamColumns::ReplicationToggle(InConcertClient, ObjecHierarchyAttribute, InClientManager),
+					MultiStreamColumns::ReassignOwnership(InConcertClient, MultiStreamEditorAttribute, ObjecHierarchyAttribute, InClientManager.GetReassignmentLogic(), InClientManager)
 				},
 				.AdditionalPropertyColumns = { MultiStreamColumns::AssignPropertyColumn(MultiStreamEditorAttribute, InConcertClient, InClientManager) },
 				.PrimaryPropertySort = { MultiStreamColumns::AssignPropertyColumnId, EColumnSortMode::Ascending}
diff --git a/Engine/Plugins/Developer/Concert/ConcertApp/MultiUserClient/Source/MultiUserClient/Private/Widgets/ActiveSession/Replication/Client/Multi/SMultiClientView.h b/Engine/Plugins/Developer/Concert/ConcertApp/MultiUserClient/Source/MultiUserClient/Private/Widgets/ActiveSession/Replication/Client/Multi/SMultiClientView.h
--- a/Engine/Plugins/Developer/Concert/ConcertApp/MultiUserClient/Source/MultiUserClient/Private/Widgets/ActiveSession/Replication/Client/Multi/SMultiClientView.h
+++ b/Engine/Plugins/Developer/Concert/ConcertApp/MultiUserClient/Source/MultiUserClient/Private/Widgets/ActiveSession/Replication/Client/Multi/SMultiClientView.h
@@ -10,6 +10,7 @@ class FMenuBuilder;
 
 namespace UE::ConcertSharedSlate
 {
+	class IObjectHierarchyModel;
 	class IMultiReplicationStreamEditor;
 	class IEditableReplicationStreamModel;
 }
@@ -48,6 +49,8 @@ namespace UE::MultiUserClient
 		TSharedPtr<FMultiStreamModel> StreamModel;
 		/** Displayed in the UI. */
 		TSharedPtr<ConcertSharedSlate::IMultiReplicationStreamEditor> StreamEditor;
+		/** Used by widgets in columns. */
+		TSharedPtr<ConcertSharedSlate::IObjectHierarchyModel> ObjectHierarchy;
 
 		/** Creates this widget's editor content */
 		TSharedRef<SWidget> CreateEditorContent(const TSharedRef<IConcertClient>& InConcertClient, FReplicationClientManager& InClientManager);
diff --git a/Engine/Plugins/Developer/Concert/ConcertUI/ConcertSharedSlate/Source/ConcertSharedSlate/Private/Replication/Editor/Model/Object/IObjectHierarchyModel.cpp b/Engine/Plugins/Developer/Concert/ConcertUI/ConcertSharedSlate/Source/ConcertSharedSlate/Private/Replication/Editor/Model/Object/IObjectHierarchyModel.cpp
--- a/Engine/Plugins/Developer/Concert/ConcertUI/ConcertSharedSlate/Source/ConcertSharedSlate/Private/Replication/Editor/Model/Object/IObjectHierarchyModel.cpp
+++ b/Engine/Plugins/Developer/Concert/ConcertUI/ConcertSharedSlate/Source/ConcertSharedSlate/Private/Replication/Editor/Model/Object/IObjectHierarchyModel.cpp
@@ -10,12 +10,12 @@ namespace UE::ConcertSharedSlate
 		const FSoftObjectPath& Root,
 		TFunctionRef<EBreakBehavior(const FSoftObjectPath& Parent, const FSoftObjectPath& ChildObject, EChildRelationship Relationship)> Callback,
 		EChildRelationshipFlags InclusionFlags
-		)
+		) const
 	{
 		struct FHelper
 		{
 			static EBreakBehavior Visit(
-				IObjectHierarchyModel& Model,
+				const IObjectHierarchyModel& Model,
 				const FSoftObjectPath& Parent,
 				const FSoftObjectPath& ChildObject,
 				EChildRelationship Relationship,
@@ -43,4 +43,4 @@ namespace UE::ConcertSharedSlate
 			return FHelper::Visit(*this, Root, ChildObject, Relationship, Callback, InclusionFlags);
 		}, InclusionFlags);
 	}
-}
\ No newline at end of file
+}
diff --git a/Engine/Plugins/Developer/Concert/ConcertUI/ConcertSharedSlate/Source/ConcertSharedSlate/Public/Replication/Editor/Model/Object/IObjectHierarchyModel.h b/Engine/Plugins/Developer/Concert/ConcertUI/ConcertSharedSlate/Source/ConcertSharedSlate/Public/Replication/Editor/Model/Object/IObjectHierarchyModel.h
--- a/Engine/Plugins/Developer/Concert/ConcertUI/ConcertSharedSlate/Source/ConcertSharedSlate/Public/Replication/Editor/Model/Object/IObjectHierarchyModel.h
+++ b/Engine/Plugins/Developer/Concert/ConcertUI/ConcertSharedSlate/Source/ConcertSharedSlate/Public/Replication/Editor/Model/Object/IObjectHierarchyModel.h
@@ -2,6 +2,7 @@
 
 #pragma once
 
+#include "Containers/Array.h"
 #include "Delegates/Delegate.h"
 #include "Misc/EBreakBehavior.h"
 #include "UObject/SoftObjectPath.h"
@@ -87,7 +88,21 @@ namespace UE::ConcertSharedSlate
 			const FSoftObjectPath& Root,
 			TFunctionRef<EBreakBehavior(const FSoftObjectPath& Parent, const FSoftObjectPath& ChildObject, EChildRelationship Relationship)> Callback,
 			EChildRelationshipFlags InclusionFlags = EChildRelationshipFlags::All
-			);
+			) const;
+
+		/** Builds a TArray containing all children, recursively. */
+		template<typename TAllocator = FDefaultAllocator>
+		TArray<FSoftObjectPath, TAllocator> GetChildrenRecursive(const FSoftObjectPath& Root, EChildRelationshipFlags InclusionFlags = EChildRelationshipFlags::All) const
+		{
+			TArray<FSoftObjectPath, TAllocator> AllObjects;
+			ForEachChildRecursive(Root,
+				[&AllObjects](const FSoftObjectPath&, const FSoftObjectPath& ChildObject, EChildRelationship)
+				{
+					AllObjects.Add(ChildObject);
+					return EBreakBehavior::Continue;
+				}, InclusionFlags);
+			return AllObjects;
+		}
 
 		virtual ~IObjectHierarchyModel() = default;
 	};
