Fix circular memory leak: Keep weak reference to Multi User's owning major tab.

diff --git a/Engine/Plugins/Developer/Concert/ConcertApp/MultiUserClient/Source/MultiUserClient/Private/Widgets/SConcertBrowser.cpp b/Engine/Plugins/Developer/Concert/ConcertApp/MultiUserClient/Source/MultiUserClient/Private/Widgets/SConcertBrowser.cpp
--- a/Engine/Plugins/Developer/Concert/ConcertApp/MultiUserClient/Source/MultiUserClient/Private/Widgets/SConcertBrowser.cpp
+++ b/Engine/Plugins/Developer/Concert/ConcertApp/MultiUserClient/Source/MultiUserClient/Private/Widgets/SConcertBrowser.cpp
@@ -6,7 +6,6 @@
 #include "IConcertSyncClient.h"
 #include "MultiUserClientUtils.h"
 #include "ActiveSession/SActiveSessionRoot.h"
-#include "Widgets/ActiveSession/Overview/SActiveSessionOverviewTab.h"
 #include "Widgets/Disconnected/SConcertClientSessionBrowser.h"
 #include "Widgets/Disconnected/SConcertNoAvailability.h"
 
@@ -60,6 +59,12 @@ void SConcertBrowser::HandleSessionConnectionChanged(IConcertClientSession& InSe
 
 void SConcertBrowser::AttachChildWidget(EConcertConnectionStatus ConnectionStatus)
 {
+	const TSharedPtr<SDockTab> OwningMajorTab = ConstructedUnderMajorTab.Pin();
+	if (!ensure(OwningMajorTab))
+	{
+		return;
+	}
+	
 	if (const TSharedPtr<IConcertSyncClient> ConcertSyncClient = WeakConcertSyncClient.Pin()
 		; ensure(ConcertSyncClient))
 	{
@@ -70,7 +75,7 @@ void SConcertBrowser::AttachChildWidget(EConcertConnectionStatus ConnectionStatu
 			{
 				ChildSlot.AttachWidget(
 					SNew(UE::MultiUserClient::SActiveSessionRoot,
-						ConstructedUnderMajorTab.ToSharedRef(),
+						OwningMajorTab.ToSharedRef(),
 						ConcertSyncClient,
 						ReplicationManager.ToSharedRef()
 						)
diff --git a/Engine/Plugins/Developer/Concert/ConcertApp/MultiUserClient/Source/MultiUserClient/Private/Widgets/SConcertBrowser.h b/Engine/Plugins/Developer/Concert/ConcertApp/MultiUserClient/Source/MultiUserClient/Private/Widgets/SConcertBrowser.h
--- a/Engine/Plugins/Developer/Concert/ConcertApp/MultiUserClient/Source/MultiUserClient/Private/Widgets/SConcertBrowser.h
+++ b/Engine/Plugins/Developer/Concert/ConcertApp/MultiUserClient/Source/MultiUserClient/Private/Widgets/SConcertBrowser.h
@@ -46,8 +46,11 @@ private:
 	/** Interacts with the replication system on behalf of Multi-User. */
 	TWeakPtr<UE::MultiUserClient::FMultiUserReplicationManager> WeakReplicationManager;
 
-	/** Kept so it can be passed on to SActiveSessionRoot. */
-	TSharedPtr<SDockTab> ConstructedUnderMajorTab;
+	/**
+	 * Kept so it can be passed on to SActiveSessionRoot.
+	 * Important: since this a pointer to the top-level widget that contains us, we must keep a weak ptr or we'll cause a memory leak.
+	 */
+	TWeakPtr<SDockTab> ConstructedUnderMajorTab;
 
 	/** Keeps the session browser searched text in memory to reapply it when a user leaves a session and goes back to the session browser. */
 	TSharedPtr<FText> SearchedText;
