Fix copy-paste typo causing authority event not to be triggered

diff --git a/Engine/Plugins/Developer/Concert/ConcertApp/MultiUserClient/Source/MultiUserClientLibrary/Private/Replication/MultiUserReplicationSubsystem.cpp b/Engine/Plugins/Developer/Concert/ConcertApp/MultiUserClient/Source/MultiUserClientLibrary/Private/Replication/MultiUserReplicationSubsystem.cpp
--- a/Engine/Plugins/Developer/Concert/ConcertApp/MultiUserClient/Source/MultiUserClientLibrary/Private/Replication/MultiUserReplicationSubsystem.cpp
+++ b/Engine/Plugins/Developer/Concert/ConcertApp/MultiUserClient/Source/MultiUserClientLibrary/Private/Replication/MultiUserReplicationSubsystem.cpp
@@ -136,7 +136,7 @@ void UMultiUserReplicationSubsystem::Initialize(FSubsystemCollectionBase& Collec
 		ReplicationInterface->RegisterReplicationDiscoverer(UObjectAdapter.ToSharedRef());
 
 		ReplicationInterface->OnStreamServerStateChanged().AddUObject(this, &UMultiUserReplicationSubsystem::OnClientStreamsChanged);
-		ReplicationInterface->OnAuthorityServerStateChanged().AddUObject(this, &UMultiUserReplicationSubsystem::OnClientStreamsChanged);
+		ReplicationInterface->OnAuthorityServerStateChanged().AddUObject(this, &UMultiUserReplicationSubsystem::OnClientAuthorityChanged);
 	}
 #endif
 }
