Expose more error information to the replication Blueprint API

diff --git a/Engine/Plugins/Developer/Concert/ConcertApp/MultiUserClient/Source/MultiUserClient/Private/Replication/Submission/Data/AuthoritySubmission.h b/Engine/Plugins/Developer/Concert/ConcertApp/MultiUserClient/Source/MultiUserClient/Private/Replication/Submission/Data/AuthoritySubmission.h
--- a/Engine/Plugins/Developer/Concert/ConcertApp/MultiUserClient/Source/MultiUserClient/Private/Replication/Submission/Data/AuthoritySubmission.h
+++ b/Engine/Plugins/Developer/Concert/ConcertApp/MultiUserClient/Source/MultiUserClient/Private/Replication/Submission/Data/AuthoritySubmission.h
@@ -2,13 +2,14 @@
 
 #pragma once
 
+#include "HAL/Platform.h"
 #include "Replication/IConcertClientReplicationManager.h"
 
 namespace UE::MultiUserClient
 {
 	class ISubmissionOperation;
 	
-	enum class EAuthoritySubmissionRequestErrorCode
+	enum class EAuthoritySubmissionRequestErrorCode : uint8
 	{
 		/** The operation completed as expected */
 		Success,
@@ -23,7 +24,7 @@ namespace UE::MultiUserClient
 		Cancelled
 	};
 
-	enum class EAuthoritySubmissionResponseErrorCode
+	enum class EAuthoritySubmissionResponseErrorCode : uint8
     {
     	/** The operation completed as expected */
     	Success,
diff --git a/Engine/Plugins/Developer/Concert/ConcertApp/MultiUserClient/Source/MultiUserClient/Private/Replication/Submission/Data/StreamSubmission.h b/Engine/Plugins/Developer/Concert/ConcertApp/MultiUserClient/Source/MultiUserClient/Private/Replication/Submission/Data/StreamSubmission.h
--- a/Engine/Plugins/Developer/Concert/ConcertApp/MultiUserClient/Source/MultiUserClient/Private/Replication/Submission/Data/StreamSubmission.h
+++ b/Engine/Plugins/Developer/Concert/ConcertApp/MultiUserClient/Source/MultiUserClient/Private/Replication/Submission/Data/StreamSubmission.h
@@ -2,13 +2,14 @@
 
 #pragma once
 
+#include "HAL/Platform.h"
 #include "Replication/IConcertClientReplicationManager.h"
 
 namespace UE::MultiUserClient
 {
 	class ISubmissionOperation;
 
-	enum class EStreamSubmissionErrorCode
+	enum class EStreamSubmissionErrorCode : uint8
 	{
 		/** Changes were submitted to and processed by the server */
 		Success,
diff --git a/Engine/Plugins/Developer/Concert/ConcertApp/MultiUserClient/Source/MultiUserClient/Private/Replication/Submission/External/ClientChangeConversionUtils.cpp b/Engine/Plugins/Developer/Concert/ConcertApp/MultiUserClient/Source/MultiUserClient/Private/Replication/Submission/External/ClientChangeConversionUtils.cpp
--- a/Engine/Plugins/Developer/Concert/ConcertApp/MultiUserClient/Source/MultiUserClient/Private/Replication/Submission/External/ClientChangeConversionUtils.cpp
+++ b/Engine/Plugins/Developer/Concert/ConcertApp/MultiUserClient/Source/MultiUserClient/Private/Replication/Submission/External/ClientChangeConversionUtils.cpp
@@ -3,12 +3,12 @@
 #include "ClientChangeConversionUtils.h"
 
 #include "Replication/ChangeOperationTypes.h"
+#include "Replication/Submission/Data/AuthoritySubmission.h"
+#include "Replication/Submission/Data/StreamSubmission.h"
 #include "Replication/Util/StreamRequestUtils.h"
 
 #include "Algo/RemoveIf.h"
 #include "Misc/Guid.h"
-#include "Replication/Submission/Data/AuthoritySubmission.h"
-#include "Replication/Submission/Data/StreamSubmission.h"
 
 namespace UE::MultiUserClient::ClientChangeConversionUtils
 {
@@ -153,16 +153,16 @@ namespace UE::MultiUserClient::ClientChangeConversionUtils
 		}
 		return Result;
 	}
-
-	EChangeStreamOperationResult Transform(const FSubmitStreamChangesResponse& Response)
+	
+	EChangeStreamOperationResult ExtractErrorCode(const FSubmitStreamChangesResponse& Response)
 	{
 		switch (Response.ErrorCode)
 		{
 		case EStreamSubmissionErrorCode::Success:
 			// EStreamSubmissionErrorCode::Success the op was processed but the EChangeStreamOperationResult distinguishes between successful or rejected ops.
-			return !ensureMsgf(Response.SubmissionInfo.IsSet(), TEXT("Success implies set response. Investigate.")) || Response.SubmissionInfo->Response.IsFailure()
-				? EChangeStreamOperationResult::Rejected
-				: EChangeStreamOperationResult::Success;
+				return !ensureMsgf(Response.SubmissionInfo.IsSet(), TEXT("Success implies set response. Investigate.")) || Response.SubmissionInfo->Response.IsFailure()
+					? EChangeStreamOperationResult::Rejected
+					: EChangeStreamOperationResult::Success;
 		case EStreamSubmissionErrorCode::NoChange: return EChangeStreamOperationResult::NoChanges; 
 		case EStreamSubmissionErrorCode::Timeout: return EChangeStreamOperationResult::Timeout; 
 		case EStreamSubmissionErrorCode::Cancelled: return EChangeStreamOperationResult::Cancelled;
@@ -172,15 +172,15 @@ namespace UE::MultiUserClient::ClientChangeConversionUtils
 		}
 	}
 	
-	EChangeAuthorityOperationResult Transform(const FSubmitAuthorityChangesResponse& Response)
+	EChangeAuthorityOperationResult ExtractErrorCode(const FSubmitAuthorityChangesResponse& Response)
 	{
 		switch (Response.ErrorCode)
 		{
 		case EAuthoritySubmissionResponseErrorCode::Success:
 			// EAuthoritySubmissionResponseErrorCode::Success the op was processed but the EChangeStreamOperationResult distinguishes between successful or rejected ops.
-			return !ensureMsgf(Response.Response.IsSet(), TEXT("Success implies set response. Investigate.")) || !Response.Response->RejectedObjects.IsEmpty()
-				? EChangeAuthorityOperationResult::RejectedFullyOrPartially
-				: EChangeAuthorityOperationResult::Success;
+				return !ensureMsgf(Response.Response.IsSet(), TEXT("Success implies set response. Investigate.")) || !Response.Response->RejectedObjects.IsEmpty()
+					? EChangeAuthorityOperationResult::RejectedFullyOrPartially
+					: EChangeAuthorityOperationResult::Success;
 		case EAuthoritySubmissionResponseErrorCode::NoChange: return EChangeAuthorityOperationResult::NoChanges; 
 		case EAuthoritySubmissionResponseErrorCode::Timeout: return EChangeAuthorityOperationResult::Timeout; 
 		case EAuthoritySubmissionResponseErrorCode::CancelledDueToStreamUpdate: return EChangeAuthorityOperationResult::CancelledDueToStreamUpdate; 
@@ -190,4 +190,77 @@ namespace UE::MultiUserClient::ClientChangeConversionUtils
 			return EChangeAuthorityOperationResult::Cancelled;
 		}
 	}
+
+	FChangeClientStreamResponse Transform(const FSubmitStreamChangesResponse& Response)
+	{
+		FChangeClientStreamResponse Result { ExtractErrorCode(Response) };
+		
+		if (Response.SubmissionInfo)
+		{
+			FConcertReplication_ChangeStream_Response ContainedResponse = Response.SubmissionInfo->Response;
+			
+			for (const TPair<FConcertObjectInStreamID, FConcertReplicatedObjectId>& Pair : ContainedResponse.AuthorityConflicts)
+			{
+				Result.AuthorityConflicts.Add(Pair.Key.Object, Pair.Value.SenderEndpointId);
+			}
+
+			for (const TPair<FConcertObjectInStreamID, EConcertPutObjectErrorCode>& Pair : ContainedResponse.ObjectsToPutSemanticErrors)
+			{
+				Result.SemanticErrors.Add(Pair.Key.Object, Transform(Pair.Value));
+			}
+
+			for (const TPair<FConcertObjectInStreamID, EConcertChangeObjectFrequencyErrorCode>& Pair : ContainedResponse.FrequencyErrors.OverrideFailures)
+			{
+				Result.FrequencyErrors.ObjectErrors.Add(Pair.Key.Object, Transform(Pair.Value));
+			}
+			if (!ContainedResponse.FrequencyErrors.DefaultFailures.IsEmpty())
+			{
+				ensure(ContainedResponse.FrequencyErrors.DefaultFailures.Num() == 1);
+				Result.FrequencyErrors.DefaultChangeErrorCode = Transform(ContainedResponse.FrequencyErrors.DefaultFailures.CreateIterator().Value());
+			}
+
+			Result.bFailedStreamCreation = !ContainedResponse.FailedStreamCreation.IsEmpty();
+		}
+		
+		return Result;
+	}
+	
+	FChangeClientAuthorityResponse Transform(const FSubmitAuthorityChangesResponse& Response)
+	{
+		FChangeClientAuthorityResponse Result{ ExtractErrorCode(Response) };
+
+		if (Response.Response)
+		{
+			for (const TPair<FSoftObjectPath, FConcertStreamArray>& Pair : Response.Response->RejectedObjects)
+			{
+				Result.RejectedObjects.Add(Pair.Key);
+			}
+		}
+		
+		return Result;
+	}
+
+	EPutObjectErrorCode Transform(EConcertPutObjectErrorCode ErrorCode)
+	{
+		static_assert(static_cast<int32>(EPutObjectErrorCode::MissingData) == static_cast<int32>(EConcertPutObjectErrorCode::MissingData), "Update this code when making modification to the enums");
+		static_assert(static_cast<int32>(EPutObjectErrorCode::UnresolvedStream) == static_cast<int32>(EConcertPutObjectErrorCode::UnresolvedStream), "Update this code when making modification to the enums");
+		static_assert(static_cast<int32>(EPutObjectErrorCode::Count) == static_cast<int32>(EConcertPutObjectErrorCode::Count), "Update this code when making modification to the enums");
+		return static_cast<EPutObjectErrorCode>(ErrorCode);
+	}
+	
+	EChangeObjectFrequencyErrorCode Transform(EConcertChangeObjectFrequencyErrorCode ErrorCode)
+	{
+		static_assert(static_cast<int32>(EChangeObjectFrequencyErrorCode::UnregisteredStream) == static_cast<int32>(EConcertChangeObjectFrequencyErrorCode::NotRegistered), "Update this code when making modification to the enums");
+		static_assert(static_cast<int32>(EChangeObjectFrequencyErrorCode::InvalidReplicationRate) == static_cast<int32>(EConcertChangeObjectFrequencyErrorCode::InvalidReplicationRate), "Update this code when making modification to the enums");
+		static_assert(static_cast<int32>(EChangeObjectFrequencyErrorCode::Count) == static_cast<int32>(EConcertChangeObjectFrequencyErrorCode::Count), "Update this code when making modification to the enums");
+		return static_cast<EChangeObjectFrequencyErrorCode>(ErrorCode);
+	}
+
+	EChangeObjectFrequencyErrorCode Transform(EConcertChangeStreamFrequencyErrorCode ErrorCode)
+	{
+		static_assert(static_cast<int32>(EChangeObjectFrequencyErrorCode::UnregisteredStream) == static_cast<int32>(EConcertChangeStreamFrequencyErrorCode::UnknownStream), "Update this code when making modification to the enums");
+        static_assert(static_cast<int32>(EChangeObjectFrequencyErrorCode::InvalidReplicationRate) == static_cast<int32>(EConcertChangeStreamFrequencyErrorCode::InvalidReplicationRate), "Update this code when making modification to the enums");
+        static_assert(static_cast<int32>(EChangeObjectFrequencyErrorCode::Count) == static_cast<int32>(EConcertChangeStreamFrequencyErrorCode::Count), "Update this code when making modification to the enums");
+        return static_cast<EChangeObjectFrequencyErrorCode>(ErrorCode);
+	}
 }
diff --git a/Engine/Plugins/Developer/Concert/ConcertApp/MultiUserClient/Source/MultiUserClient/Private/Replication/Submission/External/ClientChangeConversionUtils.h b/Engine/Plugins/Developer/Concert/ConcertApp/MultiUserClient/Source/MultiUserClient/Private/Replication/Submission/External/ClientChangeConversionUtils.h
--- a/Engine/Plugins/Developer/Concert/ConcertApp/MultiUserClient/Source/MultiUserClient/Private/Replication/Submission/External/ClientChangeConversionUtils.h
+++ b/Engine/Plugins/Developer/Concert/ConcertApp/MultiUserClient/Source/MultiUserClient/Private/Replication/Submission/External/ClientChangeConversionUtils.h
@@ -6,6 +6,7 @@
 #include "Replication/Messages/ChangeAuthority.h"
 #include "Replication/Messages/ChangeStream.h"
 
+#include "HAL/Platform.h"
 #include "Misc/Optional.h"
 
 struct FConcertObjectReplicationMap;
@@ -13,9 +14,10 @@ struct FGuid;
 
 namespace UE::MultiUserClient
 {
-	enum class EChangeStreamOperationResult : uint8;
-	enum class EChangeAuthorityOperationResult : uint8;
+	enum class EAuthoritySubmissionResponseErrorCode : uint8;
 	
+	struct FChangeClientAuthorityResponse;
+	struct FChangeClientStreamResponse;
 	struct FSubmitAuthorityChangesResponse;
 	struct FSubmitStreamChangesResponse;
 }
@@ -38,12 +40,15 @@ namespace UE::MultiUserClient::ClientChangeConversionUtils
 		const FGuid& ClientStreamId
 		);
 
+	EChangeStreamOperationResult ExtractErrorCode(const FSubmitStreamChangesResponse& Response);
+	EChangeAuthorityOperationResult ExtractErrorCode(const FSubmitAuthorityChangesResponse& Response);
+	
+	FChangeClientStreamResponse Transform(const FSubmitStreamChangesResponse& Response);
+	FChangeClientAuthorityResponse Transform(const FSubmitAuthorityChangesResponse& Response);
 
-	/** Extract an error code from FSubmitStreamChangesResponse. */
-	EChangeStreamOperationResult Transform(const FSubmitStreamChangesResponse& Response);
-	/** Extract an error code from FSubmitAuthorityChangesResponse. */
-	EChangeAuthorityOperationResult Transform(const FSubmitAuthorityChangesResponse& Response);
-
+	EChangeObjectFrequencyErrorCode Transform(EConcertChangeObjectFrequencyErrorCode ErrorCode);
+	EChangeObjectFrequencyErrorCode Transform(EConcertChangeStreamFrequencyErrorCode ErrorCode);
+	EPutObjectErrorCode Transform(EConcertPutObjectErrorCode ErrorCode);
 
 	inline TOptional<FConcertReplication_ChangeStream_Request> Transform(
 		TOptional<FChangeStreamRequest> Request,
diff --git a/Engine/Plugins/Developer/Concert/ConcertApp/MultiUserClient/Source/MultiUserClient/Private/Replication/Submission/External/ClientChangeOperation.cpp b/Engine/Plugins/Developer/Concert/ConcertApp/MultiUserClient/Source/MultiUserClient/Private/Replication/Submission/External/ClientChangeOperation.cpp
--- a/Engine/Plugins/Developer/Concert/ConcertApp/MultiUserClient/Source/MultiUserClient/Private/Replication/Submission/External/ClientChangeOperation.cpp
+++ b/Engine/Plugins/Developer/Concert/ConcertApp/MultiUserClient/Source/MultiUserClient/Private/Replication/Submission/External/ClientChangeOperation.cpp
@@ -48,12 +48,12 @@ namespace UE::MultiUserClient
 		const bool bIsAlreadyCompleted = StreamResult && AuthorityResult;
 		if (!StreamResult)
 		{
-			EmplaceStreamResult(EChangeStreamOperationResult::Cancelled);
+			EmplaceStreamResult(FChangeClientStreamResponse{ EChangeStreamOperationResult::Cancelled });
 		}
 		if (!AuthorityResult)
 		{
 			// Also executes OnOperationCompleted
-			EmplaceAuthorityResult(EChangeAuthorityOperationResult::Cancelled);
+			EmplaceAuthorityResult(FChangeClientAuthorityResponse{ EChangeAuthorityOperationResult::Cancelled });
 		}
 		if (!bIsAlreadyCompleted)
 		{
@@ -68,8 +68,8 @@ namespace UE::MultiUserClient
 		FChangeClientReplicationRequest Request = MakeSubmissionParamsAttribute.Get();
 		if (!Request.StreamChangeRequest && !Request.AuthorityChangeRequest)
 		{
-			EmplaceStreamResult(EChangeStreamOperationResult::NoChanges);
-			EmplaceAuthorityResult(EChangeAuthorityOperationResult::NoChanges);
+			EmplaceStreamResult(FChangeClientStreamResponse{ EChangeStreamOperationResult::NoChanges });
+			EmplaceAuthorityResult(FChangeClientAuthorityResponse{ EChangeAuthorityOperationResult::NoChanges });
 			CompleteOperation();
 		}
 		else
@@ -87,8 +87,8 @@ namespace UE::MultiUserClient
 		if (!SubmissionOperation)
 		{
 			UE_LOG(LogConcert, Error, TEXT("Failed to submit request to server. Check log for additional errors."));
-			EmplaceStreamResult(EChangeStreamOperationResult::FailedToSendRequest);
-			EmplaceAuthorityResult(EChangeAuthorityOperationResult::FailedToSendRequest);
+			EmplaceStreamResult(FChangeClientStreamResponse{ EChangeStreamOperationResult::FailedToSendRequest });
+			EmplaceAuthorityResult(FChangeClientAuthorityResponse{ EChangeAuthorityOperationResult::FailedToSendRequest });
 			CompleteOperation();
 			return;
 		}
@@ -101,7 +101,7 @@ namespace UE::MultiUserClient
 			{
 				if (const TSharedPtr<FClientChangeOperation> ThisPin = WeakThis.Pin())
 				{
-					ThisPin->EmplaceStreamResult(ClientChangeConversionUtils::Transform(Response));
+					ThisPin->EmplaceStreamResult(FChangeClientStreamResponse{ ClientChangeConversionUtils::Transform(Response) });
 				}
 			});
 		SubmissionOperation->OnCompleteAuthorityChangeFuture_AnyThread()
@@ -109,7 +109,7 @@ namespace UE::MultiUserClient
 			{
 				if (const TSharedPtr<FClientChangeOperation> ThisPin = WeakThis.Pin())
 				{
-					ThisPin->EmplaceAuthorityResult(ClientChangeConversionUtils::Transform(Response));
+					ThisPin->EmplaceAuthorityResult(FChangeClientAuthorityResponse{ ClientChangeConversionUtils::Transform(Response) });
 				}
 			});
 		SubmissionOperation->OnCompletedOperation_AnyThread()
diff --git a/Engine/Plugins/Developer/Concert/ConcertApp/MultiUserClient/Source/MultiUserClient/Private/Replication/Submission/External/ClientChangeOperation.h b/Engine/Plugins/Developer/Concert/ConcertApp/MultiUserClient/Source/MultiUserClient/Private/Replication/Submission/External/ClientChangeOperation.h
--- a/Engine/Plugins/Developer/Concert/ConcertApp/MultiUserClient/Source/MultiUserClient/Private/Replication/Submission/External/ClientChangeOperation.h
+++ b/Engine/Plugins/Developer/Concert/ConcertApp/MultiUserClient/Source/MultiUserClient/Private/Replication/Submission/External/ClientChangeOperation.h
@@ -53,8 +53,8 @@ namespace UE::MultiUserClient
 		virtual ~FClientChangeOperation() override;
 
 		//~ Begin IClientChangeOperation Interface
-		virtual TFuture<EChangeStreamOperationResult> OnChangeStream() override { return ChangeStreamPromise.GetFuture(); }
-		virtual TFuture<EChangeAuthorityOperationResult> OnChangeAuthority() override { return ChangeAuthorityPromise.GetFuture(); }
+		virtual TFuture<FChangeClientStreamResponse> OnChangeStream() override { return ChangeStreamPromise.GetFuture(); }
+		virtual TFuture<FChangeClientAuthorityResponse> OnChangeAuthority() override { return ChangeAuthorityPromise.GetFuture(); }
 		virtual TFuture<FChangeClientReplicationResult> OnOperationCompleted() override { return OperationCompletedPromise.GetFuture(); }
 		//~ End IClientChangeOperation Interface
 
@@ -78,14 +78,14 @@ namespace UE::MultiUserClient
 		/** Callback to build submission params when task is ready to be executed. */
 		const TAttribute<FChangeClientReplicationRequest> MakeSubmissionParamsAttribute;
 
-		TPromise<EChangeStreamOperationResult> ChangeStreamPromise;
-		TPromise<EChangeAuthorityOperationResult> ChangeAuthorityPromise;
+		TPromise<FChangeClientStreamResponse> ChangeStreamPromise;
+		TPromise<FChangeClientAuthorityResponse> ChangeAuthorityPromise;
 		TPromise<FChangeClientReplicationResult> OperationCompletedPromise;
 
 		/** Set when ChangeStreamPromise is fulfilled. */
-		TOptional<EChangeStreamOperationResult> StreamResult;
+		TOptional<FChangeClientStreamResponse> StreamResult;
 		/** Set when ChangeAuthorityPromise is fulfilled. */
-		TOptional<EChangeAuthorityOperationResult> AuthorityResult;
+		TOptional<FChangeClientAuthorityResponse> AuthorityResult;
 
 		FClientChangeOperation(
 			const FGuid& InStreamId,
@@ -94,20 +94,20 @@ namespace UE::MultiUserClient
 			TAttribute<FChangeClientReplicationRequest> InMakeSubmissionParamsAttribute
 			);
 
-		void EmplaceStreamResult(EChangeStreamOperationResult Result)
+		void EmplaceStreamResult(FChangeClientStreamResponse Result)
 		{
 			if (ensure(!StreamResult))
 			{
-				StreamResult = Result;
-				ChangeStreamPromise.EmplaceValue(Result);
+				StreamResult = MoveTemp(Result);
+				ChangeStreamPromise.EmplaceValue(*StreamResult);
 			}
 		}
-		void EmplaceAuthorityResult(EChangeAuthorityOperationResult Result)
+		void EmplaceAuthorityResult(FChangeClientAuthorityResponse Result)
 		{
 			if (ensure(!AuthorityResult))
 			{
-				AuthorityResult = Result;
-				ChangeAuthorityPromise.EmplaceValue(Result);
+				AuthorityResult = MoveTemp(Result);
+				ChangeAuthorityPromise.EmplaceValue(*AuthorityResult);
 			}
 		}
 		void CompleteOperation()
diff --git a/Engine/Plugins/Developer/Concert/ConcertApp/MultiUserClient/Source/MultiUserClient/Private/Replication/Submission/External/ExternalClientChangeRequestHandler.cpp b/Engine/Plugins/Developer/Concert/ConcertApp/MultiUserClient/Source/MultiUserClient/Private/Replication/Submission/External/ExternalClientChangeRequestHandler.cpp
--- a/Engine/Plugins/Developer/Concert/ConcertApp/MultiUserClient/Source/MultiUserClient/Private/Replication/Submission/External/ExternalClientChangeRequestHandler.cpp
+++ b/Engine/Plugins/Developer/Concert/ConcertApp/MultiUserClient/Source/MultiUserClient/Private/Replication/Submission/External/ExternalClientChangeRequestHandler.cpp
@@ -18,21 +18,21 @@ namespace UE::MultiUserClient
 	{
 		class FNotInSessionOperation : public IClientChangeOperation
 		{
-			TPromise<EChangeStreamOperationResult> ChangeStreamPromise;
-			TPromise<EChangeAuthorityOperationResult> ChangeAuthorityPromise;
+			TPromise<FChangeClientStreamResponse> ChangeStreamPromise;
+			TPromise<FChangeClientAuthorityResponse> ChangeAuthorityPromise;
 			TPromise<FChangeClientReplicationResult> OperationCompletedPromise;
 			
 		public:
 
-			FNotInSessionOperation(EChangeStreamOperationResult StreamResult, EChangeAuthorityOperationResult AuthorityResult)
-				: ChangeStreamPromise(MakeFulfilledPromise<EChangeStreamOperationResult>(StreamResult))
-				, ChangeAuthorityPromise(MakeFulfilledPromise<EChangeAuthorityOperationResult>(AuthorityResult))
-				, OperationCompletedPromise(MakeFulfilledPromise<FChangeClientReplicationResult>(FChangeClientReplicationResult{ StreamResult, AuthorityResult }))
+			FNotInSessionOperation(EChangeStreamOperationResult StreamErrorCode, EChangeAuthorityOperationResult AuthorityErrorCode)
+				: ChangeStreamPromise(MakeFulfilledPromise<FChangeClientStreamResponse>(FChangeClientStreamResponse{ StreamErrorCode }))
+				, ChangeAuthorityPromise(MakeFulfilledPromise<FChangeClientAuthorityResponse>(FChangeClientAuthorityResponse{ AuthorityErrorCode }))
+				, OperationCompletedPromise(MakeFulfilledPromise<FChangeClientReplicationResult>(FChangeClientReplicationResult{ { StreamErrorCode }, { AuthorityErrorCode } }))
 			{}
 			
 			//~ Begin IClientChangeOperation Interface
-			virtual TFuture<EChangeStreamOperationResult> OnChangeStream() override { return ChangeStreamPromise.GetFuture(); }
-			virtual TFuture<EChangeAuthorityOperationResult> OnChangeAuthority() override { return ChangeAuthorityPromise.GetFuture(); }
+			virtual TFuture<FChangeClientStreamResponse> OnChangeStream() override { return ChangeStreamPromise.GetFuture(); }
+			virtual TFuture<FChangeClientAuthorityResponse> OnChangeAuthority() override { return ChangeAuthorityPromise.GetFuture(); }
 			virtual TFuture<FChangeClientReplicationResult> OnOperationCompleted() override { return OperationCompletedPromise.GetFuture(); }
 			//~ End IClientChangeOperation Interface
 		};
diff --git a/Engine/Plugins/Developer/Concert/ConcertApp/MultiUserClient/Source/MultiUserClient/Private/Replication/Submission/Remote/SubmissionWorkflow_RemoteClient.cpp b/Engine/Plugins/Developer/Concert/ConcertApp/MultiUserClient/Source/MultiUserClient/Private/Replication/Submission/Remote/SubmissionWorkflow_RemoteClient.cpp
--- a/Engine/Plugins/Developer/Concert/ConcertApp/MultiUserClient/Source/MultiUserClient/Private/Replication/Submission/Remote/SubmissionWorkflow_RemoteClient.cpp
+++ b/Engine/Plugins/Developer/Concert/ConcertApp/MultiUserClient/Source/MultiUserClient/Private/Replication/Submission/Remote/SubmissionWorkflow_RemoteClient.cpp
@@ -188,6 +188,8 @@ namespace UE::MultiUserClient
 		FSingleClientSubmissionOperation& Operation = GetOperation();
 		if (EventData.bProcessedSuccessfully)
 		{
+			LogStreamErrorsIfNeeded(EventData.StreamChangeResponse);
+			
 			const FCompletedChangeSubmission StreamChangeSubmission { *InProgressOperation->Parameters.StreamRequest, { EventData.StreamChangeResponse } };
 			const FSubmitStreamChangesResponse StreamResponse { EStreamSubmissionErrorCode::Success, StreamChangeSubmission };
 			Operation.EmplaceStreamPromise(StreamResponse);
@@ -247,4 +249,15 @@ namespace UE::MultiUserClient
 	{
 		return ClientUtils::GetClientDisplayName(*ConcertSession, RemoteClientEndpointId);
 	}
+
+	void FSubmissionWorkflow_RemoteClient::LogStreamErrorsIfNeeded(const FConcertReplication_ChangeStream_Response& Response) const
+	{
+		if (Response.IsFailure())
+		{
+			FStringOutputDevice OutputDevice;
+			Response.LogErrors(OutputDevice);
+			const FString ClientName = ClientUtils::GetClientDisplayName(*ConcertSession, RemoteClientEndpointId);
+			UE_LOG(LogConcert, Warning, TEXT("Submission failed for remote client %s. Errors: %s"), *ClientName, *OutputDevice);
+		}
+	}
 }
diff --git a/Engine/Plugins/Developer/Concert/ConcertApp/MultiUserClient/Source/MultiUserClient/Private/Replication/Submission/Remote/SubmissionWorkflow_RemoteClient.h b/Engine/Plugins/Developer/Concert/ConcertApp/MultiUserClient/Source/MultiUserClient/Private/Replication/Submission/Remote/SubmissionWorkflow_RemoteClient.h
--- a/Engine/Plugins/Developer/Concert/ConcertApp/MultiUserClient/Source/MultiUserClient/Private/Replication/Submission/Remote/SubmissionWorkflow_RemoteClient.h
+++ b/Engine/Plugins/Developer/Concert/ConcertApp/MultiUserClient/Source/MultiUserClient/Private/Replication/Submission/Remote/SubmissionWorkflow_RemoteClient.h
@@ -82,6 +82,8 @@ namespace UE::MultiUserClient
 		void FinishSubmissionIfDone();
 
 		FString GetRemoteClientName() const;
+		
+		void LogStreamErrorsIfNeeded(const FConcertReplication_ChangeStream_Response& Response) const;
 	};
 }
 
diff --git a/Engine/Plugins/Developer/Concert/ConcertApp/MultiUserClient/Source/MultiUserClient/Private/Replication/Submission/SubmissionWorkflow_LocalClient.cpp b/Engine/Plugins/Developer/Concert/ConcertApp/MultiUserClient/Source/MultiUserClient/Private/Replication/Submission/SubmissionWorkflow_LocalClient.cpp
--- a/Engine/Plugins/Developer/Concert/ConcertApp/MultiUserClient/Source/MultiUserClient/Private/Replication/Submission/SubmissionWorkflow_LocalClient.cpp
+++ b/Engine/Plugins/Developer/Concert/ConcertApp/MultiUserClient/Source/MultiUserClient/Private/Replication/Submission/SubmissionWorkflow_LocalClient.cpp
@@ -4,6 +4,7 @@
 
 #include "IConcertSyncClient.h"
 #include "Replication/IConcertClientReplicationManager.h"
+#include "Widgets/ActiveSession/Replication/Client/ClientUtils.h"
 
 namespace UE::MultiUserClient
 {
@@ -76,6 +77,7 @@ namespace UE::MultiUserClient
 		)
 	{
 		const TSharedRef<FSingleClientSubmissionOperation> Operation = InProgressOperation->Operation;
+		LogStreamErrorsIfNeeded(ChangeStreamResponse);
 		
 		const EStreamSubmissionErrorCode ErrorCode = ChangeStreamResponse.ErrorCode == EReplicationResponseErrorCode::Handled
 			? EStreamSubmissionErrorCode::Success
@@ -153,4 +155,16 @@ namespace UE::MultiUserClient
 		Operation->EmplaceCompleteOperationPromise(ESubmissionOperationCompletedCode::Processed);
 		OnSubmitOperationCompletedDelegate.Broadcast();
 	}
+
+	void FSubmissionWorkflow_LocalClient::LogStreamErrorsIfNeeded(const FConcertReplication_ChangeStream_Response& Response) const
+	{
+		const TSharedPtr<IConcertClientSession> CurrentSession = Client->GetConcertClient()->GetCurrentSession();
+		if (Response.IsFailure() && CurrentSession)
+		{
+			FStringOutputDevice OutputDevice;
+			Response.LogErrors(OutputDevice);
+			const FString ClientName = ClientUtils::GetClientDisplayName(*CurrentSession, CurrentSession->GetSessionClientEndpointId());
+			UE_LOG(LogConcert, Warning, TEXT("Submission failed for local client %s. Errors: %s"), *ClientName, *OutputDevice);
+		}
+	}
 }
diff --git a/Engine/Plugins/Developer/Concert/ConcertApp/MultiUserClient/Source/MultiUserClient/Private/Replication/Submission/SubmissionWorkflow_LocalClient.h b/Engine/Plugins/Developer/Concert/ConcertApp/MultiUserClient/Source/MultiUserClient/Private/Replication/Submission/SubmissionWorkflow_LocalClient.h
--- a/Engine/Plugins/Developer/Concert/ConcertApp/MultiUserClient/Source/MultiUserClient/Private/Replication/Submission/SubmissionWorkflow_LocalClient.h
+++ b/Engine/Plugins/Developer/Concert/ConcertApp/MultiUserClient/Source/MultiUserClient/Private/Replication/Submission/SubmissionWorkflow_LocalClient.h
@@ -61,6 +61,8 @@ namespace UE::MultiUserClient
 		
 		void SkipAuthorityStage(EAuthoritySubmissionRequestErrorCode RequestCode, EAuthoritySubmissionResponseErrorCode ResponseCode);
 		void CleanUpSubmissionOperation();
+
+		void LogStreamErrorsIfNeeded(const FConcertReplication_ChangeStream_Response& Response) const;
 	};
 }
 
diff --git a/Engine/Plugins/Developer/Concert/ConcertApp/MultiUserClient/Source/MultiUserClient/Public/Replication/ChangeOperationTypes.h b/Engine/Plugins/Developer/Concert/ConcertApp/MultiUserClient/Source/MultiUserClient/Public/Replication/ChangeOperationTypes.h
--- a/Engine/Plugins/Developer/Concert/ConcertApp/MultiUserClient/Source/MultiUserClient/Public/Replication/ChangeOperationTypes.h
+++ b/Engine/Plugins/Developer/Concert/ConcertApp/MultiUserClient/Source/MultiUserClient/Public/Replication/ChangeOperationTypes.h
@@ -146,12 +146,91 @@ namespace UE::MultiUserClient
 		 */
 		TOptional<FChangeAuthorityRequest> AuthorityChangeRequest;
 	};
+	
+	enum class EChangeObjectFrequencyErrorCode : uint8
+	{
+		/** The object for which the frequency was being changed was not registered. */
+		UnregisteredStream,
+		/** The replication rate parameter was rejected (it cannot be 0). */
+		InvalidReplicationRate,
+
+		/** Not an actual parameter. Make sure it's always last. */
+		Count
+	};
+
+	/** Result about changing frequency. */
+	struct FChangeClientStreamFrequencyResponse
+	{
+		/** Errors encountered for specific objects */
+		TMap<FSoftObjectPath, EChangeObjectFrequencyErrorCode> ObjectErrors;
+		/** Error for changing the default replication frequency. */
+		TOptional<EChangeObjectFrequencyErrorCode> DefaultChangeErrorCode;
+	};
+
+	/** Explains why a change to an object in the stream was invalid. */
+	enum class EPutObjectErrorCode : uint8
+	{
+		/** Stream that the object referenced was not registered on the server. */
+		UnresolvedStream,
+		/**
+		 * Either PutObject contained no data to update with (ensure either ClassPath or Properties is set),
+		 * or it tried to create a new object with insufficient data (make sure ClassPath and Properties are both specified).
+		 */
+		MissingData,
+
+		/** Not an actual parameter. Make sure it's always last. */
+		Count
+	};
+
+	/** Result of processing FChangeStreamRequest. */
+	struct FChangeClientStreamResponse
+	{
+		/**
+		 * Gives general information about what happened to this request.
+		 * All other fields only make sense if ErrorCode == Success.
+		 */
+		EChangeStreamOperationResult ErrorCode;
+		
+		/**
+		 * Dynamic authority errors.
+		 * The change was rejected because
+		 * 1. this client is replicating the object already,
+		 * 2. another client is also replicating the object,
+		 * 3. this change would cause overlapping properties with the other client.
+		 * */
+		TMap<FSoftObjectPath, FGuid> AuthorityConflicts;
+		
+		/** Errors made in the format of the request */
+		TMap<FSoftObjectPath, EPutObjectErrorCode> SemanticErrors;
+
+		/** Errors made in the way frequency was changed. */
+		FChangeClientStreamFrequencyResponse FrequencyErrors;
+
+		/**
+		 * The client attempted to register the Multi-User stream but failed in doing so.
+		 * This usually indicates an internal error that you cannot do anything about as API user.
+		 */
+		bool bFailedStreamCreation = false;
+	};
+
+	/** Result of processing FChangeAuthorityRequest */
+	struct FChangeClientAuthorityResponse
+	{
+		/**
+		 * Gives general information about what happened to this request.
+		 * All other fields only make sense if ErrorCode == Success.
+		 */
+		EChangeAuthorityOperationResult ErrorCode;
+		
+		/** Objects the client did not get authority over. */
+		TSet<FSoftObjectPath> RejectedObjects;
+	};
 
 	/** Result of processing a FChangeClientReplicationRequest. */
 	struct FChangeClientReplicationResult
 	{
-		EChangeStreamOperationResult StreamChangeResult;
-		EChangeAuthorityOperationResult AuthorityChangeResult;
+		FChangeClientStreamResponse StreamChangeResult;
+		FChangeClientAuthorityResponse AuthorityChangeResult;
 	};
 }
 
diff --git a/Engine/Plugins/Developer/Concert/ConcertApp/MultiUserClient/Source/MultiUserClient/Public/Replication/IClientChangeOperation.h b/Engine/Plugins/Developer/Concert/ConcertApp/MultiUserClient/Source/MultiUserClient/Public/Replication/IClientChangeOperation.h
--- a/Engine/Plugins/Developer/Concert/ConcertApp/MultiUserClient/Source/MultiUserClient/Public/Replication/IClientChangeOperation.h
+++ b/Engine/Plugins/Developer/Concert/ConcertApp/MultiUserClient/Source/MultiUserClient/Public/Replication/IClientChangeOperation.h
@@ -24,14 +24,14 @@ namespace UE::MultiUserClient
 		 * @note This can be called at most once; subsequent calls result in an unset future.
 		 * @note This future can complete on any thread.
 		 */
-		virtual TFuture<EChangeStreamOperationResult> OnChangeStream() = 0;
+		virtual TFuture<FChangeClientStreamResponse> OnChangeStream() = 0;
 		
 		/**
 		 * Completes when the operation of changing authority has completed.
 		 * @note This can be called at most once; subsequent calls result in an unset future.
 		 * @note This future can complete on any thread.
 		 */
-		virtual TFuture<EChangeAuthorityOperationResult> OnChangeAuthority() = 0;
+		virtual TFuture<FChangeClientAuthorityResponse> OnChangeAuthority() = 0;
 		
 		/**
 		 * Completes all sub-operations have completed.
diff --git a/Engine/Plugins/Developer/Concert/ConcertApp/MultiUserClient/Source/MultiUserClientLibrary/Private/Replication/Async/ChangeClientAsyncAction.cpp b/Engine/Plugins/Developer/Concert/ConcertApp/MultiUserClient/Source/MultiUserClientLibrary/Private/Replication/Async/ChangeClientAsyncAction.cpp
--- a/Engine/Plugins/Developer/Concert/ConcertApp/MultiUserClient/Source/MultiUserClientLibrary/Private/Replication/Async/ChangeClientAsyncAction.cpp
+++ b/Engine/Plugins/Developer/Concert/ConcertApp/MultiUserClient/Source/MultiUserClientLibrary/Private/Replication/Async/ChangeClientAsyncAction.cpp
@@ -57,6 +57,6 @@ void UChangeClientAsyncAction::Activate()
 	else
 #endif
 	{
-		OnCompleted.Broadcast({ EMultiUserChangeStreamOperationResult::NotAvailable, EMultiUserChangeAuthorityOperationResult::NotAvailable });
+		OnCompleted.Broadcast({ { EMultiUserChangeStreamOperationResult::NotAvailable }, { EMultiUserChangeAuthorityOperationResult::NotAvailable } });
 	}
 }
diff --git a/Engine/Plugins/Developer/Concert/ConcertApp/MultiUserClient/Source/MultiUserClientLibrary/Private/Replication/Async/ChangeClientBlueprintParams.cpp b/Engine/Plugins/Developer/Concert/ConcertApp/MultiUserClient/Source/MultiUserClientLibrary/Private/Replication/Async/ChangeClientBlueprintParams.cpp
--- a/Engine/Plugins/Developer/Concert/ConcertApp/MultiUserClient/Source/MultiUserClientLibrary/Private/Replication/Async/ChangeClientBlueprintParams.cpp
+++ b/Engine/Plugins/Developer/Concert/ConcertApp/MultiUserClient/Source/MultiUserClientLibrary/Private/Replication/Async/ChangeClientBlueprintParams.cpp
@@ -15,6 +15,49 @@ namespace UE::MultiUserClientLibrary
 		static_assert(static_cast<int32>(MultiUserClient::EChangeAuthorityOperationResult::Count) == 10, "Update EMultiUserChangeStreamOperationResult to have the equivalent enum entry you added to EChangeStreamOperationResult");
 		return static_cast<EMultiUserChangeAuthorityOperationResult>(Data);
 	}
+
+	EMultiUserChangeFrequencyErrorCode Transform(MultiUserClient::EChangeObjectFrequencyErrorCode Data)
+	{
+		static_assert(static_cast<int32>(EMultiUserChangeFrequencyErrorCode::UnregisteredStream) == static_cast<int32>(MultiUserClient::EChangeObjectFrequencyErrorCode::UnregisteredStream), "Update this code when making modification to the enums");
+		static_assert(static_cast<int32>(EMultiUserChangeFrequencyErrorCode::InvalidReplicationRate) == static_cast<int32>(MultiUserClient::EChangeObjectFrequencyErrorCode::InvalidReplicationRate), "Update this code when making modification to the enums");
+		static_assert(static_cast<int32>(EMultiUserChangeFrequencyErrorCode::Count) == static_cast<int32>(MultiUserClient::EChangeObjectFrequencyErrorCode::Count), "Update this code when making modification to the enums");
+		return static_cast<EMultiUserChangeFrequencyErrorCode>(Data);
+	}
+	
+	EMultiUserPutObjectErrorCode Transform(MultiUserClient::EPutObjectErrorCode Data)
+	{
+		static_assert(static_cast<int32>(EMultiUserPutObjectErrorCode::UnresolvedStream) == static_cast<int32>(MultiUserClient::EPutObjectErrorCode::UnresolvedStream), "Update this code when making modification to the enums");
+		static_assert(static_cast<int32>(EMultiUserPutObjectErrorCode::MissingData) == static_cast<int32>(MultiUserClient::EPutObjectErrorCode::MissingData), "Update this code when making modification to the enums");
+		static_assert(static_cast<int32>(EMultiUserPutObjectErrorCode::Count) == static_cast<int32>(MultiUserClient::EPutObjectErrorCode::Count), "Update this code when making modification to the enums");
+		return static_cast<EMultiUserPutObjectErrorCode>(Data);
+	}
+
+	FMultiUserChangeClientStreamResponse Transform(MultiUserClient::FChangeClientStreamResponse Data)
+	{
+		FMultiUserChangeClientStreamResponse Response { Transform(Data.ErrorCode), MoveTemp(Data.AuthorityConflicts) };
+
+		for (const TPair<FSoftObjectPath, MultiUserClient::EPutObjectErrorCode>& Pair : Data.SemanticErrors)
+		{
+			Response.SemanticErrors.Add(Pair.Key, Transform(Pair.Value));
+		}
+		
+		for (const TPair<FSoftObjectPath, MultiUserClient::EChangeObjectFrequencyErrorCode>& Pair : Data.FrequencyErrors.ObjectErrors)
+		{
+			Response.FrequencyErrors.ObjectErrors.Add(Pair.Key, Transform(Pair.Value));
+		}
+		
+		if (Data.FrequencyErrors.DefaultChangeErrorCode)
+		{
+			Response.FrequencyErrors.DefaultChangeErrorCode = Transform(*Data.FrequencyErrors.DefaultChangeErrorCode);
+		}
+			
+		Response.bFailedStreamCreation = Data.bFailedStreamCreation;
+		return Response;
+	}
+	FMultiUserChangeClientAuthorityResponse Transform(MultiUserClient::FChangeClientAuthorityResponse Data)
+	{
+		return FMultiUserChangeClientAuthorityResponse{ Transform(Data.ErrorCode), MoveTemp(Data.RejectedObjects) };
+	}
 	FMultiUserChangeClientReplicationResult Transform(MultiUserClient::FChangeClientReplicationResult Data)
 	{
 		return { Transform(Data.StreamChangeResult), Transform(Data.AuthorityChangeResult) };
diff --git a/Engine/Plugins/Developer/Concert/ConcertApp/MultiUserClient/Source/MultiUserClientLibrary/Public/Replication/Async/ChangeClientBlueprintParams.h b/Engine/Plugins/Developer/Concert/ConcertApp/MultiUserClient/Source/MultiUserClientLibrary/Public/Replication/Async/ChangeClientBlueprintParams.h
--- a/Engine/Plugins/Developer/Concert/ConcertApp/MultiUserClient/Source/MultiUserClientLibrary/Public/Replication/Async/ChangeClientBlueprintParams.h
+++ b/Engine/Plugins/Developer/Concert/ConcertApp/MultiUserClient/Source/MultiUserClientLibrary/Public/Replication/Async/ChangeClientBlueprintParams.h
@@ -217,19 +217,117 @@ struct FMultiUserChangeClientReplicationRequest
 	FMultiUserChangeAuthorityRequest AuthorityChangeRequest;
 };
 
+UENUM(BlueprintType)
+enum class EMultiUserChangeFrequencyErrorCode : uint8
+{
+	/** The object for which the frequency was being changed was not registered. */
+	UnregisteredStream,
+	/** The replication rate parameter was rejected (it cannot be 0). */
+	InvalidReplicationRate,
+
+	/** Not an actual parameter. Make sure it's always last. */
+	Count UMETA(Hidden)
+};
+
+/** Result about changing frequency. */
+USTRUCT(BlueprintType)
+struct FChangeClientStreamFrequencyResponse
+{
+	GENERATED_BODY()
+	
+	/** Errors encountered for specific objects */
+	UPROPERTY(BlueprintReadOnly, Category = "Multi-user")
+	TMap<FSoftObjectPath, EMultiUserChangeFrequencyErrorCode> ObjectErrors;
+	
+	/** Error for changing the default replication frequency. */
+	UPROPERTY(BlueprintReadOnly, Category = "Multi-user")
+	TOptional<EMultiUserChangeFrequencyErrorCode> DefaultChangeErrorCode;
+};
+
+/** Explains why a change to an object in the stream was invalid. */
+UENUM(BlueprintType)
+enum class EMultiUserPutObjectErrorCode : uint8
+{
+	/** Stream that the object referenced was not registered on the server. */
+	UnresolvedStream,
+	/**
+	 * Either PutObject contained no data to update with (ensure either ClassPath or Properties is set),
+	 * or it tried to create a new object with insufficient data (make sure ClassPath and Properties are both specified).
+	 */
+	MissingData,
+
+	/** Not an actual parameter. Make sure it's always last. */
+	Count UMETA(Hidden)
+};
+
+USTRUCT(BlueprintType)
+struct FMultiUserChangeClientStreamResponse
+{
+	GENERATED_BODY()
+	
+	/** The error code of changing streams. */
+	UPROPERTY(BlueprintReadOnly, Category = "Multi-user")
+	EMultiUserChangeStreamOperationResult ErrorCode = EMultiUserChangeStreamOperationResult::Timeout;
+	
+	/**
+	 * Dynamic authority errors.
+	 * The change was rejected because
+	 * 1. this client is replicating the object already,
+	 * 2. another client is also replicating the object,
+	 * 3. this change would cause overlapping properties with the other client.
+	 * */
+	UPROPERTY(BlueprintReadOnly, Category = "Multi-user")
+	TMap<FSoftObjectPath, FGuid> AuthorityConflicts;
+		
+	/** Errors made in the format of the request */
+	UPROPERTY(BlueprintReadOnly, Category = "Multi-user")
+	TMap<FSoftObjectPath, EMultiUserPutObjectErrorCode> SemanticErrors;
+
+	/** Errors made in the way frequency was changed. */
+	UPROPERTY(BlueprintReadOnly, Category = "Multi-user")
+	FChangeClientStreamFrequencyResponse FrequencyErrors;
+
+	/**
+	 * The client attempted to register the Multi-User stream but failed in doing so.
+	 * This usually indicates an internal error that you cannot do anything about as API user.
+	 */
+	UPROPERTY(BlueprintReadOnly, Category = "Multi-user")
+	bool bFailedStreamCreation = false;
+};
+
+USTRUCT(BlueprintType)
+struct FMultiUserChangeClientAuthorityResponse
+{
+	GENERATED_BODY()
+	
+	/** The error code of changing authority. Fails if StreamChangeResult fails. */
+	UPROPERTY(BlueprintReadOnly, Category = "Multi-user")
+	EMultiUserChangeAuthorityOperationResult ErrorCode = EMultiUserChangeAuthorityOperationResult::Timeout;
+	
+	/** Objects the client did not get authority over. */
+	UPROPERTY(BlueprintReadOnly, Category = "Multi-user")
+	TSet<FSoftObjectPath> RejectedObjects;
+};
+
 /** Result of processing a FChangeClientReplicationRequest. */
 USTRUCT(BlueprintType)
 struct FMultiUserChangeClientReplicationResult
 {
 	GENERATED_BODY()
 
-	/** The result of changing streams. */
+	/**
+	 * The result of changing streams.
+	 * You can inspect this for why the operation failed.
+	 */
 	UPROPERTY(BlueprintReadOnly, Category = "Multi-user")
-	EMultiUserChangeStreamOperationResult StreamChangeResult = EMultiUserChangeStreamOperationResult::Success;
+	FMultiUserChangeClientStreamResponse StreamResponse;
 
-	/** The result of changing authority. Fails if StreamChangeResult fails. */
+	/**
+	 * The result of changing authority. Fails if StreamChangeResult fails.
+	 * You can inspect this for why the operation failed.
+	 */
 	UPROPERTY(BlueprintReadOnly, Category = "Multi-user")
-	EMultiUserChangeAuthorityOperationResult AuthorityChangeResult = EMultiUserChangeAuthorityOperationResult::Success;
+	FMultiUserChangeClientAuthorityResponse AuthorityResponse;
 };
 
 namespace UE::MultiUserClientLibrary
@@ -237,6 +335,11 @@ namespace UE::MultiUserClientLibrary
 #if WITH_CONCERT
 	MULTIUSERCLIENTLIBRARY_API EMultiUserChangeStreamOperationResult Transform(MultiUserClient::EChangeStreamOperationResult Data);
 	MULTIUSERCLIENTLIBRARY_API EMultiUserChangeAuthorityOperationResult Transform(MultiUserClient::EChangeAuthorityOperationResult Data);
+	MULTIUSERCLIENTLIBRARY_API EMultiUserChangeFrequencyErrorCode Transform(MultiUserClient::EChangeObjectFrequencyErrorCode Data);
+	MULTIUSERCLIENTLIBRARY_API EMultiUserPutObjectErrorCode Transform(MultiUserClient::EPutObjectErrorCode Data);
+	
+	MULTIUSERCLIENTLIBRARY_API FMultiUserChangeClientStreamResponse Transform(MultiUserClient::FChangeClientStreamResponse Data);
+	MULTIUSERCLIENTLIBRARY_API FMultiUserChangeClientAuthorityResponse Transform(MultiUserClient::FChangeClientAuthorityResponse Data);
 	MULTIUSERCLIENTLIBRARY_API FMultiUserChangeClientReplicationResult Transform(MultiUserClient::FChangeClientReplicationResult Data);
 	
 	MULTIUSERCLIENTLIBRARY_API MultiUserClient::EPropertyChangeType Transform(EMultiUserPropertyChangeType Data);
diff --git a/Engine/Plugins/Developer/Concert/ConcertSync/ConcertSyncCore/Source/ConcertSyncCore/Public/Replication/Messages/ChangeStream.h b/Engine/Plugins/Developer/Concert/ConcertSync/ConcertSyncCore/Source/ConcertSyncCore/Public/Replication/Messages/ChangeStream.h
--- a/Engine/Plugins/Developer/Concert/ConcertSync/ConcertSyncCore/Source/ConcertSyncCore/Public/Replication/Messages/ChangeStream.h
+++ b/Engine/Plugins/Developer/Concert/ConcertSync/ConcertSyncCore/Source/ConcertSyncCore/Public/Replication/Messages/ChangeStream.h
@@ -211,7 +211,10 @@ enum class EConcertPutObjectErrorCode : uint8
 	 * Either PutObject contained no data to update with (ensure either ClassPath or Properties is set),
 	 * or it tried to create a new object with insufficient data (make sure ClassPath and Properties are both specified).
 	 */
-	MissingData
+	MissingData,
+
+	/** Not an actual parameter. Make sure it's always last. */
+	Count
 };
 
 UENUM()
@@ -220,7 +223,10 @@ enum class EConcertChangeObjectFrequencyErrorCode : uint8
 	/** The object for which the frequency was being changed was not registered. */
 	NotRegistered,
 	/** The replication rate parameter was rejected (it cannot be 0). */
-	InvalidReplicationRate
+	InvalidReplicationRate,
+
+	/** Not an actual parameter. Make sure it's always last. */
+	Count
 };
 
 UENUM()
@@ -229,7 +235,10 @@ enum class EConcertChangeStreamFrequencyErrorCode : uint8
 	/** The stream was not registered */
 	UnknownStream,
 	/** The replication rate parameter was rejected (it cannot be 0). */
-	InvalidReplicationRate
+	InvalidReplicationRate,
+
+	/** Not an actual parameter. Make sure it's always last. */
+	Count
 };
 
 USTRUCT()
