Fix failing ensure when reassigning authority from local to other client

diff --git a/Engine/Plugins/Developer/Concert/ConcertApp/MultiUserClient/Source/MultiUserClient/Private/Replication/Submission/Remote/SubmissionWorkflow_RemoteClient.cpp b/Engine/Plugins/Developer/Concert/ConcertApp/MultiUserClient/Source/MultiUserClient/Private/Replication/Submission/Remote/SubmissionWorkflow_RemoteClient.cpp
--- a/Engine/Plugins/Developer/Concert/ConcertApp/MultiUserClient/Source/MultiUserClient/Private/Replication/Submission/Remote/SubmissionWorkflow_RemoteClient.cpp
+++ b/Engine/Plugins/Developer/Concert/ConcertApp/MultiUserClient/Source/MultiUserClient/Private/Replication/Submission/Remote/SubmissionWorkflow_RemoteClient.cpp
@@ -238,7 +238,7 @@ namespace UE::MultiUserClient
 	{
 		FSingleClientSubmissionOperation& Operation = GetOperation();
 		const bool bIsStreamChangeDone = Operation.HasSetStreamPromise() || InProgressOperation->Parameters.StreamRequest->IsEmpty();
-		const bool bIsAuthorityChangeDone = Operation.HasSetAuthorityRequestPromise() || InProgressOperation->Parameters.AuthorityRequest->IsEmpty();
+		const bool bIsAuthorityChangeDone = Operation.HasSetAuthorityResponsePromise() || InProgressOperation->Parameters.AuthorityRequest->IsEmpty();
 		if (bIsStreamChangeDone && bIsAuthorityChangeDone)
 		{
 			CleanUpSubmissionProcess();
