Fix not subscribing to events of clients that were already connected when joining the session

diff --git a/Engine/Plugins/Developer/Concert/ConcertApp/MultiUserClient/Source/MultiUserClient/Private/Replication/MultiUserReplicationManager.cpp b/Engine/Plugins/Developer/Concert/ConcertApp/MultiUserClient/Source/MultiUserClient/Private/Replication/MultiUserReplicationManager.cpp
--- a/Engine/Plugins/Developer/Concert/ConcertApp/MultiUserClient/Source/MultiUserClient/Private/Replication/MultiUserReplicationManager.cpp
+++ b/Engine/Plugins/Developer/Concert/ConcertApp/MultiUserClient/Source/MultiUserClient/Private/Replication/MultiUserReplicationManager.cpp
@@ -2,6 +2,7 @@
 
 #include "MultiUserReplicationManager.h"
 
+#include "ConcertLogGlobal.h"
 #include "IConcertSyncClient.h"
 #include "Replication/ChangeOperationTypes.h"
 #include "Replication/IConcertClientReplicationManager.h"
@@ -104,17 +105,20 @@ namespace UE::MultiUserClient
 	void FMultiUserReplicationManager::SetupClientConnectionEvents()
 	{
 		FReplicationClientManager& ClientManager = ConnectedState->ClientManager;
-		SetupClientDelegates(ClientManager.GetLocalClient());
+		ClientManager.ForEachClient([this](FReplicationClient& InClient){ SetupClientDelegates(InClient); return EBreakBehavior::Continue; });
 		ClientManager.OnPostRemoteClientAdded().AddRaw(this, &FMultiUserReplicationManager::OnReplicationClientConnected);
 	}
 
 	void FMultiUserReplicationManager::OnClientStreamServerStateChanged(const FGuid EndpointId) const
 	{
+		UE_LOG(LogConcert, Verbose, TEXT("Client %s stream changed"), *EndpointId.ToString());
 		OnStreamServerStateChangedDelegate.Broadcast(EndpointId);
 	}
 
 	void FMultiUserReplicationManager::OnClientAuthorityServerStateChanged(const FGuid EndpointId) const
 	{
+		
+		UE_LOG(LogConcert, Verbose, TEXT("Client %s authority changed"), *EndpointId.ToString());
 		OnAuthorityServerStateChangedDelegate.Broadcast(EndpointId);
 	}
 
