Do not show frequency menu options if object is not registered by any client

diff --git a/Engine/Plugins/Developer/Concert/ConcertApp/MultiUserClient/Source/MultiUserClient/Private/Widgets/ActiveSession/Replication/Client/FrequencyContextMenuUtils.cpp b/Engine/Plugins/Developer/Concert/ConcertApp/MultiUserClient/Source/MultiUserClient/Private/Widgets/ActiveSession/Replication/Client/FrequencyContextMenuUtils.cpp
--- a/Engine/Plugins/Developer/Concert/ConcertApp/MultiUserClient/Source/MultiUserClient/Private/Widgets/ActiveSession/Replication/Client/FrequencyContextMenuUtils.cpp
+++ b/Engine/Plugins/Developer/Concert/ConcertApp/MultiUserClient/Source/MultiUserClient/Private/Widgets/ActiveSession/Replication/Client/FrequencyContextMenuUtils.cpp
@@ -2,6 +2,7 @@
 
 #include "FrequencyContextMenuUtils.h"
 
+#include "Algo/AnyOf.h"
 #include "Replication/Client/ReplicationClientManager.h"
 #include "Replication/Util/FrequencyUtils.h"
 
@@ -239,7 +240,7 @@ namespace UE::MultiUserClient::FrequencyContextMenuUtils
 		{
 			const auto GetText = [ContextObject, GetClientsAttribute, &InClientManager](FText Mixed, FText Specified, FText Realtime)
 			{
-				const FInlineClientArray& Clients = GetClientsAttribute.Get(); 
+				const FInlineClientArray& Clients = GetClientsAttribute.Get();
 				const TOptional<EConcertObjectReplicationMode> Mode = FrequencyUtils::FindSharedReplicationMode(ContextObject, Clients, InClientManager);
 				if (!Mode)
 				{
@@ -278,6 +279,17 @@ namespace UE::MultiUserClient::FrequencyContextMenuUtils
 
 		static void SharedAppendFrequencyToMenu(FMenuBuilder& MenuBuilder, const FSoftObjectPath& ContextObject, TAttribute<FInlineClientArray> GetClientsAttribute, FReplicationClientManager& InClientManager)
 		{
+			// No frequency section if the object is not registered by any clients
+			const bool bHasObjectRegistered = Algo::AnyOf(GetClientsAttribute.Get(), [&ContextObject, &InClientManager](const FGuid& ClientId)
+			{
+				const FReplicationClient* Client = InClientManager.FindClient(ClientId);
+				return !Client || Client->GetStreamSynchronizer().GetServerState().HasProperties(ContextObject);
+			});
+			if (!bHasObjectRegistered)
+			{
+				return;
+			}
+			
 			MenuBuilder.BeginSection(NAME_None, LOCTEXT("Frequency", "Frequency"));
 			
 			// Change to Realtime / Specified Rate
