Added UMultiUserSubsystem::OnSessionClientChanged event allowing Blueprints to detect other clients joining and leaving the session the local editor instance is connected to

diff --git a/Engine/Plugins/Developer/Concert/ConcertApp/MultiUserClient/Source/MultiUserClientLibrary/Private/MultiUserClientStatics.cpp b/Engine/Plugins/Developer/Concert/ConcertApp/MultiUserClient/Source/MultiUserClientLibrary/Private/MultiUserClientStatics.cpp
--- a/Engine/Plugins/Developer/Concert/ConcertApp/MultiUserClient/Source/MultiUserClientLibrary/Private/MultiUserClientStatics.cpp
+++ b/Engine/Plugins/Developer/Concert/ConcertApp/MultiUserClient/Source/MultiUserClientLibrary/Private/MultiUserClientStatics.cpp
@@ -23,54 +23,61 @@ DEFINE_LOG_CATEGORY_STATIC(LogMultiUserClient, Log, All);
 #define LOCTEXT_NAMESPACE "MultiUserClientStatics"
 
 #if WITH_CONCERT
-namespace MultiUserClientUtil
+namespace UE::MultiUserClientLibrary
 {
+	FMultiUserClientInfo ConvertClientInfo(const FGuid& ClientEndpointId, const FConcertClientInfo& ClientInfo)
+	{
+		FMultiUserClientInfo Result;
+		Result.ClientEndpointId = ClientEndpointId;
+		Result.DisplayName = ClientInfo.DisplayName;
+		Result.AvatarColor = ClientInfo.AvatarColor;
+		Result.Tags = ClientInfo.Tags;
+		return Result;
+	}
 
-FMultiUserClientInfo ConvertClientInfo(const FGuid& ClientEndpointId, const FConcertClientInfo& ClientInfo)
-{
-	FMultiUserClientInfo Result;
-	Result.ClientEndpointId = ClientEndpointId;
-	Result.DisplayName = ClientInfo.DisplayName;
-	Result.AvatarColor = ClientInfo.AvatarColor;
-	Result.Tags = ClientInfo.Tags;
-	return Result;
-}
-
-FMultiUserConnectionError ConvertConnectionError(FConcertConnectionError Error)
-{
-	FMultiUserConnectionError MUError;
-	MUError.ErrorCode = static_cast<EMultiUserConnectionError>(Error.ErrorCode);
-	MUError.ErrorMessage = Error.ErrorText;
-	return MUError;
-}
+	FMultiUserConnectionError ConvertConnectionError(FConcertConnectionError Error)
+	{
+		FMultiUserConnectionError MUError;
+		MUError.ErrorCode = static_cast<EMultiUserConnectionError>(Error.ErrorCode);
+		MUError.ErrorMessage = Error.ErrorText;
+		return MUError;
+	}
 
-UConcertClientConfig* ModifyClientConfig(const FMultiUserClientConfig& InClientConfig)
-{
-	UConcertClientConfig* ClientConfig = GetMutableDefault<UConcertClientConfig>();
-	ClientConfig->DefaultServerURL = InClientConfig.DefaultServerURL;
-	ClientConfig->DefaultSessionName = InClientConfig.DefaultSessionName;
-	ClientConfig->DefaultSessionToRestore = InClientConfig.DefaultSessionToRestore;
-	ClientConfig->SourceControlSettings.ValidationMode = static_cast<EConcertSourceValidationMode>(InClientConfig.ValidationMode);
-	return ClientConfig;
-}
+	UConcertClientConfig* ModifyClientConfig(const FMultiUserClientConfig& InClientConfig)
+	{
+		UConcertClientConfig* ClientConfig = GetMutableDefault<UConcertClientConfig>();
+		ClientConfig->DefaultServerURL = InClientConfig.DefaultServerURL;
+		ClientConfig->DefaultSessionName = InClientConfig.DefaultSessionName;
+		ClientConfig->DefaultSessionToRestore = InClientConfig.DefaultSessionToRestore;
+		ClientConfig->SourceControlSettings.ValidationMode = static_cast<EConcertSourceValidationMode>(InClientConfig.ValidationMode);
+		return ClientConfig;
+	}
 
-EMultiUserConnectionStatus ConvertConnectionStatus(EConcertConnectionStatus ConnectionStatus)
-{
-	return static_cast<EMultiUserConnectionStatus>(ConnectionStatus);
-}
+	EMultiUserConnectionStatus ConvertConnectionStatus(EConcertConnectionStatus ConnectionStatus)
+	{
+		return static_cast<EMultiUserConnectionStatus>(ConnectionStatus);
+	}
 
-FMultiUserSessionInfo ConvertSessionInfo(const FConcertSessionInfo& InSessionInfo, const FConcertServerInfo& InServerInfo)
-{
-	FMultiUserSessionInfo ReturnVal;
-	ReturnVal.SessionName = InSessionInfo.SessionName;
-	ReturnVal.ServerName = InServerInfo.ServerName;
-	ReturnVal.EndpointName = InServerInfo.InstanceInfo.InstanceName;
-	ReturnVal.ServerEndpointId = InSessionInfo.ServerInstanceId;
-	ReturnVal.bValid = true;
-	return ReturnVal;
-}
+	FMultiUserSessionInfo ConvertSessionInfo(const FConcertSessionInfo& InSessionInfo, const FConcertServerInfo& InServerInfo)
+	{
+		FMultiUserSessionInfo ReturnVal;
+		ReturnVal.SessionName = InSessionInfo.SessionName;
+		ReturnVal.ServerName = InServerInfo.ServerName;
+		ReturnVal.EndpointName = InServerInfo.InstanceInfo.InstanceName;
+		ReturnVal.ServerEndpointId = InSessionInfo.ServerInstanceId;
+		ReturnVal.bValid = true;
+		return ReturnVal;
+	}
 
-} // namespace MultiUserClientUtil
+	EMultiUserClientStatus ConvertClientStatus(EConcertClientStatus Status)
+	{
+		static_assert(static_cast<int32>(EConcertClientStatus::Count) == 3, "Update this location when modifying EConcertClientStatus");
+		static_assert(static_cast<int32>(EConcertClientStatus::Connected) == static_cast<int32>(EMultiUserClientStatus::Connected));
+		static_assert(static_cast<int32>(EConcertClientStatus::Disconnected) == static_cast<int32>(EMultiUserClientStatus::Disconnected));
+		static_assert(static_cast<int32>(EConcertClientStatus::Updated) == static_cast<int32>(EMultiUserClientStatus::Updated));
+		return static_cast<EMultiUserClientStatus>(Status);
+	}
+} // namespace UE::MultiUserClientLibrary
 #endif
 
 UMultiUserClientStatics::UMultiUserClientStatics(const FObjectInitializer& ObjectInitializer)
@@ -258,7 +265,7 @@ FMultiUserClientInfo UMultiUserClientStatics::GetLocalMultiUserClientInfo()
 
 			FGuid LocalClientEndpointId = ClientSession ? ClientSession->GetSessionClientEndpointId() : FGuid();
 			const FConcertClientInfo& LocalClientInfo = ClientSession ? ClientSession->GetLocalClientInfo() : ConcertClient->GetClientInfo();
-			ClientInfo = MultiUserClientUtil::ConvertClientInfo(LocalClientEndpointId, LocalClientInfo);
+			ClientInfo = UE::MultiUserClientLibrary::ConvertClientInfo(LocalClientEndpointId, LocalClientInfo);
 		}
 	}
 #endif
@@ -284,7 +291,7 @@ FMultiUserSessionInfo UMultiUserClientStatics::GetMultiUserSessionInfo()
 				});
 				if (ServerInfo)
 				{
-					return MultiUserClientUtil::ConvertSessionInfo(SessionInfo, *ServerInfo);
+					return UE::MultiUserClientLibrary::ConvertSessionInfo(SessionInfo, *ServerInfo);
 				}
 			}
 		}
@@ -308,7 +315,7 @@ bool UMultiUserClientStatics::GetMultiUserClientInfoByName(const FString& Client
 			const FConcertClientInfo& LocalClientInfo = ClientSession ? ClientSession->GetLocalClientInfo() : ConcertClient->GetClientInfo();
 			if (ClientName == LocalClientInfo.DisplayName)
 			{
-				ClientInfo = MultiUserClientUtil::ConvertClientInfo(ClientSession ? ClientSession->GetSessionClientEndpointId() : FGuid(), LocalClientInfo);
+				ClientInfo = UE::MultiUserClientLibrary::ConvertClientInfo(ClientSession ? ClientSession->GetSessionClientEndpointId() : FGuid(), LocalClientInfo);
 				return true;
 			}
 
@@ -319,7 +326,7 @@ bool UMultiUserClientStatics::GetMultiUserClientInfoByName(const FString& Client
 				{
 					if (SessionClient.ClientInfo.DisplayName == ClientName)
 					{
-						ClientInfo = MultiUserClientUtil::ConvertClientInfo(SessionClient.ClientEndpointId, SessionClient.ClientInfo);
+						ClientInfo = UE::MultiUserClientLibrary::ConvertClientInfo(SessionClient.ClientEndpointId, SessionClient.ClientInfo);
 						return true;
 					}
 				}
@@ -345,7 +352,7 @@ bool UMultiUserClientStatics::GetRemoteMultiUserClientInfos(TArray<FMultiUserCli
 				const TArray<FConcertSessionClientInfo> SessionClients = ClientSession->GetSessionClients();
 				for (const FConcertSessionClientInfo& SessionClient : SessionClients)
 				{
-					ClientInfos.Add(MultiUserClientUtil::ConvertClientInfo(SessionClient.ClientEndpointId, SessionClient.ClientInfo));
+					ClientInfos.Add(UE::MultiUserClientLibrary::ConvertClientInfo(SessionClient.ClientEndpointId, SessionClient.ClientInfo));
 				}
 
 				return ClientInfos.Num() > 0;
@@ -363,7 +370,7 @@ bool UMultiUserClientStatics::ConfigureMultiUserClient(const FMultiUserClientCon
 	{
 		if (TSharedPtr<IConcertSyncClient> ConcertSyncClient = IMultiUserClientModule::Get().GetClient())
 		{
-			ConcertSyncClient->GetConcertClient()->Configure(MultiUserClientUtil::ModifyClientConfig(ClientConfig));
+			ConcertSyncClient->GetConcertClient()->Configure(UE::MultiUserClientLibrary::ModifyClientConfig(ClientConfig));
 			return true;
 		}
 	}
@@ -390,7 +397,7 @@ FMultiUserConnectionError UMultiUserClientStatics::GetLastMultiUserConnectionErr
 	{
 		if (TSharedPtr<IConcertSyncClient> ConcertSyncClient = IMultiUserClientModule::Get().GetClient())
 		{
-			LastError = MultiUserClientUtil::ConvertConnectionError(ConcertSyncClient->GetConcertClient()->GetLastConnectionError());
+			LastError = UE::MultiUserClientLibrary::ConvertConnectionError(ConcertSyncClient->GetConcertClient()->GetLastConnectionError());
 		}
 	}
 #endif
@@ -409,7 +416,7 @@ EMultiUserConnectionStatus UMultiUserClientStatics::GetMultiUserConnectionStatus
 			const TSharedPtr<IConcertClientSession> ClientSession = ConcertClient->GetCurrentSession();
 			if (ClientSession.IsValid())
 			{
-				return MultiUserClientUtil::ConvertConnectionStatus(ClientSession->GetConnectionStatus());
+				return UE::MultiUserClientLibrary::ConvertConnectionStatus(ClientSession->GetConnectionStatus());
 			}
 		}
 	}
diff --git a/Engine/Plugins/Developer/Concert/ConcertApp/MultiUserClient/Source/MultiUserClientLibrary/Private/MultiUserSubsystem.cpp b/Engine/Plugins/Developer/Concert/ConcertApp/MultiUserClient/Source/MultiUserClientLibrary/Private/MultiUserSubsystem.cpp
--- a/Engine/Plugins/Developer/Concert/ConcertApp/MultiUserClient/Source/MultiUserClientLibrary/Private/MultiUserSubsystem.cpp
+++ b/Engine/Plugins/Developer/Concert/ConcertApp/MultiUserClient/Source/MultiUserClientLibrary/Private/MultiUserSubsystem.cpp
@@ -2,9 +2,9 @@
 
 #include "MultiUserSubsystem.h"
 
-#include "Engine/Engine.h"
-#include "Engine/GameEngine.h"
 #include "Delegates/IDelegateInstance.h"
+#include "Engine/Engine.h"
+#include "MultiUserClientStatics.h"
 #include "Templates/UniquePtr.h"
 
 #if WITH_CONCERT
@@ -12,7 +12,6 @@
 #include "IConcertClient.h"
 #include "IConcertSyncClient.h"
 #include "IConcertSession.h"
-#include "IConcertSessionHandler.h"
 #include "IConcertSyncClientModule.h"
 #include "ConcertLogGlobal.h"
 #include "MultiUserSubsystemMessages.h"
@@ -99,9 +98,9 @@ struct FConcertManager
 	void Register(TSharedRef<IConcertClientSession> InSession)
 	{
 		WeakSession = InSession;
-
+		
 		InSession->RegisterCustomEventHandler<FConcertBlueprintEvent>(this, &FConcertManager::HandleBlueprintEvent);
-
+		InSession->OnSessionClientChanged().AddRaw(this, &FConcertManager::OnSessionClientChanged);
 	}
 
 	/**
@@ -112,6 +111,7 @@ struct FConcertManager
 		if (TSharedPtr<IConcertClientSession> Session = WeakSession.Pin())
 		{
 			Session->UnregisterCustomEventHandler<FConcertBlueprintEvent>(this);
+			Session->OnSessionClientChanged().RemoveAll(this);
 		}
 
 		WeakSession.Reset();
@@ -137,6 +137,18 @@ struct FConcertManager
 			SubSystem->DispatchEvent(InEvent);
 		}
 	}
+	
+	void OnSessionClientChanged(IConcertClientSession& Session, EConcertClientStatus Status, const FConcertSessionClientInfo& ClientInfo)
+	{
+		check(GEngine);
+		UMultiUserSubsystem* SubSystem = GEngine->GetEngineSubsystem<UMultiUserSubsystem>();
+		check(SubSystem);
+		
+		SubSystem->OnSessionClientChanged.Broadcast(
+			MultiUserClientLibrary::ConvertClientStatus(Status),
+			MultiUserClientLibrary::ConvertClientInfo(ClientInfo.ClientEndpointId, ClientInfo.ClientInfo)
+			);
+	}
 
 	/** Weak pointer to the client session with which to send events. May be null or stale. */
 	TWeakPtr<IConcertClientSession> WeakSession;
diff --git a/Engine/Plugins/Developer/Concert/ConcertApp/MultiUserClient/Source/MultiUserClientLibrary/Public/MultiUserClientStatics.h b/Engine/Plugins/Developer/Concert/ConcertApp/MultiUserClient/Source/MultiUserClientLibrary/Public/MultiUserClientStatics.h
--- a/Engine/Plugins/Developer/Concert/ConcertApp/MultiUserClient/Source/MultiUserClientLibrary/Public/MultiUserClientStatics.h
+++ b/Engine/Plugins/Developer/Concert/ConcertApp/MultiUserClient/Source/MultiUserClientLibrary/Public/MultiUserClientStatics.h
@@ -6,6 +6,16 @@
 #include "Kismet/BlueprintFunctionLibrary.h"
 #include "MultiUserClientStatics.generated.h"
 
+#if WITH_CONCERT
+class UConcertClientConfig;
+enum class EConcertClientStatus : uint8;
+enum class EConcertConnectionStatus : uint8;
+struct FConcertServerInfo;
+struct FConcertSessionInfo;
+struct FConcertClientInfo;
+struct FConcertConnectionError;
+#endif
+
 /** Delegate that is invoked when a package is saved. */
 DECLARE_DYNAMIC_MULTICAST_DELEGATE_OneParam(FOnPackageSavedSignature, FName, PackageName);
 
@@ -153,6 +163,30 @@ enum class EMultiUserConnectionStatus : uint8
 	Disconnected,
 };
 
+/** Exposes EConcertClientStatus */
+UENUM(BlueprintType)
+enum class EMultiUserClientStatus : uint8
+{
+	/** Client connected */
+	Connected,
+	/** Client disconnected */
+	Disconnected,
+	/** Client state updated */
+	Updated,
+};
+
+namespace UE::MultiUserClientLibrary
+{
+#if WITH_CONCERT
+	MULTIUSERCLIENTLIBRARY_API FMultiUserClientInfo ConvertClientInfo(const FGuid& ClientEndpointId, const FConcertClientInfo& ClientInfo);
+	MULTIUSERCLIENTLIBRARY_API FMultiUserConnectionError ConvertConnectionError(FConcertConnectionError Error);
+	MULTIUSERCLIENTLIBRARY_API UConcertClientConfig* ModifyClientConfig(const FMultiUserClientConfig& InClientConfig);
+	MULTIUSERCLIENTLIBRARY_API EMultiUserConnectionStatus ConvertConnectionStatus(EConcertConnectionStatus ConnectionStatus);
+	MULTIUSERCLIENTLIBRARY_API FMultiUserSessionInfo ConvertSessionInfo(const FConcertSessionInfo& InSessionInfo, const FConcertServerInfo& InServerInfo);
+	MULTIUSERCLIENTLIBRARY_API EMultiUserClientStatus ConvertClientStatus(EConcertClientStatus Status);
+#endif
+}
+
 UCLASS()
 class MULTIUSERCLIENTLIBRARY_API UMultiUserClientStatics : public UBlueprintFunctionLibrary
 {
diff --git a/Engine/Plugins/Developer/Concert/ConcertApp/MultiUserClient/Source/MultiUserClientLibrary/Public/MultiUserSubsystem.h b/Engine/Plugins/Developer/Concert/ConcertApp/MultiUserClient/Source/MultiUserClientLibrary/Public/MultiUserSubsystem.h
--- a/Engine/Plugins/Developer/Concert/ConcertApp/MultiUserClient/Source/MultiUserClientLibrary/Public/MultiUserSubsystem.h
+++ b/Engine/Plugins/Developer/Concert/ConcertApp/MultiUserClient/Source/MultiUserClientLibrary/Public/MultiUserSubsystem.h
@@ -9,11 +9,14 @@
 
 #include "MultiUserSubsystem.generated.h"
 
-DECLARE_DYNAMIC_MULTICAST_DELEGATE(FOnSessionConnected);
-DECLARE_DYNAMIC_MULTICAST_DELEGATE(FOnSessionDisconnected);
-
+enum class EMultiUserClientStatus : uint8;
 struct FConcertSessionContext;
 struct FConcertBlueprintEvent;
+struct FMultiUserClientInfo;
+
+DECLARE_DYNAMIC_MULTICAST_DELEGATE(FOnSessionConnected);
+DECLARE_DYNAMIC_MULTICAST_DELEGATE(FOnSessionDisconnected);
+DECLARE_DYNAMIC_MULTICAST_DELEGATE_TwoParams(FOnSessionClientChanged, EMultiUserClientStatus, Status, const FMultiUserClientInfo&, ClientInfo);
 
 USTRUCT(BlueprintType)
 struct FMultiUserBlueprintEventData
@@ -34,17 +37,21 @@ class MULTIUSERCLIENTLIBRARY_API UMultiUserSubsystem : public UEngineSubsystem
 	virtual void Deinitialize() override;
 
 public:
-	/**
-	 * Delegate invoked when a session connection event has occurred.
-	 */
+	
+	/** Invoked when the local editor instance has joined a session. */
 	UPROPERTY(BlueprintAssignable, Category = "Multi-user")
 	FOnSessionConnected OnSessionConnected;
 
+	/** Invoked when the local editor instanced has left a session. */
+	UPROPERTY(BlueprintAssignable, Category = "Multi-user")
+	FOnSessionDisconnected OnSessionDisconnected;
+
 	/**
-	 * Delegate invoked when a session disconnection event has occurred.
+	 * Invoked when information about a client changes while the local editor instance is in a session.
+	 * For example: other clients joining and leaving the session.
 	 */
 	UPROPERTY(BlueprintAssignable, Category = "Multi-user")
-	FOnSessionDisconnected OnSessionDisconnected;
+	FOnSessionClientChanged OnSessionClientChanged;
 
 	/**
 	 * Returns true if we are currently in a session.
diff --git a/Engine/Plugins/Developer/Concert/ConcertMain/Source/Concert/Public/ConcertMessages.h b/Engine/Plugins/Developer/Concert/ConcertMain/Source/Concert/Public/ConcertMessages.h
--- a/Engine/Plugins/Developer/Concert/ConcertMain/Source/Concert/Public/ConcertMessages.h
+++ b/Engine/Plugins/Developer/Concert/ConcertMain/Source/Concert/Public/ConcertMessages.h
@@ -48,6 +48,9 @@ enum class EConcertClientStatus : uint8
 	Disconnected,
 	/** Client state updated */
 	Updated,
+
+	/** Not an actual state - add new entries above. */
+	Count
 };
 
 UENUM()
