Selecting struct property now selects all child properties as well

diff --git a/Engine/Plugins/Developer/Concert/ConcertApp/MultiUserClient/Source/MultiUserClient/Private/Widgets/ActiveSession/Replication/Client/Multi/Columns/SAssignPropertyComboBox.cpp b/Engine/Plugins/Developer/Concert/ConcertApp/MultiUserClient/Source/MultiUserClient/Private/Widgets/ActiveSession/Replication/Client/Multi/Columns/SAssignPropertyComboBox.cpp
--- a/Engine/Plugins/Developer/Concert/ConcertApp/MultiUserClient/Source/MultiUserClient/Private/Widgets/ActiveSession/Replication/Client/Multi/Columns/SAssignPropertyComboBox.cpp
+++ b/Engine/Plugins/Developer/Concert/ConcertApp/MultiUserClient/Source/MultiUserClient/Private/Widgets/ActiveSession/Replication/Client/Multi/Columns/SAssignPropertyComboBox.cpp
@@ -6,6 +6,7 @@
 #include "Replication/Client/ReplicationClient.h"
 #include "Replication/Client/ReplicationClientManager.h"
 #include "Replication/Editor/Model/IEditableReplicationStreamModel.h"
+#include "Replication/Editor/Model/PropertyUtils.h"
 #include "Widgets/ActiveSession/Replication/Misc/SNoClients.h"
 #include "Widgets/ActiveSession/Replication/Client/ClientUtils.h"
 #include "Widgets/ClientName/SHorizontalClientList.h"
@@ -14,10 +15,10 @@
 
 #include "Algo/AnyOf.h"
 #include "Framework/MultiBox/MultiBoxBuilder.h"
-#include "ScopedTransaction.h"
-#include "Styling/SlateTypes.h"
 #include "Widgets/Input/SComboButton.h"
+#include "ScopedTransaction.h"
 #include "Styling/AppStyle.h"
+#include "Styling/SlateTypes.h"
 
 #define LOCTEXT_NAMESPACE "SAssignPropertyComboBox"
 
@@ -25,7 +26,7 @@ namespace UE::MultiUserClient
 {
 	namespace AssignPropertyComboBox
 	{
-		TArray<FGuid> GetDisplayedClients(const FReplicationClientManager& ClientManager, FConcertPropertyChain DisplayedProperty, const TArray<FSoftObjectPath>& EditedObjects)
+		TArray<FGuid> GetDisplayedClients(const FReplicationClientManager& ClientManager, const FConcertPropertyChain& DisplayedProperty, const TArray<FSoftObjectPath>& EditedObjects)
 		{
 			TArray<FGuid> Clients;
 			ClientManager.ForEachClient([&DisplayedProperty, &EditedObjects, &Clients](const FReplicationClient& Client)
@@ -49,7 +50,7 @@ namespace UE::MultiUserClient
 	TOptional<FString> SAssignPropertyComboBox::GetDisplayString(
 		const TSharedRef<IConcertClient>& LocalConcertClient,
 		const FReplicationClientManager& ClientManager,
-		FConcertPropertyChain DisplayedProperty,
+		const FConcertPropertyChain& DisplayedProperty,
 		const TArray<FSoftObjectPath>& EditedObjects)
 	{
 		using SWidgetType = ConcertClientSharedSlate::SHorizontalClientList;
@@ -67,7 +68,7 @@ namespace UE::MultiUserClient
 	    FReplicationClientManager& InClientManager
 	)
 	{
-		Editor = InEditor;
+		Editor = MoveTemp(InEditor);
 		ConcertClient = MoveTemp(InConcertClient);
 		ClientManager = &InClientManager;
 		
@@ -97,7 +98,7 @@ namespace UE::MultiUserClient
 		RefreshContentBoxContent();
 	}
 	
-	void SAssignPropertyComboBox::RefreshContentBoxContent()
+	void SAssignPropertyComboBox::RefreshContentBoxContent() const
 	{
 		ClientListWidget->RefreshList(
 			AssignPropertyComboBox::GetDisplayedClients(*ClientManager, Property, EditedObjects)
@@ -206,11 +207,11 @@ namespace UE::MultiUserClient
 				{
 					EditModel->AddObjects({ ObjectPath.ResolveObject() });
 				}
-				
-				if (EditModel->ContainsObjects({ ObjectPath }))
-				{
-					EditModel->AddProperties(ObjectPath, { Property });
-				}
+
+				const FSoftClassPath ClassPath = EditModel->GetObjectClass(ObjectPath);
+				TArray<FConcertPropertyChain> AddedProperties { Property };
+				ConcertClientSharedSlate::PropertyUtils::AppendAdditionalPropertiesToAdd(ClassPath, AddedProperties);
+				EditModel->AddProperties(ObjectPath, AddedProperties);
 			}
 		}
 		
diff --git a/Engine/Plugins/Developer/Concert/ConcertApp/MultiUserClient/Source/MultiUserClient/Private/Widgets/ActiveSession/Replication/Client/Multi/Columns/SAssignPropertyComboBox.h b/Engine/Plugins/Developer/Concert/ConcertApp/MultiUserClient/Source/MultiUserClient/Private/Widgets/ActiveSession/Replication/Client/Multi/Columns/SAssignPropertyComboBox.h
--- a/Engine/Plugins/Developer/Concert/ConcertApp/MultiUserClient/Source/MultiUserClient/Private/Widgets/ActiveSession/Replication/Client/Multi/Columns/SAssignPropertyComboBox.h
+++ b/Engine/Plugins/Developer/Concert/ConcertApp/MultiUserClient/Source/MultiUserClient/Private/Widgets/ActiveSession/Replication/Client/Multi/Columns/SAssignPropertyComboBox.h
@@ -37,7 +37,7 @@ namespace UE::MultiUserClient
 		static TOptional<FString> GetDisplayString(
 			const TSharedRef<IConcertClient>& LocalConcertClient,
 			const FReplicationClientManager& ClientManager,
-			FConcertPropertyChain DisplayedProperty,
+			const FConcertPropertyChain& DisplayedProperty,
 			const TArray<FSoftObjectPath>& EditedObjects
 			);
 		
@@ -78,7 +78,7 @@ namespace UE::MultiUserClient
 		FOnOptionClicked OnOptionClickedDelegate;
 
 		/** Updates the content of the combo box */
-		void RefreshContentBoxContent();
+		void RefreshContentBoxContent() const;
 		
 		/** Builds the drop-down menu */
 		TSharedRef<SWidget> GetMenuContent();
diff --git a/Engine/Plugins/Developer/Concert/ConcertApp/MultiUserClient/Source/MultiUserClient/Private/Widgets/ActiveSession/Replication/Client/Single/SReplicationClientView.h b/Engine/Plugins/Developer/Concert/ConcertApp/MultiUserClient/Source/MultiUserClient/Private/Widgets/ActiveSession/Replication/Client/Single/SReplicationClientView.h
--- a/Engine/Plugins/Developer/Concert/ConcertApp/MultiUserClient/Source/MultiUserClient/Private/Widgets/ActiveSession/Replication/Client/Single/SReplicationClientView.h
+++ b/Engine/Plugins/Developer/Concert/ConcertApp/MultiUserClient/Source/MultiUserClient/Private/Widgets/ActiveSession/Replication/Client/Single/SReplicationClientView.h
@@ -5,6 +5,7 @@
 #include "Widgets/DeclarativeSyntaxSupport.h"
 #include "Widgets/SCompoundWidget.h"
 
+struct FConcertPropertyChain;
 class IConcertClient;
 class FMenuBuilder;
 
@@ -50,6 +51,8 @@ namespace UE::MultiUserClient
 
 		/** Adds additional entries to the context menu for the object tree view. */
 		void ExtendObjectContextMenu(FMenuBuilder& MenuBuilder, TConstArrayView<FSoftObjectPath> ContextObjects) const;
+		/** Adds additional properties, e.g. struct child properties. */
+		void ExtendPropertiesToAdd(const FSoftObjectPath& Object, TArray<FConcertPropertyChain>& InOutPropertiesToAdd) const;
 		
 		/** Called when any of the streams change. */
 		void OnModelChanged() const;
diff --git a/Engine/Plugins/Developer/Concert/ConcertUI/ConcertClientSharedSlate/Source/ConcertClientSharedSlate/Private/Replication/ClientReplicationWidgetFactories.cpp b/Engine/Plugins/Developer/Concert/ConcertUI/ConcertClientSharedSlate/Source/ConcertClientSharedSlate/Private/Replication/ClientReplicationWidgetFactories.cpp
--- a/Engine/Plugins/Developer/Concert/ConcertUI/ConcertClientSharedSlate/Source/ConcertClientSharedSlate/Private/Replication/ClientReplicationWidgetFactories.cpp
+++ b/Engine/Plugins/Developer/Concert/ConcertUI/ConcertClientSharedSlate/Source/ConcertClientSharedSlate/Private/Replication/ClientReplicationWidgetFactories.cpp
@@ -5,6 +5,7 @@
 #include "Editor/View/ClientEditorColumns.h"
 #include "Editor/View/PropertyTree/SFilteredPropertyTreeView.h"
 #include "Replication/ReplicationWidgetFactories.h"
+#include "Replication/Editor/Model/PropertyUtils.h"
 #include "Replication/Editor/Model/Object/EditorObjectHierarchyModel.h"
 #include "Replication/Editor/Model/Object/EditorObjectNameModel.h"
 #include "Replication/Editor/Model/ReplicationStreamObject.h"
@@ -40,7 +41,7 @@ namespace UE::ConcertClientSharedSlate
 
 	TSharedRef<ConcertSharedSlate::IEditableReplicationStreamModel> CreateTransactionalStreamModel()
 	{
-		const EObjectFlags Flags = RF_Transient | RF_Transactional;
+		constexpr EObjectFlags Flags = RF_Transient | RF_Transactional;
 		UReplicationStreamObject* Object = NewObject<UReplicationStreamObject>(GetTransientPackage(), NAME_None, Flags);
 
 		TAttribute<FConcertObjectReplicationMap*> Attribute = TAttribute<FConcertObjectReplicationMap*>::CreateLambda([WeakPtr = TWeakObjectPtr<UReplicationStreamObject>(Object)]() -> FConcertObjectReplicationMap* 
@@ -65,6 +66,16 @@ namespace UE::ConcertClientSharedSlate
 		using namespace ConcertSharedSlate;
 		using namespace ConcertSharedSlate::ReplicationColumns;
 
+		// When a user adds e.g. a struct property all of its child properties should be auto added as well.
+		FExtendProperties ExtendPropertiesDelegate = FExtendProperties::CreateLambda(
+		[DataModel = Params.BaseEditorParams.DataModel, AppendPropertyDelegate = MoveTemp(Params.OnExtendAddedProperties)]
+			(const FSoftObjectPath& Object, TArray<FConcertPropertyChain>& InOutPropertiesToAdd)
+			{
+				const FSoftClassPath ClassPath = DataModel->GetObjectClass(Object);
+				PropertyUtils::AppendAdditionalPropertiesToAdd(ClassPath, InOutPropertiesToAdd);
+				AppendPropertyDelegate.ExecuteIfBound(Object, InOutPropertiesToAdd);
+			});
+		
 		// This is a hack.
 		// The architecturally correct way to fix is pass FReplicationPropertyColumn the FSoftObjectPath to the object for which the column is being constructed.
 		struct FEditorIndirection
@@ -76,6 +87,7 @@ namespace UE::ConcertClientSharedSlate
 		const FReplicationPropertyColumn ReplicatesColumn = ReplicationColumns::Property::ReplicatesColumns(
 			TAttribute<IReplicationStreamViewer*>::CreateLambda([Indirection](){ return Indirection->Editor.Get(); }),
 			Params.BaseEditorParams.DataModel,
+			MoveTemp(ExtendPropertiesDelegate),
 			TReplicationColumnDelegates<FReplicatedPropertyData>::FIsEnabled::CreateLambda([IsEnabled = Params.BaseEditorParams.IsEditingEnabled](const FReplicatedPropertyData&)
 			{
 				return !IsEnabled.IsBound() || IsEnabled.Get();
diff --git a/Engine/Plugins/Developer/Concert/ConcertUI/ConcertClientSharedSlate/Source/ConcertClientSharedSlate/Private/Replication/Editor/Model/PropertyUtils.cpp b/Engine/Plugins/Developer/Concert/ConcertUI/ConcertClientSharedSlate/Source/ConcertClientSharedSlate/Private/Replication/Editor/Model/PropertyUtils.cpp
new file mode 100644
--- /dev/null
+++ b/Engine/Plugins/Developer/Concert/ConcertUI/ConcertClientSharedSlate/Source/ConcertClientSharedSlate/Private/Replication/Editor/Model/PropertyUtils.cpp
@@ -0,0 +1,30 @@
+// Copyright Epic Games, Inc. All Rights Reserved.
+
+#include "Replication/Editor/Model/PropertyUtils.h"
+
+#include "Algo/AnyOf.h"
+#include "Replication/PropertyChainUtils.h"
+
+namespace UE::ConcertClientSharedSlate::PropertyUtils
+{
+	void EnumerateAdditionalPropertiesToAdd(
+		const UClass& ObjectClass,
+		TConstArrayView<FConcertPropertyChain> PropertiesToAdd,
+		TFunctionRef<void(FConcertPropertyChain&& AdditionalProperty)> Callback
+		)
+	{
+		ConcertSyncCore::PropertyChain::ForEachReplicatableConcertProperty(ObjectClass, [&PropertiesToAdd, &Callback](FConcertPropertyChain&& ConsideredProperty)
+		{
+			const bool bIsChildOfAddedProperty = Algo::AnyOf(PropertiesToAdd, [&ConsideredProperty](const FConcertPropertyChain& AddedProperty)
+			{
+				return ConsideredProperty.IsChildOf(AddedProperty);
+			});
+			if (bIsChildOfAddedProperty)
+			{
+				Callback(MoveTemp(ConsideredProperty));
+			}
+			
+			return EBreakBehavior::Continue;
+		});
+	}
+}
diff --git a/Engine/Plugins/Developer/Concert/ConcertUI/ConcertClientSharedSlate/Source/ConcertClientSharedSlate/Private/Replication/Editor/View/ClientEditorColumns.cpp b/Engine/Plugins/Developer/Concert/ConcertUI/ConcertClientSharedSlate/Source/ConcertClientSharedSlate/Private/Replication/Editor/View/ClientEditorColumns.cpp
--- a/Engine/Plugins/Developer/Concert/ConcertUI/ConcertClientSharedSlate/Source/ConcertClientSharedSlate/Private/Replication/Editor/View/ClientEditorColumns.cpp
+++ b/Engine/Plugins/Developer/Concert/ConcertUI/ConcertClientSharedSlate/Source/ConcertClientSharedSlate/Private/Replication/Editor/View/ClientEditorColumns.cpp
@@ -33,7 +33,8 @@ namespace UE::ConcertClientSharedSlate::ReplicationColumns::Property
 			bool bIsChecked,
 			const FConcertPropertyChain& PropertyChain,
 			const ConcertSharedSlate::IReplicationStreamViewer& Viewer,
-			ConcertSharedSlate::IEditableReplicationStreamModel& Model
+			ConcertSharedSlate::IEditableReplicationStreamModel& Model,
+			const FExtendProperties& ExtendPropertiesDelegate
 			)
 		{
 			const TArray PropertyAsArray{ PropertyChain };
@@ -60,7 +61,10 @@ namespace UE::ConcertClientSharedSlate::ReplicationColumns::Property
 						// FYI: This should always be a subobject because actors are in the view beforehand (this checkbox is created in response to clicking the actor)
 						Model.AddObjects({ Object });
 					}
-					Model.AddProperties(Path, PropertyAsArray);
+
+					TArray<FConcertPropertyChain> PropertiesToAdd = PropertyAsArray;
+					ExtendPropertiesDelegate.ExecuteIfBound(Path, PropertiesToAdd);
+					Model.AddProperties(Path, PropertiesToAdd);
 				}
 			}
 			else
@@ -86,6 +90,7 @@ namespace UE::ConcertClientSharedSlate::ReplicationColumns::Property
 	ConcertSharedSlate::ReplicationColumns::FReplicationPropertyColumn ReplicatesColumns(
 		TAttribute<ConcertSharedSlate::IReplicationStreamViewer*> Viewer,
 		TWeakPtr<ConcertSharedSlate::IEditableReplicationStreamModel> Model,
+		FExtendProperties ExtendPropertiesDelegate,
 		ConcertSharedSlate::TReplicationColumnDelegates<ConcertSharedSlate::FReplicatedPropertyData>::FIsEnabled IsEnabledDelegate,
 		TAttribute<FText> DisabledToolTipText,
 		const float ColumnWidth,
@@ -105,13 +110,13 @@ namespace UE::ConcertClientSharedSlate::ReplicationColumns::Property
 					return ensure(ViewerPin && ModelPin) ? Private::OnGetPropertyCheckboxState(Data.GetProperty(), *ViewerPin, *ModelPin) : ECheckBoxState::Undetermined;
 				}),
 				FPropertyColumnDelegates::FOnColumnCheckboxChanged::CreateLambda(
-				[Viewer, Model](bool bIsChecked, const FReplicatedPropertyData& Data)
+				[Viewer, Model, ExtendPropertiesDelegate = MoveTemp(ExtendPropertiesDelegate)](bool bIsChecked, const FReplicatedPropertyData& Data)
 				{
 					const IReplicationStreamViewer* ViewerPin = Viewer.Get();
 					const TSharedPtr<IEditableReplicationStreamModel> ModelPin = Model.Pin();
 					if (ensure(ViewerPin && ModelPin))
 					{
-						Private::OnPropertyCheckboxChanged(bIsChecked, Data.GetProperty(), *ViewerPin, *ModelPin);
+						Private::OnPropertyCheckboxChanged(bIsChecked, Data.GetProperty(), *ViewerPin, *ModelPin, ExtendPropertiesDelegate);
 					}
 				}),
 				FPropertyColumnDelegates::FGetToolTipText::CreateLambda(
diff --git a/Engine/Plugins/Developer/Concert/ConcertUI/ConcertClientSharedSlate/Source/ConcertClientSharedSlate/Private/Replication/Editor/View/ClientEditorColumns.h b/Engine/Plugins/Developer/Concert/ConcertUI/ConcertClientSharedSlate/Source/ConcertClientSharedSlate/Private/Replication/Editor/View/ClientEditorColumns.h
--- a/Engine/Plugins/Developer/Concert/ConcertUI/ConcertClientSharedSlate/Source/ConcertClientSharedSlate/Private/Replication/Editor/View/ClientEditorColumns.h
+++ b/Engine/Plugins/Developer/Concert/ConcertUI/ConcertClientSharedSlate/Source/ConcertClientSharedSlate/Private/Replication/Editor/View/ClientEditorColumns.h
@@ -2,6 +2,7 @@
 
 #pragma once
 
+#include "Replication/ClientReplicationWidgetDelegates.h"
 #include "Replication/Editor/View/PredefinedReplicationColumns.h"
 #include "Replication/Editor/View/ReplicationColumnsUtils.h"
 #include "Templates/SharedPointer.h"
@@ -30,6 +31,7 @@ namespace UE::ConcertClientSharedSlate::ReplicationColumns::Property
 	 *
 	 * @param Viewer Used to determine which objects are currently selected (and being edited by the checkbox)
 	 * @param Model Used to actually edit the object to properties mapping
+	 * @param ExtendPropertiesDelegate Delegate to determine additional properties to add when the property checkbox is checked.
 	 * @param IsEnabledDelegate Determines whether the checkbox is enabled
 	 * @param DisabledToolTipText Tooltip to display when IsEnabledDelegate returns false
 	 * @param ColumnWidth Width to use for the column
@@ -38,6 +40,7 @@ namespace UE::ConcertClientSharedSlate::ReplicationColumns::Property
 	ConcertSharedSlate::ReplicationColumns::FReplicationPropertyColumn ReplicatesColumns(
 		TAttribute<ConcertSharedSlate::IReplicationStreamViewer*> Viewer,
 		TWeakPtr<ConcertSharedSlate::IEditableReplicationStreamModel> Model,
+		FExtendProperties ExtendPropertiesDelegate = {},
 		ConcertSharedSlate::TReplicationColumnDelegates<ConcertSharedSlate::FReplicatedPropertyData>::FIsEnabled IsEnabledDelegate = {},
 		TAttribute<FText> DisabledToolTipText = {},
 		const float ColumnWidth = 20.f,
diff --git a/Engine/Plugins/Developer/Concert/ConcertUI/ConcertClientSharedSlate/Source/ConcertClientSharedSlate/Public/Replication/ClientReplicationWidgetDelegates.h b/Engine/Plugins/Developer/Concert/ConcertUI/ConcertClientSharedSlate/Source/ConcertClientSharedSlate/Public/Replication/ClientReplicationWidgetDelegates.h
new file mode 100644
--- /dev/null
+++ b/Engine/Plugins/Developer/Concert/ConcertUI/ConcertClientSharedSlate/Source/ConcertClientSharedSlate/Public/Replication/ClientReplicationWidgetDelegates.h
@@ -0,0 +1,13 @@
+// Copyright Epic Games, Inc. All Rights Reserved.
+
+#pragma once
+
+#include "Containers/ContainersFwd.h"
+#include "Delegates/Delegate.h"
+
+struct FConcertPropertyChain;
+
+namespace UE::ConcertClientSharedSlate
+{
+	DECLARE_DELEGATE_TwoParams(FExtendProperties, const FSoftObjectPath& Object, TArray<FConcertPropertyChain>& InOutPropertiesToAdd);
+}
\ No newline at end of file
diff --git a/Engine/Plugins/Developer/Concert/ConcertUI/ConcertClientSharedSlate/Source/ConcertClientSharedSlate/Public/Replication/ClientReplicationWidgetFactories.h b/Engine/Plugins/Developer/Concert/ConcertUI/ConcertClientSharedSlate/Source/ConcertClientSharedSlate/Public/Replication/ClientReplicationWidgetFactories.h
--- a/Engine/Plugins/Developer/Concert/ConcertUI/ConcertClientSharedSlate/Source/ConcertClientSharedSlate/Public/Replication/ClientReplicationWidgetFactories.h
+++ b/Engine/Plugins/Developer/Concert/ConcertUI/ConcertClientSharedSlate/Source/ConcertClientSharedSlate/Public/Replication/ClientReplicationWidgetFactories.h
@@ -2,6 +2,7 @@
 
 #pragma once
 
+#include "ClientReplicationWidgetDelegates.h"
 #include "Replication/ReplicationWidgetFactories.h"
 #include "Templates/SharedPointer.h"
 
@@ -80,18 +81,27 @@ namespace UE::ConcertClientSharedSlate
 		
 		/**
 		 * Optional. Determines the objects displayed as children to the top-level objects in the top section.
-		 * If left unspecified, the top will only display actors.
+		 * If left unspecified, the object outliner will only display actors.
 		 * @note The created view will keep a strong reference to this.
 		 */
 		TSharedPtr<ConcertSharedSlate::IObjectHierarchyModel> ObjectHierarchy;
 
-		/** Optional. Determines the display name of objects. */
+		/**
+		 * Optional. Determines the display name of objects.
+		 * If left unspecified, defaults to the names in the object paths.
+		 */
 		TSharedPtr<ConcertSharedSlate::IObjectNameModel> NameModel;
 
-		/** Called to generate the context menu for objects. This extends the options already generated by this widget and is called at the end. */
+		/** Optional. Called to generate the context menu for objects. This extends the options already generated by this widget and is called at the end. */
 		ConcertSharedSlate::FExtendObjectMenu OnExtendObjectsContextMenu;
+
+		/**
+		 * Optional. When a property is assigned to the client, this is called to add any additional properties.
+		 * By default, if the user adds a struct property, its child properties are also added.
+		 */
+		FExtendProperties OnExtendAddedProperties;
 		
-		/** Additional columns to add to the object view */
+		/** Optional. Additional columns to add to the object view */
 		TArray<ConcertSharedSlate::ReplicationColumns::FReplicationTopLevelObjectColumn> AdditionalObjectColumns;
 	};
 	
diff --git a/Engine/Plugins/Developer/Concert/ConcertUI/ConcertClientSharedSlate/Source/ConcertClientSharedSlate/Public/Replication/Editor/Model/PropertyUtils.h b/Engine/Plugins/Developer/Concert/ConcertUI/ConcertClientSharedSlate/Source/ConcertClientSharedSlate/Public/Replication/Editor/Model/PropertyUtils.h
new file mode 100644
--- /dev/null
+++ b/Engine/Plugins/Developer/Concert/ConcertUI/ConcertClientSharedSlate/Source/ConcertClientSharedSlate/Public/Replication/Editor/Model/PropertyUtils.h
@@ -0,0 +1,49 @@
+// Copyright Epic Games, Inc. All Rights Reserved.
+
+#pragma once
+
+#include "Replication/Data/ConcertPropertySelection.h"
+
+#include "Containers/ContainersFwd.h"
+#include "Containers/ArrayView.h"
+#include "Containers/Array.h"
+#include "Templates/Function.h"
+
+namespace UE::ConcertClientSharedSlate::PropertyUtils
+{
+	/**
+	 * Called when the user adds a property in the property tree view.
+	 * This finds additional properties that should also be added.
+	 *
+	 * Examples:
+	 * - User adds a struct property: Add all child properties automatically
+	 * - User adds a container property: All all child properties automatically if it is a struct
+	 *
+	 * @param ObjectClass The class the properties belong to
+	 * @param PropertiesToAdd The properties about to be added a replication stream model
+	 * @param Callback Callback to receive the additional properties
+	 */
+	CONCERTCLIENTSHAREDSLATE_API void EnumerateAdditionalPropertiesToAdd(
+		const UClass& ObjectClass,
+		TConstArrayView<FConcertPropertyChain> PropertiesToAdd,
+		TFunctionRef<void(FConcertPropertyChain&& AdditionalProperty)> Callback
+		);
+
+	/** Special version of EnumerateAdditionalPropertiesToAdd that appends any additional properties to the same array. */
+	inline void AppendAdditionalPropertiesToAdd(const UClass& ObjectClass, TArray<FConcertPropertyChain>& InOutPropertiesToAdd)
+	{
+		// No we cannot directly add to InOutPropertiesToAdd in the callback because that may reallocate the memory that the TConstArrayView is reading!
+		TArray<FConcertPropertyChain> AdditionalProperties;
+		EnumerateAdditionalPropertiesToAdd(ObjectClass, InOutPropertiesToAdd, [&AdditionalProperties](FConcertPropertyChain&& Chain){ AdditionalProperties.Emplace(MoveTemp(Chain)); });
+		InOutPropertiesToAdd.Append(AdditionalProperties);
+	}
+
+	inline void AppendAdditionalPropertiesToAdd(const FSoftClassPath& ObjectClass, TArray<FConcertPropertyChain>& InOutPropertiesToAdd)
+	{
+		const UClass* Class = ObjectClass.TryLoadClass<UObject>();
+		if (ensureAlwaysMsgf(Class, TEXT("Unresolved class %s"), *ObjectClass.ToString()))
+		{
+			AppendAdditionalPropertiesToAdd(*Class, InOutPropertiesToAdd);
+		}
+	}
+}
