Update the local client's stream and tree view when the user modifies the level
- Removing an actor removes it from the local client's stream
- Refresh subobject hierarchy in tree view UI when it is changed (e.g. by adding a component)

diff --git a/Engine/Plugins/Developer/Concert/ConcertApp/MultiUserClient/Source/MultiUserClient/Private/Replication/Client/ReplicationClient.cpp b/Engine/Plugins/Developer/Concert/ConcertApp/MultiUserClient/Source/MultiUserClient/Private/Replication/Client/ReplicationClient.cpp
--- a/Engine/Plugins/Developer/Concert/ConcertApp/MultiUserClient/Source/MultiUserClient/Private/Replication/Client/ReplicationClient.cpp
+++ b/Engine/Plugins/Developer/Concert/ConcertApp/MultiUserClient/Source/MultiUserClient/Private/Replication/Client/ReplicationClient.cpp
@@ -53,6 +53,7 @@ namespace UE::MultiUserClient
 		, LocalFrequencyChangeTracker(*StreamSynchronizer)
 		, ChangeRequestBuilder(EndpointId, InAuthorityCache, *StreamSynchronizer, LocalClientStreamDiffer, LocalAuthorityDiffer, LocalFrequencyChangeTracker)
 		, AutoSubmissionPolicy(SubmissionQueue, ChangeRequestBuilder, LocalClientEditModel.Get(), LocalAuthorityDiffer, LocalFrequencyChangeTracker)
+		, LevelModificationHandler(GetClientEditModel().Get())
 	{
 		LocalClientEditModel->OnObjectsChanged().AddRaw(this, &FReplicationClient::OnObjectsChanged);
 		LocalClientEditModel->OnPropertiesChanged().AddRaw(this, &FReplicationClient::OnPropertiesChanged);
@@ -60,6 +61,11 @@ namespace UE::MultiUserClient
 
 		SubmissionWorkflow->OnAuthorityRequestCompleted_AnyThread().AddRaw(this, &FReplicationClient::OnAuthoritySubmissionCompleted);
 		StreamSynchronizer->OnServerStateChanged().AddRaw(this, &FReplicationClient::OnServerStateChanged);
+
+		LevelModificationHandler.OnHierarchyNeedsRefresh().AddLambda([this]()
+		{
+			OnHierarchyNeedsRefreshDelegate.Broadcast();
+		});
 	}
 
 	FReplicationClient::~FReplicationClient()
diff --git a/Engine/Plugins/Developer/Concert/ConcertApp/MultiUserClient/Source/MultiUserClient/Private/Replication/Client/ReplicationClient.h b/Engine/Plugins/Developer/Concert/ConcertApp/MultiUserClient/Source/MultiUserClient/Private/Replication/Client/ReplicationClient.h
--- a/Engine/Plugins/Developer/Concert/ConcertApp/MultiUserClient/Source/MultiUserClient/Private/Replication/Client/ReplicationClient.h
+++ b/Engine/Plugins/Developer/Concert/ConcertApp/MultiUserClient/Source/MultiUserClient/Private/Replication/Client/ReplicationClient.h
@@ -4,6 +4,7 @@
 
 #include "Replication/Authority/AuthorityChangeTracker.h"
 #include "Replication/Authority/IClientAuthoritySynchronizer.h"
+#include "Replication/Editor/UnrealEditor/ModifyObjectInLevelHandler.h"
 #include "Replication/Frequency/FrequencyChangeTracker.h"
 #include "Replication/Stream/IClientStreamSynchronizer.h"
 #include "Replication/Stream/StreamChangeTracker.h"
@@ -101,6 +102,10 @@ namespace UE::MultiUserClient
 		 */
 		FOnModelExternallyChanged& OnModelChanged() { return OnModelChangedDelegate; }
 		
+		DECLARE_MULTICAST_DELEGATE(FOnHiearchyNeedsRefresh);
+		/** Broadcasts when the hierarchy may have changed. */
+		FOnHiearchyNeedsRefresh& OnHierarchyNeedsRefresh() { return OnHierarchyNeedsRefreshDelegate; }
+		
 	private:
 
 		/** This client's Concert Endpoint ID. */
@@ -144,6 +149,12 @@ namespace UE::MultiUserClient
 		FChangeRequestBuilder ChangeRequestBuilder;
 		/** Automatically submits changes as they are made by the user. */
 		FAutoSubmissionPolicy AutoSubmissionPolicy;
+		
+		/** Handles updating the local client's stream when the level is modified, e.g. when an actor is removed. */
+		ConcertClientSharedSlate::FModifyObjectInLevelHandler LevelModificationHandler;
+
+		/** Broadcasts when the hierarchy may have changed. */
+		FOnHiearchyNeedsRefresh OnHierarchyNeedsRefreshDelegate;
 
 		/**
 		 * Broadcast when the data underlying the model has changed for any reason:
diff --git a/Engine/Plugins/Developer/Concert/ConcertApp/MultiUserClient/Source/MultiUserClient/Private/Widgets/ActiveSession/Replication/Client/Multi/AllClientsSelectionModel.h b/Engine/Plugins/Developer/Concert/ConcertApp/MultiUserClient/Source/MultiUserClient/Private/Widgets/ActiveSession/Replication/Client/Multi/AllClientsSelectionModel.h
--- a/Engine/Plugins/Developer/Concert/ConcertApp/MultiUserClient/Source/MultiUserClient/Private/Widgets/ActiveSession/Replication/Client/Multi/AllClientsSelectionModel.h
+++ b/Engine/Plugins/Developer/Concert/ConcertApp/MultiUserClient/Source/MultiUserClient/Private/Widgets/ActiveSession/Replication/Client/Multi/AllClientsSelectionModel.h
@@ -18,6 +18,7 @@ namespace UE::MultiUserClient
 
 		//~ Begin IClientSelectionModel Interface
 		virtual void ForEachSelectedClient(TFunctionRef<EBreakBehavior(FReplicationClient&)> ProcessClient) const override;
+		virtual bool ContainsClient(const FGuid& ClientId) const override { return true; }
 		virtual FOnSelectionChanged& OnSelectionChanged() override { return OnSelectionChangedDelegate; }
 		//~ End IClientSelectionModel Interface
 
diff --git a/Engine/Plugins/Developer/Concert/ConcertApp/MultiUserClient/Source/MultiUserClient/Private/Widgets/ActiveSession/Replication/Client/Multi/IClientSelectionModel.h b/Engine/Plugins/Developer/Concert/ConcertApp/MultiUserClient/Source/MultiUserClient/Private/Widgets/ActiveSession/Replication/Client/Multi/IClientSelectionModel.h
--- a/Engine/Plugins/Developer/Concert/ConcertApp/MultiUserClient/Source/MultiUserClient/Private/Widgets/ActiveSession/Replication/Client/Multi/IClientSelectionModel.h
+++ b/Engine/Plugins/Developer/Concert/ConcertApp/MultiUserClient/Source/MultiUserClient/Private/Widgets/ActiveSession/Replication/Client/Multi/IClientSelectionModel.h
@@ -19,6 +19,9 @@ namespace UE::MultiUserClient
 		/** Iterates through every client in the selection. */
 		virtual void ForEachSelectedClient(TFunctionRef<EBreakBehavior(FReplicationClient&)> ProcessClient) const = 0;
 
+		/** @return Whether the client is selcted. */
+		virtual bool ContainsClient(const FGuid& ClientId) const = 0;
+
 		DECLARE_MULTICAST_DELEGATE(FOnSelectionChanged);
 		/** Called when the clients ForEachSelectedClient enumerates has changed. */
 		virtual FOnSelectionChanged& OnSelectionChanged() = 0;
diff --git a/Engine/Plugins/Developer/Concert/ConcertApp/MultiUserClient/Source/MultiUserClient/Private/Widgets/ActiveSession/Replication/Client/Multi/SMultiClientView.cpp b/Engine/Plugins/Developer/Concert/ConcertApp/MultiUserClient/Source/MultiUserClient/Private/Widgets/ActiveSession/Replication/Client/Multi/SMultiClientView.cpp
--- a/Engine/Plugins/Developer/Concert/ConcertApp/MultiUserClient/Source/MultiUserClient/Private/Widgets/ActiveSession/Replication/Client/Multi/SMultiClientView.cpp
+++ b/Engine/Plugins/Developer/Concert/ConcertApp/MultiUserClient/Source/MultiUserClient/Private/Widgets/ActiveSession/Replication/Client/Multi/SMultiClientView.cpp
@@ -72,7 +72,7 @@ namespace UE::MultiUserClient
 		   {
 			   return StreamEditor;
 		   });
-		const TAttribute<IObjectHierarchyModel*> ObjecHierarchyAttribute =
+		const TAttribute<IObjectHierarchyModel*> ObjectHierarchyAttribute =
 		   TAttribute<IObjectHierarchyModel*>::CreateLambda([this]()
 		   {
 			   return ObjectHierarchy.Get();
@@ -113,8 +113,8 @@ namespace UE::MultiUserClient
 			.OnExtendObjectsContextMenu = FExtendObjectMenu::CreateSP(this, &SMultiClientView::ExtendObjectContextMenu),
 			.ObjectColumns =
 			{
-				MultiStreamColumns::ReplicationToggle(InConcertClient, ObjecHierarchyAttribute, InClientManager),
-				MultiStreamColumns::ReassignOwnership(InConcertClient, MultiStreamEditorAttribute, ObjecHierarchyAttribute, InClientManager.GetReassignmentLogic(), InClientManager)
+				MultiStreamColumns::ReplicationToggle(InConcertClient, ObjectHierarchyAttribute, InClientManager),
+				MultiStreamColumns::ReassignOwnership(InConcertClient, MultiStreamEditorAttribute, ObjectHierarchyAttribute, InClientManager.GetReassignmentLogic(), InClientManager)
 			}
 		};
 		
@@ -134,7 +134,7 @@ namespace UE::MultiUserClient
 		return ClientIds;
 	}
 
-	void SMultiClientView::EnumerateObjectsInStreams(TFunctionRef<void(const FSoftObjectPath&)> Consumer)
+	void SMultiClientView::EnumerateObjectsInStreams(TFunctionRef<void(const FSoftObjectPath&)> Consumer) const
 	{
 		StreamModel->ForEachClient([&Consumer](const FReplicationClient* Client)
 		{
@@ -153,27 +153,39 @@ namespace UE::MultiUserClient
 
 		ClientManager->ForEachClient([this](FReplicationClient& Client)
 		{
-			Client.OnModelChanged().AddSP(this, &SMultiClientView::OnClientChanged, Client.GetEndpointId());
+			if (SelectionModel->ContainsClient(Client.GetEndpointId()))
+			{
+				Client.OnModelChanged().AddSP(this, &SMultiClientView::OnClientChanged);
+				Client.OnHierarchyNeedsRefresh().AddRaw(this, &SMultiClientView::OnHierarchyNeedsRefresh);
+			}
+			
 			return EBreakBehavior::Continue;
 		});
 	}
 
-	void SMultiClientView::CleanClientSubscriptions()
+	void SMultiClientView::CleanClientSubscriptions() const
 	{
 		ClientManager->ForEachClient([this](FReplicationClient& Client)
 		{
 			Client.OnModelChanged().RemoveAll(this);
+			Client.OnHierarchyNeedsRefresh().RemoveAll(this);
 			return EBreakBehavior::Continue;
 		});
 	}
 
-	void SMultiClientView::OnClientChanged(FGuid)
+	void SMultiClientView::OnClientChanged() const
 	{
 		// When reassignment operations complete, the content of the columns changes so a resort is required.
 		StreamEditor->GetEditorBase().RequestObjectColumnResort(MultiStreamColumns::ReassignOwnershipColumnId);
 		StreamEditor->GetEditorBase().RequestPropertyColumnResort(MultiStreamColumns::AssignPropertyColumnId);
 	}
 
+	void SMultiClientView::OnHierarchyNeedsRefresh() const
+	{
+		// It's a bit excessive to refresh all objects when the hierarchy might have changed but it's simple (and only happens once at end of tick)
+		StreamEditor->GetEditorBase().Refresh();
+	}
+
 	void SMultiClientView::ExtendObjectContextMenu(FMenuBuilder& MenuBuilder, TConstArrayView<FSoftObjectPath> ContextObjects) const
 	{
 		FrequencyContextMenuUtils::AddFrequencyOptionsIfOneContextObject_MultiClient(MenuBuilder, ContextObjects, *ClientManager);
diff --git a/Engine/Plugins/Developer/Concert/ConcertApp/MultiUserClient/Source/MultiUserClient/Private/Widgets/ActiveSession/Replication/Client/Multi/SMultiClientView.h b/Engine/Plugins/Developer/Concert/ConcertApp/MultiUserClient/Source/MultiUserClient/Private/Widgets/ActiveSession/Replication/Client/Multi/SMultiClientView.h
--- a/Engine/Plugins/Developer/Concert/ConcertApp/MultiUserClient/Source/MultiUserClient/Private/Widgets/ActiveSession/Replication/Client/Multi/SMultiClientView.h
+++ b/Engine/Plugins/Developer/Concert/ConcertApp/MultiUserClient/Source/MultiUserClient/Private/Widgets/ActiveSession/Replication/Client/Multi/SMultiClientView.h
@@ -57,11 +57,12 @@ namespace UE::MultiUserClient
 
 		// SClientToolbar attributes
 		TSet<FGuid> GetDisplayClientIds() const;
-		void EnumerateObjectsInStreams(TFunctionRef<void(const FSoftObjectPath&)> Consumer);
+		void EnumerateObjectsInStreams(TFunctionRef<void(const FSoftObjectPath&)> Consumer) const;
 		
 		void RebuildClientSubscriptions();
-		void CleanClientSubscriptions();
-		void OnClientChanged(FGuid Guid);
+		void CleanClientSubscriptions() const;
+		void OnClientChanged() const;
+		void OnHierarchyNeedsRefresh() const;
 		
 		/** Adds additional entries to the context menu for the object tree view. */
 		void ExtendObjectContextMenu(FMenuBuilder& MenuBuilder, TConstArrayView<FSoftObjectPath> ContextObjects) const;
diff --git a/Engine/Plugins/Developer/Concert/ConcertApp/MultiUserClient/Source/MultiUserClient/Private/Widgets/ActiveSession/Replication/Client/Single/SReplicationClientView.cpp b/Engine/Plugins/Developer/Concert/ConcertApp/MultiUserClient/Source/MultiUserClient/Private/Widgets/ActiveSession/Replication/Client/Single/SReplicationClientView.cpp
--- a/Engine/Plugins/Developer/Concert/ConcertApp/MultiUserClient/Source/MultiUserClient/Private/Widgets/ActiveSession/Replication/Client/Single/SReplicationClientView.cpp
+++ b/Engine/Plugins/Developer/Concert/ConcertApp/MultiUserClient/Source/MultiUserClient/Private/Widgets/ActiveSession/Replication/Client/Single/SReplicationClientView.cpp
@@ -56,6 +56,8 @@ namespace UE::MultiUserClient
 		
 		// Refresh UI if streams change externally, e.g. a remote client changed what they sent
 		ReplicationClient->OnModelChanged().AddSP(this, &SReplicationClientView::OnModelChanged);
+		// It's a bit excessive to refresh all objects when the hierarchy might have changed but it's simple (and only happens once at end of tick)
+		ReplicationClient->OnHierarchyNeedsRefresh().AddLambda([this](){ EditorView->Refresh(); });
 	}
 
 	TSharedRef<SWidget> SReplicationClientView::CreateContent(FReplicationClient& InReplicationClient)
diff --git a/Engine/Plugins/Developer/Concert/ConcertUI/ConcertClientSharedSlate/Source/ConcertClientSharedSlate/Private/Replication/Editor/UnrealEditor/ModifyObjectInLevelHandler.cpp b/Engine/Plugins/Developer/Concert/ConcertUI/ConcertClientSharedSlate/Source/ConcertClientSharedSlate/Private/Replication/Editor/UnrealEditor/ModifyObjectInLevelHandler.cpp
new file mode 100644
--- /dev/null
+++ b/Engine/Plugins/Developer/Concert/ConcertUI/ConcertClientSharedSlate/Source/ConcertClientSharedSlate/Private/Replication/Editor/UnrealEditor/ModifyObjectInLevelHandler.cpp
@@ -0,0 +1,65 @@
+// Copyright Epic Games, Inc. All Rights Reserved.
+
+#include "Replication/Editor/UnrealEditor/ModifyObjectInLevelHandler.h"
+
+#include "Replication/Editor/Model/IEditableReplicationStreamModel.h"
+
+#include "Engine/Engine.h"
+#include "GameFramework/Actor.h"
+#include "UObject/UObjectGlobals.h"
+
+namespace UE::ConcertClientSharedSlate
+{
+	FModifyObjectInLevelHandler::FModifyObjectInLevelHandler(
+		ConcertSharedSlate::IEditableReplicationStreamModel& UpdatedModel
+		)
+		: UpdatedModel(UpdatedModel)
+	{
+		if (ensureMsgf(GEngine, TEXT("Expected to be valid at this stage")))
+		{
+			GEngine->OnLevelActorDeleted().AddRaw(this, &FModifyObjectInLevelHandler::OnActorDeleted);
+		}
+
+		// This handles a lot of cases where subobjects are added / remove to the actor hierarchy
+		FCoreUObjectDelegates::OnObjectTransacted.AddRaw(this, &FModifyObjectInLevelHandler::OnObjectTransacted);
+	}
+
+	FModifyObjectInLevelHandler::~FModifyObjectInLevelHandler()
+	{
+		if (GEngine)
+		{
+			GEngine->OnLevelActorDeleted().RemoveAll(this);
+		}
+	}
+
+	void FModifyObjectInLevelHandler::OnActorDeleted(AActor* Actor) const
+	{
+		const FSoftObjectPath ActorPath { Actor };
+		TArray<FSoftObjectPath> RemovedObjects { ActorPath };
+		UpdatedModel.ForEachSubobject(ActorPath, [&RemovedObjects](const FSoftObjectPath& Subobject)
+		{
+			RemovedObjects.Add(Subobject); return EBreakBehavior::Continue;
+		});
+		
+		if (!RemovedObjects.IsEmpty())
+		{
+			UpdatedModel.RemoveObjects(RemovedObjects);
+		}
+	}
+
+	void FModifyObjectInLevelHandler::OnObjectTransacted(UObject* Object, const FTransactionObjectEvent& TransactionObjectEvent) const
+	{
+		const bool bIsContainedInModel = UpdatedModel.ContainsObjects({ Object });
+		bool bContainsAnyChildInModel = false;
+		UpdatedModel.ForEachSubobject(Object, [&bContainsAnyChildInModel](const FSoftObjectPath& Child)
+		{
+			bContainsAnyChildInModel = true;
+			return EBreakBehavior::Break;
+		});
+
+		if (bIsContainedInModel || bContainsAnyChildInModel)
+		{
+			OnHierarchyNeedsRefreshDelegate.Broadcast();
+		}
+	}
+}
diff --git a/Engine/Plugins/Developer/Concert/ConcertUI/ConcertClientSharedSlate/Source/ConcertClientSharedSlate/Public/Replication/Editor/UnrealEditor/ModifyObjectInLevelHandler.h b/Engine/Plugins/Developer/Concert/ConcertUI/ConcertClientSharedSlate/Source/ConcertClientSharedSlate/Public/Replication/Editor/UnrealEditor/ModifyObjectInLevelHandler.h
new file mode 100644
--- /dev/null
+++ b/Engine/Plugins/Developer/Concert/ConcertUI/ConcertClientSharedSlate/Source/ConcertClientSharedSlate/Public/Replication/Editor/UnrealEditor/ModifyObjectInLevelHandler.h
@@ -0,0 +1,53 @@
+// Copyright Epic Games, Inc. All Rights Reserved.
+
+#pragma once
+
+#include "Delegates/Delegate.h"
+#include "HAL/Platform.h"
+#include "Templates/UnrealTemplate.h"
+
+class AActor;
+class FTransactionObjectEvent;
+class UObject;
+
+namespace UE::ConcertSharedSlate
+{
+	class IEditableReplicationStreamModel;
+	class IReplicationStreamViewer;
+}
+
+namespace UE::ConcertClientSharedSlate
+{
+	/**
+	 * Handles updating the passed in model in the following cases:
+	 * - Actor removed: Remove from model
+	 * - Component or subobject added / removed: refresh hierarchy
+	 */
+	class CONCERTCLIENTSHAREDSLATE_API FModifyObjectInLevelHandler : public FNoncopyable
+	{
+	public:
+		
+		/**
+		 * @param UpdatedModel The model to modified when an object is modified. The caller ensures it outlives the constructed object.
+		 */
+		FModifyObjectInLevelHandler(
+			ConcertSharedSlate::IEditableReplicationStreamModel& UpdatedModel UE_LIFETIMEBOUND
+			);
+		~FModifyObjectInLevelHandler();
+
+		DECLARE_MULTICAST_DELEGATE(FOnHierarchyNeedsRefresh);
+		/** Broadcasts when the hierarchy may have changed. */
+		FOnHierarchyNeedsRefresh& OnHierarchyNeedsRefresh() { return OnHierarchyNeedsRefreshDelegate; }
+
+	private:
+
+		/** The model that is updated when an object is modified. */
+		ConcertSharedSlate::IEditableReplicationStreamModel& UpdatedModel;
+
+		/** Broadcasts when the hierarchy may have changed. */
+		FOnHierarchyNeedsRefresh OnHierarchyNeedsRefreshDelegate;
+		
+		void OnActorDeleted(AActor* Actor) const;
+		void OnObjectTransacted(UObject* Object, const FTransactionObjectEvent& TransactionObjectEvent) const;
+	};
+}
