MU replication: Fix several contributing issues that caused assigning different clients to fail and block the replication system
- Struct property stays assigned when reassigning child (UE-207448)
- Property assignment combo box does not reflect correct state (UE-207449)
- Added clear button to unassign properties from all clients (trivial)
- Fix property view closing when you remove the last property from an actor

diff --git a/Engine/Plugins/Developer/Concert/ConcertApp/MultiUserClient/Source/MultiUserClient/Private/Widgets/ActiveSession/Replication/Client/Multi/Columns/MultiStreamColumns_AssignProperties.cpp b/Engine/Plugins/Developer/Concert/ConcertApp/MultiUserClient/Source/MultiUserClient/Private/Widgets/ActiveSession/Replication/Client/Multi/Columns/MultiStreamColumns_AssignProperties.cpp
--- a/Engine/Plugins/Developer/Concert/ConcertApp/MultiUserClient/Source/MultiUserClient/Private/Widgets/ActiveSession/Replication/Client/Multi/Columns/MultiStreamColumns_AssignProperties.cpp
+++ b/Engine/Plugins/Developer/Concert/ConcertApp/MultiUserClient/Source/MultiUserClient/Private/Widgets/ActiveSession/Replication/Client/Multi/Columns/MultiStreamColumns_AssignProperties.cpp
@@ -111,7 +111,7 @@ namespace UE::MultiUserClient::MultiStreamColumns
 					.DisplayedProperty(InArgs.RowItem.RowData.GetProperty())
 					.EditedObjects(DisplayedObjects)
 					.HighlightText(InArgs.HighlightText)
-					.OnOptionSelected_Lambda([this](auto)
+					.OnPropertyAssignmentChanged_Lambda([this]()
 					{
 						if (const TSharedPtr<IMultiReplicationStreamEditor> Editor = MultiStreamEditor.Get())
 						{
diff --git a/Engine/Plugins/Developer/Concert/ConcertApp/MultiUserClient/Source/MultiUserClient/Private/Widgets/ActiveSession/Replication/Client/Multi/Columns/SAssignPropertyComboBox.cpp b/Engine/Plugins/Developer/Concert/ConcertApp/MultiUserClient/Source/MultiUserClient/Private/Widgets/ActiveSession/Replication/Client/Multi/Columns/SAssignPropertyComboBox.cpp
--- a/Engine/Plugins/Developer/Concert/ConcertApp/MultiUserClient/Source/MultiUserClient/Private/Widgets/ActiveSession/Replication/Client/Multi/Columns/SAssignPropertyComboBox.cpp
+++ b/Engine/Plugins/Developer/Concert/ConcertApp/MultiUserClient/Source/MultiUserClient/Private/Widgets/ActiveSession/Replication/Client/Multi/Columns/SAssignPropertyComboBox.cpp
@@ -15,10 +15,12 @@
 
 #include "Algo/AnyOf.h"
 #include "Framework/MultiBox/MultiBoxBuilder.h"
-#include "Widgets/Input/SComboButton.h"
+#include "GameFramework/Actor.h"
+#include "UObject/Class.h"
 #include "ScopedTransaction.h"
 #include "Styling/AppStyle.h"
 #include "Styling/SlateTypes.h"
+#include "Widgets/Input/SComboButton.h"
 
 #define LOCTEXT_NAMESPACE "SAssignPropertyComboBox"
 
@@ -77,7 +79,7 @@ namespace UE::MultiUserClient
 		HighlightText = InArgs._HighlightText;
 		check(!EditedObjects.IsEmpty());
 
-		OnOptionClickedDelegate = InArgs._OnOptionSelected;
+		OnOptionClickedDelegate = InArgs._OnPropertyAssignmentChanged;
 		
 		ChildSlot
 		[
@@ -125,6 +127,19 @@ namespace UE::MultiUserClient
 		};
 		
 		FMenuBuilder MenuBuilder(true, nullptr);
+		MenuBuilder.AddMenuEntry(
+			LOCTEXT("Clear.Label", "Clear"),
+			LOCTEXT("Clear.Tooltip", "Stop this property from being replicated"),
+			FSlateIcon(),
+			FUIAction(
+				FExecuteAction::CreateSP(this, &SAssignPropertyComboBox::OnClickClear),
+				FCanExecuteAction::CreateSP(this, &SAssignPropertyComboBox::CanClickClear)
+				),
+			NAME_None,
+			EUserInterfaceActionType::Button
+		);
+		
+		MenuBuilder.BeginSection(TEXT("AssignTo"), LOCTEXT("AssignTo", "Assign to"));
 		for (const FReplicationClient* Client : ClientUtils::GetSortedClientList(*ConcertClient, *ClientManager))
 		{
 			TAttribute<FText> Tooltip = TAttribute<FText>::CreateLambda([this, EndpointId = Client->GetEndpointId()]()
@@ -154,9 +169,10 @@ namespace UE::MultiUserClient
 				MakeWidget(Client->GetEndpointId()),
 				NAME_None,
 				Tooltip,
-				EUserInterfaceActionType::ToggleButton
+				EUserInterfaceActionType::Check
 				);
 		}
+		MenuBuilder.EndSection();
 		
 		return MenuBuilder.MakeWidget();
 	}
@@ -173,35 +189,23 @@ namespace UE::MultiUserClient
 		const FText TransactionText = FText::Format(LOCTEXT("AllClientsAssignFmt", "Assign {0} property"), FText::FromString(Property.ToString(FConcertPropertyChain::EToStringMethod::LeafProperty)));
 		FScopedTransaction Transaction(TransactionText);
 
-		// Only one client is supposed to own the object: remove the other clients (if possible)
-		ClientManager->ForEachClient([this, Client](const FReplicationClient& ClientToRemoveFrom)
-		{
-			if (*Client != ClientToRemoveFrom
-				&& ClientToRemoveFrom.AllowsEditing())
-			{
-				for (const FSoftObjectPath& ObjectPath : EditedObjects)
-				{
-					ClientToRemoveFrom.GetClientEditModel()->RemoveProperties(ObjectPath, { Property });
-				}
-			}
-			
-			return EBreakBehavior::Continue;
-		});
-
 		const ECheckBoxState CheckBoxState = GetOptionCheckState(EndpointId);
-		const bool bRemoveProperty = CheckBoxState == ECheckBoxState::Checked; 
-		const TSharedRef<ConcertSharedSlate::IEditableReplicationStreamModel> EditModel = Client->GetClientEditModel();
-		for (const FSoftObjectPath& ObjectPath : EditedObjects)
+		const bool bRemovePropertyFromEditedClient = CheckBoxState == ECheckBoxState::Checked;
+		
+		// To make it simpler for the user, at most one client is supposed to be assigned to the object at any given time so ...
+		if (bRemovePropertyFromEditedClient)
 		{
-			if (bRemoveProperty)
-			{
-				EditModel->RemoveProperties(ObjectPath, { Property });
-				if (!EditModel->HasAnyPropertyAssigned(ObjectPath))
-				{
-					EditModel->RemoveObjects({ ObjectPath });
-				}
-			}
-			else
+			// ... remove property from all clients
+			UnassignPropertyFromClients([](const FReplicationClient& ClientToRemoveFrom){ return true; });
+		}
+		else
+		{
+			// ... remove the property from all clients but the one we'll assign to ...
+			UnassignPropertyFromClients([Client](const FReplicationClient& ClientToRemoveFrom){ return *Client != ClientToRemoveFrom; });
+
+			// ... and then assign the property
+			const TSharedRef<ConcertSharedSlate::IEditableReplicationStreamModel> EditModel = Client->GetClientEditModel();
+			for (const FSoftObjectPath& ObjectPath : EditedObjects)
 			{
 				if (!EditModel->ContainsObjects({ ObjectPath }))
 				{
@@ -215,7 +219,7 @@ namespace UE::MultiUserClient
 			}
 		}
 		
-		OnOptionClickedDelegate.ExecuteIfBound(EndpointId);
+		OnOptionClickedDelegate.ExecuteIfBound();
 	}
 
 #define SET_REASON(Text) if (Reason) { *Reason = Text; }
@@ -304,7 +308,67 @@ namespace UE::MultiUserClient
 
 		return CheckBoxState;
 	}
-	
+
+	void SAssignPropertyComboBox::OnClickClear()
+	{
+		const FText TransactionText = FText::Format(LOCTEXT("ClearAllClientsFmt", "Clear {0} property"), FText::FromString(Property.ToString(FConcertPropertyChain::EToStringMethod::LeafProperty)));
+		FScopedTransaction Transaction(TransactionText);
+
+		UnassignPropertyFromClients([](const FReplicationClient& ClientToRemoveFrom){ return true; });
+		OnOptionClickedDelegate.ExecuteIfBound();
+	}
+
+	bool SAssignPropertyComboBox::CanClickClear() const
+	{
+		bool bIsAssignedToAnyClient = false;
+		ClientManager->ForEachClient([this, &bIsAssignedToAnyClient](const FReplicationClient& Client)
+		{
+			for (int32 i = 0; !bIsAssignedToAnyClient && i < EditedObjects.Num(); ++i)
+			{
+				const FSoftObjectPath& ObjectPath = EditedObjects[i];
+				const TSharedRef<ConcertSharedSlate::IEditableReplicationStreamModel> Model = Client.GetClientEditModel();
+				const bool bHasProperty = Model->HasProperty(ObjectPath, Property);
+				bIsAssignedToAnyClient |= bHasProperty;
+			}
+			
+			return bIsAssignedToAnyClient ? EBreakBehavior::Break : EBreakBehavior::Continue;
+		});
+		return bIsAssignedToAnyClient;
+	}
+
+	void SAssignPropertyComboBox::UnassignPropertyFromClients(TFunctionRef<bool(const FReplicationClient& Client)> ShouldRemoveFromClient) const
+	{
+		ClientManager->ForEachClient([this, &ShouldRemoveFromClient](const FReplicationClient& ClientToRemoveFrom)
+		{
+			const TSharedRef<ConcertSharedSlate::IEditableReplicationStreamModel> EditModel = ClientToRemoveFrom.GetClientEditModel();
+			if (ClientToRemoveFrom.AllowsEditing() && ShouldRemoveFromClient(ClientToRemoveFrom))
+			{
+				for (const FSoftObjectPath& ObjectPath : EditedObjects)
+				{
+					EditModel->RemoveProperties(ObjectPath, { Property });
+					
+					if (EditModel->HasAnyPropertyAssigned(ObjectPath))
+					{
+						continue;
+					}
+
+					// We want to remove subobjects that have no properties. Retain actors because they cause their entire component / subobject hierarchy to be displayed.
+					// Skipping this check would close the entire property tree view and remove the actor hierarchy from the view.
+					// That would feel very unnatural / unexpected for the user. 
+					// If the user does not want the actor anymore, they should click it and delete it.
+					const UClass* ObjectClass = EditModel->GetObjectClass(ObjectPath).TryLoadClass<UObject>();
+					const bool bIsTopLevelObject = !ObjectClass->IsChildOf<AActor>();
+					if (bIsTopLevelObject)
+					{
+						EditModel->RemoveObjects({ ObjectPath });
+					}
+				}
+			}
+			
+			return EBreakBehavior::Continue;
+		});
+	}
+
 	void SAssignPropertyComboBox::RebuildSubscriptions()
 	{
 		ClientManager->ForEachClient([this](FReplicationClient& Client)
diff --git a/Engine/Plugins/Developer/Concert/ConcertApp/MultiUserClient/Source/MultiUserClient/Private/Widgets/ActiveSession/Replication/Client/Multi/Columns/SAssignPropertyComboBox.h b/Engine/Plugins/Developer/Concert/ConcertApp/MultiUserClient/Source/MultiUserClient/Private/Widgets/ActiveSession/Replication/Client/Multi/Columns/SAssignPropertyComboBox.h
--- a/Engine/Plugins/Developer/Concert/ConcertApp/MultiUserClient/Source/MultiUserClient/Private/Widgets/ActiveSession/Replication/Client/Multi/Columns/SAssignPropertyComboBox.h
+++ b/Engine/Plugins/Developer/Concert/ConcertApp/MultiUserClient/Source/MultiUserClient/Private/Widgets/ActiveSession/Replication/Client/Multi/Columns/SAssignPropertyComboBox.h
@@ -5,13 +5,12 @@
 #include "Replication/Data/ConcertPropertySelection.h"
 
 #include "Containers/Array.h"
+#include "Templates/Function.h"
 #include "Widgets/DeclarativeSyntaxSupport.h"
 #include "Widgets/SCompoundWidget.h"
 
 class IConcertClient;
-
 enum class ECheckBoxState : uint8;
-
 struct FGuid;
 
 namespace UE::ConcertClientSharedSlate { class SHorizontalClientList; }
@@ -19,6 +18,7 @@ namespace UE::ConcertSharedSlate { class IMultiReplicationStreamEditor; }
 
 namespace UE::MultiUserClient
 {
+	class FReplicationClient;
 	class FReplicationClientManager;
 	
 	/**
@@ -31,7 +31,7 @@ namespace UE::MultiUserClient
 	{
 	public:
 
-		DECLARE_DELEGATE_OneParam(FOnOptionClicked, const FGuid& ClientId);
+		DECLARE_DELEGATE(FOnPropertyAssignmentChanged);
 
 		/** @return The display string this widget would have with the given state. If unset, no clients are displayed in the combobox.*/
 		static TOptional<FString> GetDisplayString(
@@ -46,8 +46,8 @@ namespace UE::MultiUserClient
 			SLATE_ARGUMENT(FConcertPropertyChain, DisplayedProperty)
 			SLATE_ARGUMENT(TArray<FSoftObjectPath>, EditedObjects)
 			SLATE_ARGUMENT(TSharedPtr<FText>, HighlightText)
-			/** Called when a valid option is selected. */
-			SLATE_EVENT(FOnOptionClicked, OnOptionSelected)
+			/** Called when the property assignment of any client(s) is changed by this widget. */
+			SLATE_EVENT(FOnPropertyAssignmentChanged, OnPropertyAssignmentChanged)
 		SLATE_END_ARGS()
 
 		void Construct(const FArguments& InArgs,
@@ -74,8 +74,8 @@ namespace UE::MultiUserClient
 		/** Passed to MakeListWidgetDelegate */
 		TSharedPtr<FText> HighlightText;
 
-		/** Called when a valid option is selected. */
-		FOnOptionClicked OnOptionClickedDelegate;
+		/** Called when the property assignment of any client(s) is changed by this widget. */
+		FOnPropertyAssignmentChanged OnOptionClickedDelegate;
 
 		/** Updates the content of the combo box */
 		void RefreshContentBoxContent() const;
@@ -85,14 +85,20 @@ namespace UE::MultiUserClient
 
 		/** Called when a client is clicked in the drop-down menu */
 		void OnClickOption(const FGuid EndpointId) const;
-		
 		/** @return Whether an item in the drop-down menu should be clickable (returns true if the stream is writable) */
 		bool CanClickOption(const FGuid EndpointId) const { return CanClickOptionWithReason(EndpointId, nullptr); }
 		bool CanClickOptionWithReason(const FGuid& EndpointId, FText* Reason = nullptr) const;
-
 		/** @return Checkbox state for the item. Handles cases of multiple objects being property edited. */
 		ECheckBoxState GetOptionCheckState(const FGuid EndpointId) const;
 
+		/** Called when a client clicks the clear button */
+		void OnClickClear();
+		/** Whether the property is assigned to any client */
+		bool CanClickClear() const;
+
+		/** Unassigns this widget's property from all clients passing the predicate and removes the object from the model if it is a subobject. */
+		void UnassignPropertyFromClients(TFunctionRef<bool(const FReplicationClient& Client)> ShouldRemoveFromClient) const;
+
 		/** Resubscribes to all clients changing. */
 		void RebuildSubscriptions();
 		void RebuildSubscriptionsAndRefresh() { RebuildSubscriptions(); RefreshContentBoxContent(); }
diff --git a/Engine/Plugins/Developer/Concert/ConcertSync/ConcertSyncTest/ConcertSyncTest.uplugin b/Engine/Plugins/Developer/Concert/ConcertSync/ConcertSyncTest/ConcertSyncTest.uplugin
--- a/Engine/Plugins/Developer/Concert/ConcertSync/ConcertSyncTest/ConcertSyncTest.uplugin
+++ b/Engine/Plugins/Developer/Concert/ConcertSync/ConcertSyncTest/ConcertSyncTest.uplugin
@@ -29,6 +29,10 @@
 			"Name": "ConcertClientSharedSlate",
 			"Enabled": true
 		},
+		{
+			"Name": "ConcertSharedSlate",
+			"Enabled": true
+		},
 		{
 			"Name": "ConcertSyncClient",
 			"Enabled": true
diff --git a/Engine/Plugins/Developer/Concert/ConcertSync/ConcertSyncTest/Source/ConcertSyncTest/ConcertSyncTest.Build.cs b/Engine/Plugins/Developer/Concert/ConcertSync/ConcertSyncTest/Source/ConcertSyncTest/ConcertSyncTest.Build.cs
--- a/Engine/Plugins/Developer/Concert/ConcertSync/ConcertSyncTest/Source/ConcertSyncTest/ConcertSyncTest.Build.cs
+++ b/Engine/Plugins/Developer/Concert/ConcertSync/ConcertSyncTest/Source/ConcertSyncTest/ConcertSyncTest.Build.cs
@@ -21,6 +21,7 @@ namespace UnrealBuildTool.Rules
 				{
 					"Concert",
 					"ConcertClientSharedSlate",
+					"ConcertSharedSlate",
 					"ConcertTransport",
 					"ConcertSyncCore",
 					"ConcertSyncClient",
diff --git a/Engine/Plugins/Developer/Concert/ConcertSync/ConcertSyncTest/Source/ConcertSyncTest/Private/Replication/UI/GenericReplicationStreamModel.spec.cpp b/Engine/Plugins/Developer/Concert/ConcertSync/ConcertSyncTest/Source/ConcertSyncTest/Private/Replication/UI/GenericReplicationStreamModel.spec.cpp
new file mode 100644
--- /dev/null
+++ b/Engine/Plugins/Developer/Concert/ConcertSync/ConcertSyncTest/Source/ConcertSyncTest/Private/Replication/UI/GenericReplicationStreamModel.spec.cpp
@@ -0,0 +1,82 @@
+// Copyright Epic Games, Inc. All Rights Reserved.
+
+#include "Misc/AutomationTest.h"
+
+#include "Replication/ReplicationWidgetFactories.h"
+#include "Replication/TestReflectionObject.h"
+#include "Replication/Data/ObjectReplicationMap.h"
+#include "Replication/Editor/Model/IEditableReplicationStreamModel.h"
+
+namespace UE::ConcertSyncTests::Replication::UI
+{
+	BEGIN_DEFINE_SPEC(FReplicationStreamSpec, "Editor.Concert.Replication.UI", EAutomationTestFlags::EditorContext | EAutomationTestFlags::EngineFilter)
+	FConcertObjectReplicationMap ReplicationMap;
+	TSharedPtr<ConcertSharedSlate::IEditableReplicationStreamModel> Model;
+	UTestReflectionObject* Object;
+
+	TOptional<FConcertPropertyChain> Prop_TestStruct;
+	TOptional<FConcertPropertyChain> Prop_TestStruct_Nested;
+	TOptional<FConcertPropertyChain> Prop_TestStruct_Nested_Vector;
+	TOptional<FConcertPropertyChain> Prop_TestStruct_Nested_Vector_X;
+	TOptional<FConcertPropertyChain> Prop_TestStruct_NestedArray;
+	TOptional<FConcertPropertyChain> Prop_TestStruct_NestedArray_Value;
+	END_DEFINE_SPEC(FReplicationStreamSpec);
+
+	void FReplicationStreamSpec::Define()
+	{
+		Object = GetMutableDefault<UTestReflectionObject>();
+		Prop_TestStruct
+			= FConcertPropertyChain::CreateFromPath(*UTestReflectionObject::StaticClass(), { TEXT("TestStruct") });
+		Prop_TestStruct_Nested
+			= FConcertPropertyChain::CreateFromPath(*UTestReflectionObject::StaticClass(), { TEXT("TestStruct"), TEXT("Nested") });
+		Prop_TestStruct_Nested_Vector
+			= FConcertPropertyChain::CreateFromPath(*UTestReflectionObject::StaticClass(), { TEXT("TestStruct"), TEXT("Nested"), TEXT("Vector") });
+		Prop_TestStruct_Nested_Vector_X
+			= FConcertPropertyChain::CreateFromPath(*UTestReflectionObject::StaticClass(), { TEXT("TestStruct"), TEXT("Nested"), TEXT("Vector"), TEXT("X") });
+		Prop_TestStruct_NestedArray
+			= FConcertPropertyChain::CreateFromPath(*UTestReflectionObject::StaticClass(), { TEXT("TestStruct"), TEXT("NestedArray") });
+		Prop_TestStruct_NestedArray_Value
+			= FConcertPropertyChain::CreateFromPath(*UTestReflectionObject::StaticClass(), { TEXT("TestStruct"), TEXT("NestedArray"), TEXT("Value") });
+		
+		BeforeEach([this]
+		{
+			if (!Prop_TestStruct || !Prop_TestStruct_Nested || !Prop_TestStruct_Nested_Vector || !Prop_TestStruct_Nested_Vector_X ||! Prop_TestStruct_NestedArray || !Prop_TestStruct_NestedArray_Value)
+			{
+				AddError(TEXT("FConcertPropertyChain::CreateFromPath failed"));
+			}
+
+			ReplicationMap = FConcertObjectReplicationMap{};
+			Model = UE::ConcertSharedSlate::CreateBaseStreamModel(&ReplicationMap);
+			Model->AddObjects({ Object });
+		});
+
+		Describe("Adding a property", [this]()
+		{
+			It("Adds all parent properties", [this]()
+			{
+				Model->AddProperties(Object, { *Prop_TestStruct_Nested_Vector_X });
+
+				TestEqual(TEXT("4 properties added"), Model->GetAllProperties(Object).Num(), 4);
+				TestTrue(TEXT("TestStruct"), Model->HasProperty(Object, *Prop_TestStruct));
+				TestTrue(TEXT("TestStruct.Nested"), Model->HasProperty(Object, *Prop_TestStruct_Nested));
+				TestTrue(TEXT("TestStruct.Nested.Vector"), Model->HasProperty(Object, *Prop_TestStruct_Nested_Vector));
+				TestTrue(TEXT("TestStruct.Nested.Vector.X"), Model->HasProperty(Object, *Prop_TestStruct_Nested_Vector_X));
+			});
+		});
+		
+		Describe("Removing a property", [this]()
+		{
+			It("Should remove parent properties left with 0 children", [this]()
+			{
+				Model->AddProperties(Object, { *Prop_TestStruct_Nested_Vector_X, *Prop_TestStruct_NestedArray_Value });
+				Model->RemoveProperties(Object, { *Prop_TestStruct_Nested_Vector_X });
+				
+				// Remove TestStruct.Nested.Vector and TestStruct.Nested as well but not TestStruct since it still references TestStruct.NestedArray.
+				TestEqual(TEXT("3 properties remain"), Model->GetAllProperties(Object).Num(), 3);
+				TestTrue(TEXT("TestStruct"), Model->HasProperty(Object, *Prop_TestStruct));
+				TestTrue(TEXT("TestStruct.Nested.NestedArray"), Model->HasProperty(Object, *Prop_TestStruct_NestedArray));
+				TestTrue(TEXT("TestStruct.Nested.NestedArray.Value"), Model->HasProperty(Object, *Prop_TestStruct_NestedArray_Value));
+			});
+		});
+	}
+}
diff --git a/Engine/Plugins/Developer/Concert/ConcertUI/ConcertSharedSlate/Source/ConcertSharedSlate/Private/Replication/Editor/Model/GenericReplicationStreamModel.cpp b/Engine/Plugins/Developer/Concert/ConcertUI/ConcertSharedSlate/Source/ConcertSharedSlate/Private/Replication/Editor/Model/GenericReplicationStreamModel.cpp
--- a/Engine/Plugins/Developer/Concert/ConcertUI/ConcertSharedSlate/Source/ConcertSharedSlate/Private/Replication/Editor/Model/GenericReplicationStreamModel.cpp
+++ b/Engine/Plugins/Developer/Concert/ConcertUI/ConcertSharedSlate/Source/ConcertSharedSlate/Private/Replication/Editor/Model/GenericReplicationStreamModel.cpp
@@ -29,6 +29,39 @@ namespace UE::ConcertSharedSlate
 				});
 			return bAddedAtLeastOne;
 		}
+
+		/** After removing RemovedProperty from Selection, walk up the chain property chain and remove any parent properties that now have 0 children. */
+		static int32 RemoveParentPropertiesWithoutChildren(const FConcertPropertyChain& RemovedProperty, FConcertPropertySelection& Selection)
+		{
+			const auto HasNoChildren = [&Selection](const FConcertPropertyChain& CheckedProperty)
+			{
+				const bool bHasNoChildren = Algo::AllOf(Selection.ReplicatedProperties, [&CheckedProperty](const FConcertPropertyChain& SelectedProperty)
+				{
+					return !SelectedProperty.IsChildOf(CheckedProperty) || CheckedProperty == SelectedProperty;
+				});
+				return bHasNoChildren;
+			};
+
+			int32 NumRemoved = 0;
+			bool bFollowChain = true;
+			FConcertPropertyChain ParentProperty = RemovedProperty;
+			while (!ParentProperty.IsEmpty() && bFollowChain)
+			{
+				ParentProperty = ParentProperty.GetParent();
+				
+				const bool bHasNoChildren = HasNoChildren(ParentProperty);
+				if (bHasNoChildren)
+				{
+					++NumRemoved;
+					Selection.ReplicatedProperties.Remove(ParentProperty);
+				}
+
+				// If current property still has children then transitively its parents will also continue to have a child
+				bFollowChain = bHasNoChildren;
+			}
+
+			return NumRemoved;
+		}
 	}
 
 	FGenericReplicationStreamModel::FGenericReplicationStreamModel(
@@ -254,6 +287,8 @@ namespace UE::ConcertSharedSlate
 				}
 				return EBreakBehavior::Continue;
 			});
+
+			NumRemoved += Private::RemoveParentPropertiesWithoutChildren(RemovedProperty, AssignedProperties->PropertySelection);
 		}
 		
 		if (NumRemoved > 0)
