Move the "Replicating x objects for y actors" status text to the bottom of the widget instead of at the top-right corner.
This makes it more consistent with the rest of the engine.

diff --git a/Engine/Plugins/Developer/Concert/ConcertApp/MultiUserClient/Source/MultiUserClient/Private/Widgets/ActiveSession/Replication/Client/Multi/SMultiClientView.cpp b/Engine/Plugins/Developer/Concert/ConcertApp/MultiUserClient/Source/MultiUserClient/Private/Widgets/ActiveSession/Replication/Client/Multi/SMultiClientView.cpp
--- a/Engine/Plugins/Developer/Concert/ConcertApp/MultiUserClient/Source/MultiUserClient/Private/Widgets/ActiveSession/Replication/Client/Multi/SMultiClientView.cpp
+++ b/Engine/Plugins/Developer/Concert/ConcertApp/MultiUserClient/Source/MultiUserClient/Private/Widgets/ActiveSession/Replication/Client/Multi/SMultiClientView.cpp
@@ -15,6 +15,7 @@
 #include "Widgets/ActiveSession/Replication/Client/FrequencyContextMenuUtils.h"
 #include "Widgets/ActiveSession/Replication/Client/Multi/Columns/MultiStreamColumns.h"
 #include "Widgets/ActiveSession/Replication/Client/SClientToolbar.h"
+#include "Widgets/ActiveSession/Replication/Client/SReplicationStatus.h"
 
 #include "Widgets/SBoxPanel.h"
 
@@ -30,20 +31,19 @@ namespace UE::MultiUserClient
 		ClientManager->OnRemoteClientsChanged().AddSP(this, &SMultiClientView::RebuildClientSubscriptions);
 		SelectionModel = &InDisplayClientsModel;
 		SelectionModel->OnSelectionChanged().AddSP(this, &SMultiClientView::RebuildClientSubscriptions);
-		
+
+		TSharedPtr<SVerticalBox> Content;
 		ChildSlot
 		[
-			SNew(SVerticalBox)
+			SAssignNew(Content, SVerticalBox)
 
 			// Toolbar
 			+SVerticalBox::Slot()
 			.AutoHeight()
 			.Padding(2.f)
 			[
-				SAssignNew(Toolbar, SClientToolbar, InClientManager.GetAuthorityCache())
+				SNew(SClientToolbar)
 				.ViewSelectionArea() [ InArgs._ViewSelectionArea.Widget ]
-				.DisplayedClients(this, &SMultiClientView::GetDisplayClientIds)
-				.ForEachReplicatedObject(this, &SMultiClientView::EnumerateObjectsInStreams)
 			]
 
 			// Editor
@@ -53,6 +53,12 @@ namespace UE::MultiUserClient
 				CreateEditorContent(InConcertClient, InClientManager)
 			]
 		];
+		
+		SReplicationStatus::AppendReplicationStatus(*Content, InClientManager.GetAuthorityCache(),
+			SReplicationStatus::FArguments()
+			.DisplayedClients(this, &SMultiClientView::GetDisplayClientIds)
+			.ForEachReplicatedObject(this, &SMultiClientView::EnumerateObjectsInStreams)
+			);
 
 		RebuildClientSubscriptions();
 	}
diff --git a/Engine/Plugins/Developer/Concert/ConcertApp/MultiUserClient/Source/MultiUserClient/Private/Widgets/ActiveSession/Replication/Client/Multi/SMultiClientView.h b/Engine/Plugins/Developer/Concert/ConcertApp/MultiUserClient/Source/MultiUserClient/Private/Widgets/ActiveSession/Replication/Client/Multi/SMultiClientView.h
--- a/Engine/Plugins/Developer/Concert/ConcertApp/MultiUserClient/Source/MultiUserClient/Private/Widgets/ActiveSession/Replication/Client/Multi/SMultiClientView.h
+++ b/Engine/Plugins/Developer/Concert/ConcertApp/MultiUserClient/Source/MultiUserClient/Private/Widgets/ActiveSession/Replication/Client/Multi/SMultiClientView.h
@@ -22,7 +22,6 @@ namespace UE::MultiUserClient
 	class FReplicationClient;
 	class FReplicationClientManager;
 	class IClientSelectionModel;
-	class SClientToolbar;
 
 	/** Displays a selection of clients. */
 	class SMultiClientView
@@ -39,9 +38,6 @@ namespace UE::MultiUserClient
 
 	private:
 
-		/** Used to access and update the replication status widget. */
-		TSharedPtr<SClientToolbar> Toolbar;
-
 		FReplicationClientManager* ClientManager = nullptr;
 		IClientSelectionModel* SelectionModel = nullptr;
 		
diff --git a/Engine/Plugins/Developer/Concert/ConcertApp/MultiUserClient/Source/MultiUserClient/Private/Widgets/ActiveSession/Replication/Client/SClientToolbar.cpp b/Engine/Plugins/Developer/Concert/ConcertApp/MultiUserClient/Source/MultiUserClient/Private/Widgets/ActiveSession/Replication/Client/SClientToolbar.cpp
--- a/Engine/Plugins/Developer/Concert/ConcertApp/MultiUserClient/Source/MultiUserClient/Private/Widgets/ActiveSession/Replication/Client/SClientToolbar.cpp
+++ b/Engine/Plugins/Developer/Concert/ConcertApp/MultiUserClient/Source/MultiUserClient/Private/Widgets/ActiveSession/Replication/Client/SClientToolbar.cpp
@@ -9,7 +9,7 @@
 
 namespace UE::MultiUserClient
 {
-	void SClientToolbar::Construct(const FArguments& InArgs, FGlobalAuthorityCache& InAuthorityCache)
+	void SClientToolbar::Construct(const FArguments& InArgs)
 	{
 		ChildSlot
 		[
@@ -21,27 +21,6 @@ namespace UE::MultiUserClient
 			[
 				InArgs._ViewSelectionArea.Widget
 			]
-			
-			+SHorizontalBox::Slot()
-			.FillWidth(1.f)
-			[
-				SNew(SSpacer)
-			]
-
-			+SHorizontalBox::Slot()
-			.AutoWidth()
-			.VAlign(VAlign_Center)
-			.Padding(0.f, 0.f, 5.f, 0.f)
-			[
-				SAssignNew(ReplicationStatus, SReplicationStatus, InAuthorityCache)
-				.DisplayedClients(InArgs._DisplayedClients)
-				.ForEachReplicatedObject(InArgs._ForEachReplicatedObject)
-			]
 		];
 	}
-
-	void SClientToolbar::RefreshStatusText()
-	{
-		ReplicationStatus->RefreshStatusText();
-	}
 }
diff --git a/Engine/Plugins/Developer/Concert/ConcertApp/MultiUserClient/Source/MultiUserClient/Private/Widgets/ActiveSession/Replication/Client/SClientToolbar.h b/Engine/Plugins/Developer/Concert/ConcertApp/MultiUserClient/Source/MultiUserClient/Private/Widgets/ActiveSession/Replication/Client/SClientToolbar.h
--- a/Engine/Plugins/Developer/Concert/ConcertApp/MultiUserClient/Source/MultiUserClient/Private/Widgets/ActiveSession/Replication/Client/SClientToolbar.h
+++ b/Engine/Plugins/Developer/Concert/ConcertApp/MultiUserClient/Source/MultiUserClient/Private/Widgets/ActiveSession/Replication/Client/SClientToolbar.h
@@ -2,26 +2,12 @@
 
 #pragma once
 
-#include "SReplicationStatus.h"
-
 #include "Widgets/DeclarativeSyntaxSupport.h"
 #include "Widgets/SCompoundWidget.h"
 
-namespace UE::ConcertSharedSlate
-{
-	class IReplicationStreamModel;
-}
-
 namespace UE::MultiUserClient
 {
-	class FGlobalAuthorityCache;
-	class FStreamChangeTracker;
-	class ISubmissionWorkflow;
-	class SReplicationStatus;
-
-	/**
-	 * Contains a bunch of actions that can be performed on client view.
-	 */
+	/** Contains a shared actions that can be performed on a client view. */
 	class SClientToolbar : public SCompoundWidget
 	{
 	public:
@@ -30,23 +16,9 @@ namespace UE::MultiUserClient
 		{}
 			/** Dedicated space for a widget with which to change the view. */
 			SLATE_NAMED_SLOT(FArguments, ViewSelectionArea)
-		
-			/** The clients to show statistics for */
-			SLATE_ATTRIBUTE(TSet<FGuid>, DisplayedClients)
-		
-			/** Delegate which enumerates every replicated object. */
-			SLATE_EVENT(FForEachReplicatedObject, ForEachReplicatedObject)
 		SLATE_END_ARGS()
 
-		void Construct(const FArguments& InArgs, FGlobalAuthorityCache& InAuthorityCache);
-
-		/** Updates the status text of how many objects are being replicated. */
-		void RefreshStatusText();
-		
-	private:
-
-		/** Displays how many objects are being replicated */
-		TSharedPtr<SReplicationStatus> ReplicationStatus;
+		void Construct(const FArguments& InArgs);
 	};
 }
 
diff --git a/Engine/Plugins/Developer/Concert/ConcertApp/MultiUserClient/Source/MultiUserClient/Private/Widgets/ActiveSession/Replication/Client/SReplicationStatus.cpp b/Engine/Plugins/Developer/Concert/ConcertApp/MultiUserClient/Source/MultiUserClient/Private/Widgets/ActiveSession/Replication/Client/SReplicationStatus.cpp
--- a/Engine/Plugins/Developer/Concert/ConcertApp/MultiUserClient/Source/MultiUserClient/Private/Widgets/ActiveSession/Replication/Client/SReplicationStatus.cpp
+++ b/Engine/Plugins/Developer/Concert/ConcertApp/MultiUserClient/Source/MultiUserClient/Private/Widgets/ActiveSession/Replication/Client/SReplicationStatus.cpp
@@ -6,6 +6,7 @@
 
 #include "Algo/AnyOf.h"
 #include "Styling/AppStyle.h"
+#include "Widgets/Layout/SSeparator.h"
 #include "Widgets/SBoxPanel.h"
 #include "Widgets/Text/STextBlock.h"
 
@@ -44,7 +45,26 @@ namespace UE::MultiUserClient
 			return PathToOwningActor;
 		}
 	}
-	
+
+	void SReplicationStatus::AppendReplicationStatus(SVerticalBox& VerticalBox, FGlobalAuthorityCache& InAuthorityCache, const FArguments& InArgs)
+	{
+		VerticalBox.AddSlot()
+			.AutoHeight()
+			.Padding(4, 3)
+			[
+				SNew(SSeparator)
+			];
+
+		VerticalBox.AddSlot()
+			.AutoHeight()
+			.Padding(4, 0, 4, 3)
+			[
+				SNew(SReplicationStatus, InAuthorityCache)
+				.DisplayedClients(InArgs._DisplayedClients)
+				.ForEachReplicatedObject(InArgs._ForEachReplicatedObject)
+			];
+	}
+
 	void SReplicationStatus::Construct(const FArguments& InArgs, FGlobalAuthorityCache& InAuthorityCache)
 	{
 		AuthorityCache = &InAuthorityCache;
diff --git a/Engine/Plugins/Developer/Concert/ConcertApp/MultiUserClient/Source/MultiUserClient/Private/Widgets/ActiveSession/Replication/Client/SReplicationStatus.h b/Engine/Plugins/Developer/Concert/ConcertApp/MultiUserClient/Source/MultiUserClient/Private/Widgets/ActiveSession/Replication/Client/SReplicationStatus.h
--- a/Engine/Plugins/Developer/Concert/ConcertApp/MultiUserClient/Source/MultiUserClient/Private/Widgets/ActiveSession/Replication/Client/SReplicationStatus.h
+++ b/Engine/Plugins/Developer/Concert/ConcertApp/MultiUserClient/Source/MultiUserClient/Private/Widgets/ActiveSession/Replication/Client/SReplicationStatus.h
@@ -6,11 +6,9 @@
 #include "Widgets/SCompoundWidget.h"
 
 class STextBlock;
+class SVerticalBox;
 
-namespace UE::ConcertSharedSlate
-{
-	class IReplicationStreamModel;
-}
+namespace UE::ConcertSharedSlate { class IReplicationStreamModel; }
 
 namespace UE::MultiUserClient
 {
@@ -38,6 +36,9 @@ namespace UE::MultiUserClient
 			SLATE_EVENT(FForEachReplicatedObject, ForEachReplicatedObject)
 		SLATE_END_ARGS()
 
+		/** Adds a separator and SReplicationStatus to the bottom of the VerticalBox- */
+		static void AppendReplicationStatus(SVerticalBox& VerticalBox, FGlobalAuthorityCache& InAuthorityCache, const FArguments& InArgs);
+
 		void Construct(const FArguments& InArgs, FGlobalAuthorityCache& InAuthorityCache);
 		virtual ~SReplicationStatus() override;
 
diff --git a/Engine/Plugins/Developer/Concert/ConcertApp/MultiUserClient/Source/MultiUserClient/Private/Widgets/ActiveSession/Replication/Client/Single/SReplicationClientView.cpp b/Engine/Plugins/Developer/Concert/ConcertApp/MultiUserClient/Source/MultiUserClient/Private/Widgets/ActiveSession/Replication/Client/Single/SReplicationClientView.cpp
--- a/Engine/Plugins/Developer/Concert/ConcertApp/MultiUserClient/Source/MultiUserClient/Private/Widgets/ActiveSession/Replication/Client/Single/SReplicationClientView.cpp
+++ b/Engine/Plugins/Developer/Concert/ConcertApp/MultiUserClient/Source/MultiUserClient/Private/Widgets/ActiveSession/Replication/Client/Single/SReplicationClientView.cpp
@@ -10,11 +10,12 @@
 #include "Replication/ReplicationWidgetFactories.h"
 #include "Replication/Client/ReplicationClientManager.h"
 #include "Replication/Submission/ISubmissionWorkflow.h"
+#include "Widgets/ActiveSession/Replication/Client/FrequencyContextMenuUtils.h"
 #include "Widgets/ActiveSession/Replication/Client/Single/Columns/SingleClientColumns.h"
+#include "Widgets/ActiveSession/Replication/Client/SReplicationStatus.h"
 #include "Widgets/ActiveSession/Replication/Client/SClientToolbar.h"
 
 #include "Widgets/SBoxPanel.h"
-#include "Widgets/ActiveSession/Replication/Client/FrequencyContextMenuUtils.h"
 #include "Widgets/Layout/SBox.h"
 
 #define LOCTEXT_NAMESPACE "SReplicationClientView"
@@ -31,19 +32,18 @@ namespace UE::MultiUserClient
 		check(ReplicationClient);
 		
 		FGlobalAuthorityCache& AuthorityCache = ClientManager->GetAuthorityCache();
+		TSharedPtr<SVerticalBox> Content;
 		ChildSlot
 		[
-			SNew(SVerticalBox)
+			SAssignNew(Content, SVerticalBox)
 
 			// Toolbar
 			+SVerticalBox::Slot()
 			.AutoHeight()
 			.Padding(2.f)
 			[
-				SNew(SClientToolbar, AuthorityCache)
+				SNew(SClientToolbar)
 				.ViewSelectionArea() [ InArgs._ViewSelectionArea.Widget ]
-				.DisplayedClients(TSet{ ReplicationClient->GetEndpointId() })
-				.ForEachReplicatedObject(this, &SReplicationClientView::EnumerateReplicatedObjects)
 			]
 
 			// Editor
@@ -53,6 +53,12 @@ namespace UE::MultiUserClient
 				CreateContent(*ReplicationClient)
 			]
 		];
+
+		SReplicationStatus::AppendReplicationStatus(*Content, InClientManager.GetAuthorityCache(),
+			SReplicationStatus::FArguments()
+			.DisplayedClients(TSet{ ReplicationClient->GetEndpointId() })
+			.ForEachReplicatedObject(this, &SReplicationClientView::EnumerateReplicatedObjects)
+			);
 		
 		// Refresh UI if streams change externally, e.g. a remote client changed what they sent
 		ReplicationClient->OnModelChanged().AddSP(this, &SReplicationClientView::OnModelChanged);
diff --git a/Engine/Plugins/Developer/Concert/ConcertUI/ConcertSharedSlate/Source/ConcertSharedSlate/Private/Replication/Editor/View/ObjectViewer/SReplicationStreamViewer.cpp b/Engine/Plugins/Developer/Concert/ConcertUI/ConcertSharedSlate/Source/ConcertSharedSlate/Private/Replication/Editor/View/ObjectViewer/SReplicationStreamViewer.cpp
--- a/Engine/Plugins/Developer/Concert/ConcertUI/ConcertSharedSlate/Source/ConcertSharedSlate/Private/Replication/Editor/View/ObjectViewer/SReplicationStreamViewer.cpp
+++ b/Engine/Plugins/Developer/Concert/ConcertUI/ConcertSharedSlate/Source/ConcertSharedSlate/Private/Replication/Editor/View/ObjectViewer/SReplicationStreamViewer.cpp
@@ -259,7 +259,7 @@ namespace UE::ConcertSharedSlate
 				.HeaderContent()
 				[
 					SNew(STextBlock)
-					.Text(LOCTEXT("ReplicatedProperties", "Subobjects & Properties"))
+					.Text(LOCTEXT("ReplicatedProperties", "Properties"))
 					.Font(FAppStyle::Get().GetFontStyle("DetailsView.CategoryFontStyle"))
 					.ShadowOffset(FVector2D(1.0f, 1.0f))
 				]
