MU replication: Add missing nullptr check when calling TryLoadClass from a class path that may be invalid.
Also adds logging to similar places to be able to get info next time the class is null (when it should not be).

diff --git a/Engine/Plugins/Developer/Concert/ConcertApp/MultiUserClient/Source/MultiUserClient/Private/Widgets/ActiveSession/Replication/Client/Multi/Columns/SAssignPropertyComboBox.cpp b/Engine/Plugins/Developer/Concert/ConcertApp/MultiUserClient/Source/MultiUserClient/Private/Widgets/ActiveSession/Replication/Client/Multi/Columns/SAssignPropertyComboBox.cpp
--- a/Engine/Plugins/Developer/Concert/ConcertApp/MultiUserClient/Source/MultiUserClient/Private/Widgets/ActiveSession/Replication/Client/Multi/Columns/SAssignPropertyComboBox.cpp
+++ b/Engine/Plugins/Developer/Concert/ConcertApp/MultiUserClient/Source/MultiUserClient/Private/Widgets/ActiveSession/Replication/Client/Multi/Columns/SAssignPropertyComboBox.cpp
@@ -2,6 +2,7 @@
 
 #include "SAssignPropertyComboBox.h"
 
+#include "ConcertLogGlobal.h"
 #include "IConcertClient.h"
 #include "Replication/Client/ReplicationClient.h"
 #include "Replication/Client/ReplicationClientManager.h"
@@ -345,6 +346,7 @@ namespace UE::MultiUserClient
 			{
 				for (const FSoftObjectPath& ObjectPath : EditedObjects)
 				{
+					const FSoftClassPath ClassPath = EditModel->GetObjectClass(ObjectPath);
 					EditModel->RemoveProperties(ObjectPath, { Property });
 					
 					if (EditModel->HasAnyPropertyAssigned(ObjectPath))
@@ -356,8 +358,9 @@ namespace UE::MultiUserClient
 					// Skipping this check would close the entire property tree view and remove the actor hierarchy from the view.
 					// That would feel very unnatural / unexpected for the user. 
 					// If the user does not want the actor anymore, they should click it and delete it.
-					const UClass* ObjectClass = EditModel->GetObjectClass(ObjectPath).TryLoadClass<UObject>();
-					const bool bIsTopLevelObject = !ObjectClass->IsChildOf<AActor>();
+					const UClass* ObjectClass = ClassPath.IsValid() ? ClassPath.TryLoadClass<UObject>() : nullptr;
+					UE_CLOG(ClassPath.IsValid() && !ObjectClass, LogConcert, Warning, TEXT("SAssignPropertyComboBox: Failed to resolve class %s"), *ClassPath.ToString());
+					const bool bIsTopLevelObject = ObjectClass && !ObjectClass->IsChildOf<AActor>();
 					if (bIsTopLevelObject)
 					{
 						EditModel->RemoveObjects({ ObjectPath });
diff --git a/Engine/Plugins/Developer/Concert/ConcertUI/ConcertSharedSlate/Source/ConcertSharedSlate/Private/Replication/Editor/Model/GenericReplicationStreamModel.cpp b/Engine/Plugins/Developer/Concert/ConcertUI/ConcertSharedSlate/Source/ConcertSharedSlate/Private/Replication/Editor/Model/GenericReplicationStreamModel.cpp
--- a/Engine/Plugins/Developer/Concert/ConcertUI/ConcertSharedSlate/Source/ConcertSharedSlate/Private/Replication/Editor/Model/GenericReplicationStreamModel.cpp
+++ b/Engine/Plugins/Developer/Concert/ConcertUI/ConcertSharedSlate/Source/ConcertSharedSlate/Private/Replication/Editor/Model/GenericReplicationStreamModel.cpp
@@ -231,6 +231,7 @@ namespace UE::ConcertSharedSlate
 		UClass* Class = AssignedProperties->ClassPath.TryLoadClass<UObject>();
 		if (!Class)
 		{
+			UE_LOG(LogConcert, Warning, TEXT("FGenericReplicationStreamModel::AddProperties: Failed to resolve class %s"), *AssignedProperties->ClassPath.ToString());
 			return;
 		}
 
@@ -265,6 +266,7 @@ namespace UE::ConcertSharedSlate
 			return;
 		}
 
+		bool bLoggedWarning = false;
 		UClass* Class = AssignedProperties->ClassPath.TryLoadClass<UObject>();
 		int32 NumRemoved = 0;
 		for (const FConcertPropertyChain& RemovedProperty : Properties)
@@ -274,6 +276,8 @@ namespace UE::ConcertSharedSlate
 			// Removal should not fail if the class is not available
 			if (!Class)
 			{
+				UE_CLOG(!bLoggedWarning, LogConcert, Warning, TEXT("FGenericReplicationStreamModel::RemoveProperties: Failed to resolve class %s"), *AssignedProperties->ClassPath.ToString());
+				bLoggedWarning = true;
 				continue;
 			}
 
