Remove single client view

diff --git a/Engine/Plugins/Developer/Concert/ConcertApp/MultiUserClient/Source/MultiUserClient/Private/Widgets/ActiveSession/Replication/Client/Multi/SAllClientsView.cpp b/Engine/Plugins/Developer/Concert/ConcertApp/MultiUserClient/Source/MultiUserClient/Private/Widgets/ActiveSession/Replication/Client/Multi/SAllClientsView.cpp
--- a/Engine/Plugins/Developer/Concert/ConcertApp/MultiUserClient/Source/MultiUserClient/Private/Widgets/ActiveSession/Replication/Client/Multi/SAllClientsView.cpp
+++ b/Engine/Plugins/Developer/Concert/ConcertApp/MultiUserClient/Source/MultiUserClient/Private/Widgets/ActiveSession/Replication/Client/Multi/SAllClientsView.cpp
@@ -18,7 +18,6 @@ namespace UE::MultiUserClient
 		ChildSlot
 		[
 			SNew(SMultiClientView, InConcertClient, InClientManager, *AllClientsModel)
-			.ViewSelectionArea() [ InArgs._ViewSelectionArea.Widget ]
 		];
 	}
 
diff --git a/Engine/Plugins/Developer/Concert/ConcertApp/MultiUserClient/Source/MultiUserClient/Private/Widgets/ActiveSession/Replication/Client/Multi/SAllClientsView.h b/Engine/Plugins/Developer/Concert/ConcertApp/MultiUserClient/Source/MultiUserClient/Private/Widgets/ActiveSession/Replication/Client/Multi/SAllClientsView.h
--- a/Engine/Plugins/Developer/Concert/ConcertApp/MultiUserClient/Source/MultiUserClient/Private/Widgets/ActiveSession/Replication/Client/Multi/SAllClientsView.h
+++ b/Engine/Plugins/Developer/Concert/ConcertApp/MultiUserClient/Source/MultiUserClient/Private/Widgets/ActiveSession/Replication/Client/Multi/SAllClientsView.h
@@ -16,10 +16,7 @@ namespace UE::MultiUserClient
 	/** Leverages SMultiClientView to display all replication clients. */
 	class SAllClientsView : public SCompoundWidget
 	{
-		SLATE_BEGIN_ARGS(SAllClientsView)
-			{}
-			/** Dedicated space for a widget with which to change the view. */
-			SLATE_NAMED_SLOT(FArguments, ViewSelectionArea)
+		SLATE_BEGIN_ARGS(SAllClientsView){}
 		SLATE_END_ARGS()
 
 		void Construct(const FArguments& InArgs, TSharedRef<IConcertClient> InConcertClient, FReplicationClientManager& InClientManager);
diff --git a/Engine/Plugins/Developer/Concert/ConcertApp/MultiUserClient/Source/MultiUserClient/Private/Widgets/ActiveSession/Replication/Client/Multi/SMultiClientView.cpp b/Engine/Plugins/Developer/Concert/ConcertApp/MultiUserClient/Source/MultiUserClient/Private/Widgets/ActiveSession/Replication/Client/Multi/SMultiClientView.cpp
--- a/Engine/Plugins/Developer/Concert/ConcertApp/MultiUserClient/Source/MultiUserClient/Private/Widgets/ActiveSession/Replication/Client/Multi/SMultiClientView.cpp
+++ b/Engine/Plugins/Developer/Concert/ConcertApp/MultiUserClient/Source/MultiUserClient/Private/Widgets/ActiveSession/Replication/Client/Multi/SMultiClientView.cpp
@@ -14,7 +14,6 @@
 #include "Replication/Editor/Model/ObjectSource/ActorSelectionSourceModel.h"
 #include "Widgets/ActiveSession/Replication/Client/FrequencyContextMenuUtils.h"
 #include "Widgets/ActiveSession/Replication/Client/Multi/Columns/MultiStreamColumns.h"
-#include "Widgets/ActiveSession/Replication/Client/SClientToolbar.h"
 #include "Widgets/ActiveSession/Replication/Client/SReplicationStatus.h"
 
 #include "Widgets/SBoxPanel.h"
@@ -37,15 +36,6 @@ namespace UE::MultiUserClient
 		[
 			SAssignNew(Content, SVerticalBox)
 
-			// Toolbar
-			+SVerticalBox::Slot()
-			.AutoHeight()
-			.Padding(2.f)
-			[
-				SNew(SClientToolbar)
-				.ViewSelectionArea() [ InArgs._ViewSelectionArea.Widget ]
-			]
-
 			// Editor
 			+SVerticalBox::Slot()
 			.FillHeight(1.f)
diff --git a/Engine/Plugins/Developer/Concert/ConcertApp/MultiUserClient/Source/MultiUserClient/Private/Widgets/ActiveSession/Replication/Client/Multi/SMultiClientView.h b/Engine/Plugins/Developer/Concert/ConcertApp/MultiUserClient/Source/MultiUserClient/Private/Widgets/ActiveSession/Replication/Client/Multi/SMultiClientView.h
--- a/Engine/Plugins/Developer/Concert/ConcertApp/MultiUserClient/Source/MultiUserClient/Private/Widgets/ActiveSession/Replication/Client/Multi/SMultiClientView.h
+++ b/Engine/Plugins/Developer/Concert/ConcertApp/MultiUserClient/Source/MultiUserClient/Private/Widgets/ActiveSession/Replication/Client/Multi/SMultiClientView.h
@@ -27,10 +27,7 @@ namespace UE::MultiUserClient
 	class SMultiClientView
 		: public SCompoundWidget
 	{
-		SLATE_BEGIN_ARGS(SMultiClientView)
-		{}
-			/** Dedicated space for a widget with which to change the view. */
-			SLATE_NAMED_SLOT(FArguments, ViewSelectionArea)
+		SLATE_BEGIN_ARGS(SMultiClientView){}
 		SLATE_END_ARGS()
 
 		void Construct(const FArguments& InArgs, TSharedRef<IConcertClient> InConcertClient, FReplicationClientManager& InClientManager, IClientSelectionModel& InDisplayClientsModel);
diff --git a/Engine/Plugins/Developer/Concert/ConcertApp/MultiUserClient/Source/MultiUserClient/Private/Widgets/ActiveSession/Replication/Client/SClientToolbar.cpp b/Engine/Plugins/Developer/Concert/ConcertApp/MultiUserClient/Source/MultiUserClient/Private/Widgets/ActiveSession/Replication/Client/SClientToolbar.cpp
deleted file mode 100644
diff --git a/Engine/Plugins/Developer/Concert/ConcertApp/MultiUserClient/Source/MultiUserClient/Private/Widgets/ActiveSession/Replication/Client/SClientToolbar.h b/Engine/Plugins/Developer/Concert/ConcertApp/MultiUserClient/Source/MultiUserClient/Private/Widgets/ActiveSession/Replication/Client/SClientToolbar.h
deleted file mode 100644
diff --git a/Engine/Plugins/Developer/Concert/ConcertApp/MultiUserClient/Source/MultiUserClient/Private/Widgets/ActiveSession/Replication/Client/Single/Columns/SOwnerClientList.cpp b/Engine/Plugins/Developer/Concert/ConcertApp/MultiUserClient/Source/MultiUserClient/Private/Widgets/ActiveSession/Replication/Client/Single/Columns/SOwnerClientList.cpp
deleted file mode 100644
diff --git a/Engine/Plugins/Developer/Concert/ConcertApp/MultiUserClient/Source/MultiUserClient/Private/Widgets/ActiveSession/Replication/Client/Single/Columns/SOwnerClientList.h b/Engine/Plugins/Developer/Concert/ConcertApp/MultiUserClient/Source/MultiUserClient/Private/Widgets/ActiveSession/Replication/Client/Single/Columns/SOwnerClientList.h
deleted file mode 100644
diff --git a/Engine/Plugins/Developer/Concert/ConcertApp/MultiUserClient/Source/MultiUserClient/Private/Widgets/ActiveSession/Replication/Client/Single/Columns/SingleClientColumns.h b/Engine/Plugins/Developer/Concert/ConcertApp/MultiUserClient/Source/MultiUserClient/Private/Widgets/ActiveSession/Replication/Client/Single/Columns/SingleClientColumns.h
deleted file mode 100644
diff --git a/Engine/Plugins/Developer/Concert/ConcertApp/MultiUserClient/Source/MultiUserClient/Private/Widgets/ActiveSession/Replication/Client/Single/Columns/SingleClientColumns_ConflictWarning.cpp b/Engine/Plugins/Developer/Concert/ConcertApp/MultiUserClient/Source/MultiUserClient/Private/Widgets/ActiveSession/Replication/Client/Single/Columns/SingleClientColumns_ConflictWarning.cpp
deleted file mode 100644
diff --git a/Engine/Plugins/Developer/Concert/ConcertApp/MultiUserClient/Source/MultiUserClient/Private/Widgets/ActiveSession/Replication/Client/Single/Columns/SingleClientColumns_Owner.cpp b/Engine/Plugins/Developer/Concert/ConcertApp/MultiUserClient/Source/MultiUserClient/Private/Widgets/ActiveSession/Replication/Client/Single/Columns/SingleClientColumns_Owner.cpp
deleted file mode 100644
diff --git a/Engine/Plugins/Developer/Concert/ConcertApp/MultiUserClient/Source/MultiUserClient/Private/Widgets/ActiveSession/Replication/Client/Single/Columns/SingleClientColumns_ToggleAuthority.cpp b/Engine/Plugins/Developer/Concert/ConcertApp/MultiUserClient/Source/MultiUserClient/Private/Widgets/ActiveSession/Replication/Client/Single/Columns/SingleClientColumns_ToggleAuthority.cpp
deleted file mode 100644
diff --git a/Engine/Plugins/Developer/Concert/ConcertApp/MultiUserClient/Source/MultiUserClient/Private/Widgets/ActiveSession/Replication/Client/Single/SReplicationClientView.cpp b/Engine/Plugins/Developer/Concert/ConcertApp/MultiUserClient/Source/MultiUserClient/Private/Widgets/ActiveSession/Replication/Client/Single/SReplicationClientView.cpp
deleted file mode 100644
diff --git a/Engine/Plugins/Developer/Concert/ConcertApp/MultiUserClient/Source/MultiUserClient/Private/Widgets/ActiveSession/Replication/Client/Single/SReplicationClientView.h b/Engine/Plugins/Developer/Concert/ConcertApp/MultiUserClient/Source/MultiUserClient/Private/Widgets/ActiveSession/Replication/Client/Single/SReplicationClientView.h
deleted file mode 100644
diff --git a/Engine/Plugins/Developer/Concert/ConcertApp/MultiUserClient/Source/MultiUserClient/Private/Widgets/ActiveSession/Replication/Joined/SReplicationJoinedView.cpp b/Engine/Plugins/Developer/Concert/ConcertApp/MultiUserClient/Source/MultiUserClient/Private/Widgets/ActiveSession/Replication/Joined/SReplicationJoinedView.cpp
--- a/Engine/Plugins/Developer/Concert/ConcertApp/MultiUserClient/Source/MultiUserClient/Private/Widgets/ActiveSession/Replication/Joined/SReplicationJoinedView.cpp
+++ b/Engine/Plugins/Developer/Concert/ConcertApp/MultiUserClient/Source/MultiUserClient/Private/Widgets/ActiveSession/Replication/Joined/SReplicationJoinedView.cpp
@@ -4,234 +4,19 @@
 
 #include "IConcertSyncClient.h"
 #include "Replication/MultiUserReplicationManager.h"
-#include "SSelectClientViewComboButton.h"
 #include "Widgets/ActiveSession/Replication/Client/Multi/SAllClientsView.h"
-#include "Widgets/ActiveSession/Replication/Client/Single/SReplicationClientView.h"
-
-#include "Replication/Client/RemoteReplicationClient.h"
-#include "Widgets/Layout/SBox.h"
-#include "Widgets/Layout/SWidgetSwitcher.h"
-#include "Widgets/SBoxPanel.h"
-#include "Widgets/Text/STextBlock.h"
-
-#define LOCTEXT_NAMESPACE "SReplicationJoinedWidget"
 
 namespace UE::MultiUserClient
 {
 	void SReplicationJoinedView::Construct(
 		const FArguments& InArgs,
-		TSharedRef<FMultiUserReplicationManager> InReplicationManager,
-		TSharedRef<IConcertSyncClient> InClient
+		const TSharedRef<FMultiUserReplicationManager>& InReplicationManager,
+		const TSharedRef<IConcertSyncClient>& InClient
 		)
 	{
-		Client = MoveTemp(InClient);
-		ReplicationManager = MoveTemp(InReplicationManager);
-
 		ChildSlot
 		[
-			SAssignNew(ClientViewSwitcher, SWidgetSwitcher)
-			// These must be aligned with EClientViewType
-			+SWidgetSwitcher::Slot()
-			[
-				SNew(SAllClientsView, InClient->GetConcertClient(), *ReplicationManager->GetClientManager())
-				.ViewSelectionArea()
-				[
-					MakeClientSelectionArea()
-				]
-			]
-			+SWidgetSwitcher::Slot()
-			[
-				SNew(SReplicationClientView, InClient->GetConcertClient(), *ReplicationManager->GetClientManager())
-				.GetReplicationClient_Lambda([this](){ return &ReplicationManager->GetClientManager()->GetLocalClient(); })
-				.ViewSelectionArea()
-				[
-					MakeClientSelectionArea()
-				]
-			]
+			SNew(SAllClientsView, InClient->GetConcertClient(), *InReplicationManager->GetClientManager())
 		];
-		
-		RefreshClientViewSwitcher();
-		ReplicationManager->GetClientManager()->OnRemoteClientsChanged().AddSP(this, &SReplicationJoinedView::RefreshClientViewSwitcher);
 	}
-
-	void SReplicationJoinedView::RefreshClientViewSwitcher()
-	{
-		const int32 ActiveViewIndex = ClientViewSwitcher->GetActiveWidgetIndex();
-		const FGuid DisplayedClientId = ActiveViewIndex == 0
-			? Client->GetConcertClient()->GetCurrentSession()->GetSessionClientEndpointId()
-			: [this, &ActiveViewIndex]() 
-			{
-				for (const TPair<FGuid, int32>& Pair : RemoteClientToWidgetSwitcherIndex)
-				{
-					if (Pair.Value == ActiveViewIndex)
-					{
-						return Pair.Key;
-					}
-				}
-				return FGuid{};
-			}();
-		
-		// Retain the old SWidgets so their transient state, like highlights, are retained when the remote clients change.
-		TArray<TSharedRef<SWidget>> OldClientWidgets;
-		OldClientWidgets.Reserve(ClientViewSwitcher->GetNumWidgets());
-		while (ClientViewSwitcher->GetNumWidgets() > 0)
-		{
-			const TSharedRef<SWidget> Widget = ClientViewSwitcher->GetWidget(0).ToSharedRef();
-			OldClientWidgets.Add(Widget);
-			ClientViewSwitcher->RemoveSlot(Widget);
-		}
-		checkf(OldClientWidgets.Num() >= 2, TEXT("Was supposed to contain all clients and local client widgets"));
-
-		// All clients
-		ClientViewSwitcher->AddSlot()
-			[
-				OldClientWidgets[0]
-			];
-		// Local client
-		ClientViewSwitcher->AddSlot()
-			[
-				OldClientWidgets[1]
-			];
-		RebuildClientViewSwitcherChildren(OldClientWidgets);
-
-		// Need to change view if the displayed client was removed
-		const bool bDisplayedClientWasRemoved = RemoteClientToWidgetSwitcherIndex.Contains(DisplayedClientId); 
-		if (bDisplayedClientWasRemoved)
-		{
-			// Will show all clients
-			ClientViewSwitcher->SetActiveWidgetIndex(0);
-		}
-		
-		// ~OldClientWidgets now disposes of the widgets for which no client exists anymore
-	}
-
-	void SReplicationJoinedView::RebuildClientViewSwitcherChildren(const TArray<TSharedRef<SWidget>> OldClientWidgets)
-	{
-		TMap<FGuid, int32> OldRemoteClientToWidgetSwitcherIndex = MoveTemp(RemoteClientToWidgetSwitcherIndex);
-        for (const TNonNullPtr<FRemoteReplicationClient>& RemoteClient : ReplicationManager->GetClientManager()->GetRemoteClients())
-        {
-        	const FGuid& EndpointId = RemoteClient->GetEndpointId();
-        	
-        	const int32* ExistingClientWidgetIndex = OldRemoteClientToWidgetSwitcherIndex.Find(EndpointId);
-        	if (ExistingClientWidgetIndex)
-        	{
-        		ClientViewSwitcher->AddSlot()
-        		[
-        			OldClientWidgets[*ExistingClientWidgetIndex]
-        		];
-        	}
-        	else
-        	{
-        		ClientViewSwitcher->AddSlot()
-        		[
-        			SNew(SReplicationClientView, Client->GetConcertClient(), *ReplicationManager->GetClientManager())
-        			.GetReplicationClient_Lambda([this, EndpointId = RemoteClient->GetEndpointId()]()
-        			{
-        				// It is unsafe to simply capture RemoteClient because the containing TArray may reallocate its location
-        				return ReplicationManager->GetClientManager()->FindRemoteClient(EndpointId);
-        			})
-        			.ViewSelectionArea()
-        			[
-        				MakeClientSelectionArea()
-        			]
-        		];
-        	}
-
-        	const int32 NewIndex = ClientViewSwitcher->GetNumWidgets() - 1;
-        	RemoteClientToWidgetSwitcherIndex.Add(EndpointId, NewIndex);
-        }
-	}
-
-	TSharedRef<SWidget> SReplicationJoinedView::MakeClientSelectionArea()
-	{
-		return SNew(SHorizontalBox)
-			.ToolTipText(LOCTEXT("ClientView.ToolTip", "Select the client(s) of which you want to see the registered objects & properties."))
-
-			+SHorizontalBox::Slot()
-			.AutoWidth()
-			.VAlign(VAlign_Center)
-			[
-				SNew(STextBlock)
-				.Text(LOCTEXT("ClientView.Label", "Client View"))
-			]
-
-			+SHorizontalBox::Slot()
-			.AutoWidth()
-			.Padding(5.f, 0.f)
-			[
-				SNew(SBox)
-				.MinDesiredWidth(200.f)
-				[
-					MakeClientSelectionComboBox()
-				]
-			];
-	}
-
-	TSharedRef<SWidget> SReplicationJoinedView::MakeClientSelectionComboBox()
-	{
-		return SNew(SSelectClientViewComboButton)
-		
-			// Returning client info
-			.Client(Client->GetConcertClient())
-			.SelectableClients_Lambda([this]()
-			{
-				const TArray<TNonNullPtr<FRemoteReplicationClient>> RemoteClients = ReplicationManager->GetClientManager()->GetRemoteClients();
-				TArray<FGuid> Result;
-				Algo::Transform(RemoteClients, Result, [](const TNonNullPtr<FRemoteReplicationClient>& InClient){ return InClient->GetEndpointId(); });
-
-				const TSharedPtr<IConcertClientSession> CurrentSession = Client->GetConcertClient()->GetCurrentSession();
-				Result.Sort([&CurrentSession](const FGuid& Left, const FGuid& Right)
-				{
-					FConcertSessionClientInfo LeftInfo;
-					FConcertSessionClientInfo RightInfo;
-					CurrentSession->FindSessionClient(Left, LeftInfo);
-					CurrentSession->FindSessionClient(Right, RightInfo);
-					return LeftInfo.ClientInfo.DisplayName <= RightInfo.ClientInfo.DisplayName;
-				});
-				Result.Insert(CurrentSession->GetSessionClientEndpointId(), 0);
-				return Result;
-			})
-			.CurrentSelection_Lambda([this]()
-			{
-				const TOptional<FGuid> RemoteClient = GetRemoteClientBySwitcherIndex(ClientViewSwitcher->GetActiveWidgetIndex());
-				return RemoteClient
-					? *RemoteClient
-					: Client->GetConcertClient()->GetCurrentSession()->GetSessionClientEndpointId();
-			})
-			.CurrentDisplayMode_Lambda([this]()
-			{
-				switch(ClientViewSwitcher->GetActiveWidgetIndex())
-				{
-				case 0: return EClientViewType::AllClients;
-				case 1: return EClientViewType::LocalClient;
-				default: return EClientViewType::RemoteClient;
-				}
-			})
-			.OnSelectClient_Lambda([this](const FGuid& InClientId)
-			{
-				const int32* RemoteClientIndex = RemoteClientToWidgetSwitcherIndex.Find(InClientId);
-				// If it's not a remote client, it's the local client
-				const int32 ActiveWidgetIndex = RemoteClientIndex ? *RemoteClientIndex : 1;
-				ClientViewSwitcher->SetActiveWidgetIndex(ActiveWidgetIndex);
-			})
-			.OnSelectAllClients_Lambda([this]()
-			{
-				ClientViewSwitcher->SetActiveWidgetIndex(0);
-			})
-		;
-	}
-
-	TOptional<FGuid> SReplicationJoinedView::GetRemoteClientBySwitcherIndex(int32 WidgetSwitcherIndex) const
-	{
-		for (const TPair<FGuid, int32>& Pair : RemoteClientToWidgetSwitcherIndex)
-		{
-			if (Pair.Value == WidgetSwitcherIndex)
-			{
-				return { Pair.Key };
-			}
-		}
-		return {};
-	}
-}
-
-#undef LOCTEXT_NAMESPACE
\ No newline at end of file
+}
\ No newline at end of file
diff --git a/Engine/Plugins/Developer/Concert/ConcertApp/MultiUserClient/Source/MultiUserClient/Private/Widgets/ActiveSession/Replication/Joined/SReplicationJoinedView.h b/Engine/Plugins/Developer/Concert/ConcertApp/MultiUserClient/Source/MultiUserClient/Private/Widgets/ActiveSession/Replication/Joined/SReplicationJoinedView.h
--- a/Engine/Plugins/Developer/Concert/ConcertApp/MultiUserClient/Source/MultiUserClient/Private/Widgets/ActiveSession/Replication/Joined/SReplicationJoinedView.h
+++ b/Engine/Plugins/Developer/Concert/ConcertApp/MultiUserClient/Source/MultiUserClient/Private/Widgets/ActiveSession/Replication/Joined/SReplicationJoinedView.h
@@ -2,7 +2,6 @@
 
 #pragma once
 
-#include "Replication/IConcertClientReplicationManager.h"
 #include "Widgets/DeclarativeSyntaxSupport.h"
 #include "Widgets/SCompoundWidget.h"
 
@@ -10,10 +9,7 @@ class IConcertSyncClient;
 class SNotificationItem;
 class SWidgetSwitcher;
 
-namespace UE::ConcertSharedSlate
-{
-	class IReplicationStreamEditor;
-}
+namespace UE::ConcertSharedSlate { class IReplicationStreamEditor; }
 
 namespace UE::MultiUserClient
 {
@@ -31,35 +27,8 @@ namespace UE::MultiUserClient
 
 		void Construct(
 			const FArguments& InArgs,
-			TSharedRef<FMultiUserReplicationManager> InReplicationManager,
-			TSharedRef<IConcertSyncClient> InClient
+			const TSharedRef<FMultiUserReplicationManager>& InReplicationManager,
+			const TSharedRef<IConcertSyncClient>& InClient
 			);
-
-	private:
-
-		/** The local client this widget is created for. */
-		TSharedPtr<IConcertSyncClient> Client;
-		/** Acts as the model of this view */
-		TSharedPtr<FMultiUserReplicationManager> ReplicationManager;
-		
-		/** Selects which client is being view. */
-		TSharedPtr<SWidgetSwitcher> ClientViewSwitcher;
-		/** Maps a remote client to index in ClientViewSwitcher. */
-		TMap<FGuid, int32> RemoteClientToWidgetSwitcherIndex;
-		
-		/** Notification about in progress authority change, if any. */
-		TSharedPtr<SNotificationItem> AuthorityChangeNotification;
-		
-		// Building ClientViewSwitcher
-		/** Gets the remote clients and makes sure ClientViewSwitcher has a widget for each. */
-		void RefreshClientViewSwitcher();
-		/** Util for adding back old client widgets after ClientViewSwitcher was cleared. */
-		void RebuildClientViewSwitcherChildren(const TArray<TSharedRef<SWidget>> OldClientWidgets);
-
-		/** Warps the combo box with a text */
-		TSharedRef<SWidget> MakeClientSelectionArea();
-		/** Creates a combobox with which the content of ClientViewSwitcher can be changed. */
-		TSharedRef<SWidget> MakeClientSelectionComboBox();
-		TOptional<FGuid> GetRemoteClientBySwitcherIndex(int32 WidgetSwitcherIndex) const;
 	};
 }
\ No newline at end of file
diff --git a/Engine/Plugins/Developer/Concert/ConcertApp/MultiUserClient/Source/MultiUserReplicationEditor/Private/StreamEditor/ReplicationSessionPresetEditorToolkit.cpp b/Engine/Plugins/Developer/Concert/ConcertApp/MultiUserClient/Source/MultiUserReplicationEditor/Private/StreamEditor/ReplicationSessionPresetEditorToolkit.cpp
--- a/Engine/Plugins/Developer/Concert/ConcertApp/MultiUserClient/Source/MultiUserReplicationEditor/Private/StreamEditor/ReplicationSessionPresetEditorToolkit.cpp
+++ b/Engine/Plugins/Developer/Concert/ConcertApp/MultiUserClient/Source/MultiUserReplicationEditor/Private/StreamEditor/ReplicationSessionPresetEditorToolkit.cpp
@@ -75,27 +75,11 @@ namespace UE::MultiUserReplicationEditor
 	{
 		using namespace ConcertClientSharedSlate;
 		using namespace ConcertSharedSlate;
-		
-		const TSharedRef<IEditableReplicationStreamModel> AssetReadWriteModel = CreateTransactionalStreamModel(
-			CreateBaseStreamModel(GetEditedStreamAsset()->GetUnassignedClient()->Stream->MakeReplicationMapGetterAttribute()),
-			*GetEditedStreamAsset()
-			);
-		
-		const TSharedRef<FActorSelectionSourceModel> ObjectSourceModel = MakeShared<FActorSelectionSourceModel>();
-		const TSharedRef<FSelectPropertyFromUClassModel> PropertySourceModel = MakeShared<FSelectPropertyFromUClassModel>();
-		FDefaultStreamEditorParams DefaultParams{
-			.BaseEditorParams =
-			{
-				.DataModel = AssetReadWriteModel,
-				.ObjectSource = ObjectSourceModel,
-				.PropertySource = PropertySourceModel
-			}
-		};
-		const TSharedRef<IReplicationStreamEditor> EditorView = CreateDefaultStreamEditor(MoveTemp(DefaultParams));
 		return SNew(SDockTab)
 			.Label(LOCTEXT("BaseDetailsTitle", "Details"))
 			[
-				EditorView
+				SNew(STextBlock)
+				.Text(LOCTEXT("NotImplemented", "This feature is not implemented"))
 			];
 	}
 }
diff --git a/Engine/Plugins/Developer/Concert/ConcertUI/ConcertClientSharedSlate/Source/ConcertClientSharedSlate/Private/Replication/ClientReplicationWidgetFactories.cpp b/Engine/Plugins/Developer/Concert/ConcertUI/ConcertClientSharedSlate/Source/ConcertClientSharedSlate/Private/Replication/ClientReplicationWidgetFactories.cpp
--- a/Engine/Plugins/Developer/Concert/ConcertUI/ConcertClientSharedSlate/Source/ConcertClientSharedSlate/Private/Replication/ClientReplicationWidgetFactories.cpp
+++ b/Engine/Plugins/Developer/Concert/ConcertUI/ConcertClientSharedSlate/Source/ConcertClientSharedSlate/Private/Replication/ClientReplicationWidgetFactories.cpp
@@ -2,15 +2,11 @@
 
 #include "Replication/ClientReplicationWidgetFactories.h"
 
-#include "Editor/View/ClientEditorColumns.h"
 #include "Editor/View/PropertyTree/SFilteredPropertyTreeView.h"
 #include "Replication/ReplicationWidgetFactories.h"
-#include "Replication/Editor/Model/PropertyUtils.h"
 #include "Replication/Editor/Model/Object/EditorObjectHierarchyModel.h"
 #include "Replication/Editor/Model/Object/EditorObjectNameModel.h"
 #include "Replication/Editor/Model/ReplicationStreamObject.h"
-#include "Replication/Editor/View/Column/SelectionViewerColumns.h"
-#include "Replication/Editor/View/ObjectEditor/SDefaultReplicationStreamEditor.h"
 #include "Replication/Editor/Model/TransactionalReplicationStreamModel.h"
 
 #include "UObject/UObjectGlobals.h"
@@ -60,65 +56,4 @@ namespace UE::ConcertClientSharedSlate
 	{
 		return SNew(SFilteredPropertyTreeView, MoveTemp(Params));
 	}
-	
-	TSharedRef<ConcertSharedSlate::IReplicationStreamEditor> CreateDefaultStreamEditor(FDefaultStreamEditorParams Params)
-	{
-		using namespace ConcertSharedSlate;
-		using namespace ConcertSharedSlate::ReplicationColumns;
-
-		// When a user adds e.g. a struct property all of its child properties should be auto added as well.
-		FExtendProperties ExtendPropertiesDelegate = FExtendProperties::CreateLambda(
-		[DataModel = Params.BaseEditorParams.DataModel, AppendPropertyDelegate = MoveTemp(Params.OnExtendAddedProperties)]
-			(const FSoftObjectPath& Object, TArray<FConcertPropertyChain>& InOutPropertiesToAdd)
-			{
-				const FSoftClassPath ClassPath = DataModel->GetObjectClass(Object);
-				PropertyUtils::AppendAdditionalPropertiesToAdd(ClassPath, InOutPropertiesToAdd);
-				AppendPropertyDelegate.ExecuteIfBound(Object, InOutPropertiesToAdd);
-			});
-		
-		// This is a hack.
-		// The architecturally correct way to fix is pass FReplicationPropertyColumn the FSoftObjectPath to the object for which the column is being constructed.
-		struct FEditorIndirection
-		{
-			TSharedPtr<IReplicationStreamEditor> Editor;
-		};
-		TSharedRef<FEditorIndirection> Indirection = MakeShared<FEditorIndirection>();
-		
-		const FPropertyColumnEntry ReplicatesColumn = ReplicationColumns::Property::ReplicatesColumns(
-			TAttribute<IReplicationStreamViewer*>::CreateLambda([Indirection](){ return Indirection->Editor.Get(); }),
-			Params.BaseEditorParams.DataModel,
-			MoveTemp(ExtendPropertiesDelegate),
-			TCheckboxColumnDelegates<FPropertyTreeRowContext>::FIsEnabled::CreateLambda([IsEnabled = Params.BaseEditorParams.IsEditingEnabled](const FPropertyTreeRowContext&)
-			{
-				return !IsEnabled.IsBound() || IsEnabled.Get();
-			}),
-			Params.BaseEditorParams.EditingDisabledToolTipText
-			);
-		
-		TArray<FPropertyColumnEntry>& PropertyColumns = Params.PropertyColumns;
-		const bool bHasType = PropertyColumns.ContainsByPredicate([](const FPropertyColumnEntry& Entry)
-		{
-			return Entry.ColumnId == Property::TypeColumnId;
-		});
-		if (!bHasType)
-		{
-			PropertyColumns.Add(Property::TypeColumn());
-		}
-		PropertyColumns.Add(ReplicatesColumn);
-
-		FCreateViewerParams ViewerParams
-		{
-			.PropertyTreeView = CreateFilterablePropertyTreeView({ .AdditionalPropertyColumns = MoveTemp(PropertyColumns) }),
-			.ObjectHierarchy = MoveTemp(Params.ObjectHierarchy),
-			.NameModel = MoveTemp(Params.NameModel),
-			.OnExtendObjectsContextMenu = MoveTemp(Params.OnExtendObjectsContextMenu),
-			.ObjectColumns = MoveTemp(Params.ObjectColumns),
-			.PrimaryObjectSort = FColumnSortInfo{ TopLevel::LabelColumnId, EColumnSortMode::Ascending },
-			.SecondaryObjectSort = FColumnSortInfo{ TopLevel::LabelColumnId, EColumnSortMode::Ascending },
-		};
-		
-		TSharedRef<SDefaultReplicationStreamEditor> Editor = SNew(SDefaultReplicationStreamEditor, MoveTemp(Params.BaseEditorParams), MoveTemp(ViewerParams));
-		Indirection->Editor = Editor;
-		return Editor;
-	}
 }
\ No newline at end of file
diff --git a/Engine/Plugins/Developer/Concert/ConcertUI/ConcertClientSharedSlate/Source/ConcertClientSharedSlate/Private/Replication/Editor/View/ClientEditorColumns.cpp b/Engine/Plugins/Developer/Concert/ConcertUI/ConcertClientSharedSlate/Source/ConcertClientSharedSlate/Private/Replication/Editor/View/ClientEditorColumns.cpp
deleted file mode 100644
diff --git a/Engine/Plugins/Developer/Concert/ConcertUI/ConcertClientSharedSlate/Source/ConcertClientSharedSlate/Private/Replication/Editor/View/ClientEditorColumns.h b/Engine/Plugins/Developer/Concert/ConcertUI/ConcertClientSharedSlate/Source/ConcertClientSharedSlate/Private/Replication/Editor/View/ClientEditorColumns.h
deleted file mode 100644
diff --git a/Engine/Plugins/Developer/Concert/ConcertUI/ConcertClientSharedSlate/Source/ConcertClientSharedSlate/Private/Replication/Editor/View/ObjectEditor/SDefaultReplicationStreamEditor.cpp b/Engine/Plugins/Developer/Concert/ConcertUI/ConcertClientSharedSlate/Source/ConcertClientSharedSlate/Private/Replication/Editor/View/ObjectEditor/SDefaultReplicationStreamEditor.cpp
deleted file mode 100644
diff --git a/Engine/Plugins/Developer/Concert/ConcertUI/ConcertClientSharedSlate/Source/ConcertClientSharedSlate/Private/Replication/Editor/View/ObjectEditor/SDefaultReplicationStreamEditor.h b/Engine/Plugins/Developer/Concert/ConcertUI/ConcertClientSharedSlate/Source/ConcertClientSharedSlate/Private/Replication/Editor/View/ObjectEditor/SDefaultReplicationStreamEditor.h
deleted file mode 100644
diff --git a/Engine/Plugins/Developer/Concert/ConcertUI/ConcertClientSharedSlate/Source/ConcertClientSharedSlate/Public/Replication/ClientReplicationWidgetDelegates.h b/Engine/Plugins/Developer/Concert/ConcertUI/ConcertClientSharedSlate/Source/ConcertClientSharedSlate/Public/Replication/ClientReplicationWidgetDelegates.h
deleted file mode 100644
diff --git a/Engine/Plugins/Developer/Concert/ConcertUI/ConcertClientSharedSlate/Source/ConcertClientSharedSlate/Public/Replication/ClientReplicationWidgetFactories.h b/Engine/Plugins/Developer/Concert/ConcertUI/ConcertClientSharedSlate/Source/ConcertClientSharedSlate/Public/Replication/ClientReplicationWidgetFactories.h
--- a/Engine/Plugins/Developer/Concert/ConcertUI/ConcertClientSharedSlate/Source/ConcertClientSharedSlate/Public/Replication/ClientReplicationWidgetFactories.h
+++ b/Engine/Plugins/Developer/Concert/ConcertUI/ConcertClientSharedSlate/Source/ConcertClientSharedSlate/Public/Replication/ClientReplicationWidgetFactories.h
@@ -2,7 +2,6 @@
 
 #pragma once
 
-#include "Replication/ClientReplicationWidgetDelegates.h"
 #include "Replication/ReplicationWidgetFactories.h"
 #include "Templates/SharedPointer.h"
 
@@ -68,56 +67,4 @@ namespace UE::ConcertClientSharedSlate
 	 * The user can toggle used filters on and off under the search bar.
 	 */
 	CONCERTCLIENTSHAREDSLATE_API TSharedRef<ConcertSharedSlate::IPropertyTreeView> CreateFilterablePropertyTreeView(FFilterablePropertyTreeViewParams Params);
-
-	/** Params for creating the default single stream editor */
-	struct FDefaultStreamEditorParams
-	{
-		ConcertSharedSlate::FCreateEditorParams BaseEditorParams;
-
-		/**
-		 * Additional property columns to add to the default ones.
-		 * By default the default editor has the label, type, and replicates (checkbox) columns.
-		 */
-		TArray<ConcertSharedSlate::FPropertyColumnEntry> PropertyColumns
-		{
-			ConcertSharedSlate::ReplicationColumns::Property::LabelColumn(),
-			ConcertSharedSlate::ReplicationColumns::Property::TypeColumn()
-		};
-		
-		/**
-		 * Optional. Determines the objects displayed as children to the top-level objects in the top section.
-		 * If left unspecified, the object outliner will only display actors.
-		 * @note The created view will keep a strong reference to this.
-		 */
-		TSharedPtr<ConcertSharedSlate::IObjectHierarchyModel> ObjectHierarchy;
-
-		/**
-		 * Optional. Determines the display name of objects.
-		 * If left unspecified, defaults to the names in the object paths.
-		 */
-		TSharedPtr<ConcertSharedSlate::IObjectNameModel> NameModel;
-
-		/** Optional. Called to generate the context menu for objects. This extends the options already generated by this widget and is called at the end. */
-		ConcertSharedSlate::FExtendObjectMenu OnExtendObjectsContextMenu;
-
-		/**
-		 * Optional. When a property is assigned to the client, this is called to add any additional properties.
-		 * By default, if the user adds a struct property, its child properties are also added.
-		 */
-		FExtendProperties OnExtendAddedProperties;
-		
-		/** Optional. Additional columns to add to the object view */
-		TArray<ConcertSharedSlate::FObjectColumnEntry> ObjectColumns;
-	};
-	
-	/**
-	 * Creates a default IReplicationStreamEditor which is used to edit a single stream, e.g. a single client's stream.
-	 *
-	 * This editor adds a checkbox to the start of every property row.
-	 * - Checking adds the property to the selected objects' property mappings.
-	 * - Unchecking removes the property to the selected objects' property mappings
-	 * 
-	 * @see CreateViewer for a description of the UI layout.
-	 */
-	CONCERTCLIENTSHAREDSLATE_API TSharedRef<ConcertSharedSlate::IReplicationStreamEditor> CreateDefaultStreamEditor(FDefaultStreamEditorParams Params);
 }
