The reassignment combo box has been replaced by a flat list displaying assigned clients. Reassignment can be done via the context menu.

diff --git a/Engine/Plugins/Developer/Concert/ConcertApp/MultiUserClient/Source/MultiUserClient/Private/Widgets/ActiveSession/Replication/Client/Context/ContextMenuUtils.h b/Engine/Plugins/Developer/Concert/ConcertApp/MultiUserClient/Source/MultiUserClient/Private/Widgets/ActiveSession/Replication/Client/Context/ContextMenuUtils.h
new file mode 100644
--- /dev/null
+++ b/Engine/Plugins/Developer/Concert/ConcertApp/MultiUserClient/Source/MultiUserClient/Private/Widgets/ActiveSession/Replication/Client/Context/ContextMenuUtils.h
@@ -0,0 +1,45 @@
+﻿// Copyright Epic Games, Inc. All Rights Reserved.
+
+#pragma once
+
+#include "Containers/ArrayView.h"
+#include "UObject/SoftObjectPath.h"
+
+class IConcertClient;
+class FMenuBuilder;
+
+namespace UE::ConcertSharedSlate
+{
+	class IMultiReplicationStreamEditor;
+	class IObjectHierarchyModel;
+}
+namespace UE::MultiUserClient
+{
+	class FReassignObjectPropertiesLogic;
+	class FReplicationClientManager;
+}
+
+namespace UE::MultiUserClient::ContextMenuUtils
+{
+	/** Adds menu entries for reassigning the object to another client. */
+	void AddReassignmentOptions(
+		FMenuBuilder& MenuBuilder,
+		const FSoftObjectPath& ContextObject,
+		const IConcertClient& ConcertClient,
+		const FReplicationClientManager& ReplicationManager,
+		ConcertSharedSlate::IObjectHierarchyModel& ObjectHierarchy,
+		FReassignObjectPropertiesLogic& ReassignmentLogic,
+		ConcertSharedSlate::IMultiReplicationStreamEditor& MultiStreamEditor
+		);
+	
+	/** Adds an edit box for batch reassigning the select object's frequencies for all replicating clients. */
+	void AddFrequencyOptionsForMultipleClients(FMenuBuilder& MenuBuilder, const FSoftObjectPath& ContextObjects, FReplicationClientManager& InClientManager);
+	/** Adds an edit box for batch reassigning the select object's frequencies for all replicating clients. */
+	inline void AddFrequencyOptionsIfOneContextObject_MultiClient(FMenuBuilder& MenuBuilder, TConstArrayView<FSoftObjectPath> ContextObjects, FReplicationClientManager& InClientManager)
+	{
+		if (ContextObjects.Num() == 1)
+		{
+			AddFrequencyOptionsForMultipleClients(MenuBuilder, ContextObjects[0], InClientManager);
+		}
+	}
+}
\ No newline at end of file
diff --git a/Engine/Plugins/Developer/Concert/ConcertApp/MultiUserClient/Source/MultiUserClient/Private/Widgets/ActiveSession/Replication/Client/Context/ContextMenuUtils_Frequency.cpp b/Engine/Plugins/Developer/Concert/ConcertApp/MultiUserClient/Source/MultiUserClient/Private/Widgets/ActiveSession/Replication/Client/Context/ContextMenuUtils_Frequency.cpp
new file mode 100644
--- /dev/null
+++ b/Engine/Plugins/Developer/Concert/ConcertApp/MultiUserClient/Source/MultiUserClient/Private/Widgets/ActiveSession/Replication/Client/Context/ContextMenuUtils_Frequency.cpp
@@ -0,0 +1,301 @@
+// Copyright Epic Games, Inc. All Rights Reserved.
+
+#include "ContextMenuUtils.h"
+
+#include "Replication/Client/ReplicationClientManager.h"
+#include "Replication/Util/FrequencyUtils.h"
+
+#include "Algo/AnyOf.h"
+#include "Async/Async.h"
+#include "Async/TaskGraphInterfaces.h"
+#include "Framework/Commands/UIAction.h"
+#include "Framework/MultiBox/MultiBoxBuilder.h"
+#include "Textures/SlateIcon.h"
+#include "Widgets/Images/SThrobber.h"
+#include "Widgets/Input/SNumericEntryBox.h"
+#include "Widgets/Layout/SScaleBox.h"
+
+#define LOCTEXT_NAMESPACE "FrequencyContextMenuUtils"
+
+namespace UE::MultiUserClient::ContextMenuUtils
+{
+	namespace Private
+	{
+		using FInlineClientArray = FrequencyUtils::FInlineClientArray;
+
+		static FText GetRealtimeTooltip() { return LOCTEXT("ListView_Realtime.ToolTip", "The object will replicate as fast as possible (every tick)."); }
+		static FText GetSpecifiedRateTooltip() { return LOCTEXT("ListView_SpecifiedRate.ToolTip", "The object replicates at the specified rate per second."); }
+		static FText GetEditBoxInstructionsTooltip() { return LOCTEXT("Frequency.ToolTip", "Specifies how often this object should replicate in 1 second."); }
+
+		/** Displayed in context menu for changing an object's replication frequency. */
+		class SFrequencyNumericBox : public SCompoundWidget
+		{
+		public:
+
+			SLATE_BEGIN_ARGS(SFrequencyNumericBox)
+			{}
+				SLATE_ARGUMENT(FSoftObjectPath, SelectedObject)
+				SLATE_ATTRIBUTE(FInlineClientArray, Clients)
+			SLATE_END_ARGS()
+			
+			void Construct(const FArguments& InArgs, FReplicationClientManager& InClientManager)
+			{
+				ContextObject = InArgs._SelectedObject;
+				ClientsAttribute = InArgs._Clients;
+				ClientManager = &InClientManager;
+				
+				ChildSlot
+				[
+					SNew(SBox)
+					.MinDesiredWidth(100.f)
+					[
+						SNew(SHorizontalBox)
+
+						+SHorizontalBox::Slot()
+						.FillWidth(1.f)
+						[
+							SNew(SNumericEntryBox<uint8>)
+							.AllowSpin(false)
+							.IsEnabled(this, &SFrequencyNumericBox::GetIsEnabled)
+							.Value(this, &SFrequencyNumericBox::GetValue)
+							.OnValueCommitted(this, &SFrequencyNumericBox::OnValueCommitted)
+							.ToolTipText(this, &SFrequencyNumericBox::GetToolTipText)
+							.UndeterminedString(LOCTEXT("Single.Undetermined", "n/a"))
+						]
+						
+						+SHorizontalBox::Slot()
+						.AutoWidth()
+						[
+							SNew(SScaleBox)
+							[
+								SNew(SThrobber)
+								.Visibility_Lambda([this](){ return InProgressChange.IsSet() ? EVisibility::HitTestInvisible : EVisibility::Collapsed; })
+							]
+						]
+					]
+				];
+			}
+
+		private:
+
+			FSoftObjectPath ContextObject;
+			TAttribute<FInlineClientArray> ClientsAttribute;
+			FReplicationClientManager* ClientManager = nullptr;
+
+			struct FChangeOperation
+			{
+				TSharedPtr<IParallelSubmissionOperation> ChangeOperation;
+				uint8 ValueBeingSet;
+			};
+			/** Set for as long a change operation is in progress. */
+			TOptional<FChangeOperation> InProgressChange;
+
+			/** @return Gets all clients that have the given mode set for ContextObject. */
+			FInlineClientArray GetClientsWithMode(const EConcertObjectReplicationMode Mode) const
+			{
+				FInlineClientArray Result = ClientsAttribute.Get();
+				for (auto It = Result.CreateIterator(); It; ++It)
+				{
+					const TOptional<EConcertObjectReplicationMode> ClientMode = FrequencyUtils::FindSharedReplicationMode(ContextObject, { *It }, *ClientManager);
+					if (!ClientMode || *ClientMode != Mode)
+					{
+						It.RemoveCurrent();
+					}
+				}
+				return Result;
+			}
+
+			/** @return Whether all clients with specified rate can be changed. */
+			bool GetIsEnabled() const
+			{
+				const FInlineClientArray Clients = GetClientsWithMode(EConcertObjectReplicationMode::SpecifiedRate);
+				return !InProgressChange && !Clients.IsEmpty() && FrequencyUtils::CanChangeFrequencySettings(ContextObject, Clients, *ClientManager);
+			}
+			
+			void OnValueCommitted(uint8 NewValue, ETextCommit::Type CommitType)
+			{
+				if (CommitType != ETextCommit::OnEnter || NewValue == 0 || InProgressChange.IsSet())
+				{
+					return;
+				}
+			
+				const TSharedPtr<IParallelSubmissionOperation> ChangeOperation = FrequencyUtils::SetFrequencySettingForClients(ContextObject, GetClientsWithMode(EConcertObjectReplicationMode::SpecifiedRate), *ClientManager,
+					[NewValue](FConcertObjectReplicationSettings& Override)
+					{
+						Override.ReplicationRate = NewValue;
+					});
+				InProgressChange = { ChangeOperation, NewValue };
+
+				// For better UX, while we're waiting on the server to progress the request this widget should show what the user entered last.
+				// Once the request completes, the widget will show what the server state (so it may revert if the request failed, e.g. due to time out).
+				InProgressChange->ChangeOperation->OnCompletedFuture_AnyThread()
+					.Next([WeakThis = TWeakPtr<SFrequencyNumericBox>(SharedThis(this))](FParallelExecutionResult&& ExecutionResult)
+					{
+						// User may have close the menu containing this widget before the request completed.
+						if (const TSharedPtr<SFrequencyNumericBox> ThisPin = WeakThis.Pin())
+						{
+							// GuardedExecuteOnGameThread internally checks that "this" is valid before executing lambda
+							ThisPin->GuardedExecuteOnGameThread([WeakThis]() { WeakThis.Pin()->InProgressChange.Reset(); });
+						}
+					});
+			}
+
+			void GuardedExecuteOnGameThread(TUniqueFunction<void()> Callback) const
+			{
+				if (IsInGameThread())
+				{
+					Callback();
+				}
+				else
+				{
+					AsyncTask(ENamedThreads::GameThread, [Callback = MoveTemp(Callback), WeakThis = AsWeak()]()
+					{
+						if (WeakThis.IsValid())
+						{
+							Callback();
+						}
+					});
+				}
+			}
+
+			TOptional<uint8> GetValue() const
+			{
+				// For better UX, if a change is in progress show whatever the user specified for the request. Once the request finishes, show whatever the state is on the server.
+				return InProgressChange
+					? InProgressChange->ValueBeingSet
+					: FrequencyUtils::FindSharedFrequencyRate(ContextObject, GetClientsWithMode(EConcertObjectReplicationMode::SpecifiedRate), *ClientManager);
+			}
+			
+			FText GetToolTipText() const
+			{
+				if (GetValue())
+				{
+					return GetEditBoxInstructionsTooltip();
+				}
+
+				const TOptional<EConcertObjectReplicationMode> Mode = FrequencyUtils::FindSharedReplicationMode(ContextObject, GetClientsWithMode(EConcertObjectReplicationMode::SpecifiedRate), *ClientManager);
+				return Mode
+					? LOCTEXT("NotApplicable.Mixed", "Multiple clients are replicating this object with different rates.")
+					: LOCTEXT("NotApplicable.AllRealtime", "This object is repliacting in realtime");
+			}
+		};
+
+		static void AddReplicationModeSubMenu(FMenuBuilder& MenuBuilder, const FSoftObjectPath& ContextObject, TAttribute<FInlineClientArray> Clients, FReplicationClientManager& InClientManager)
+		{
+			const auto SetReplicationMode = [ContextObject, Clients, &InClientManager](const EConcertObjectReplicationMode ModeToSet)
+			{
+				FrequencyUtils::SetFrequencySettingForClients(ContextObject, Clients.Get(), InClientManager,
+					[ModeToSet](FConcertObjectReplicationSettings& SettingToOverride)
+					{
+						SettingToOverride.ReplicationMode = ModeToSet;
+					});
+			};
+			
+			MenuBuilder.AddMenuEntry(
+				LOCTEXT("ListView_Realtime.Label", "Realtime"),
+				GetRealtimeTooltip(),
+				FSlateIcon(),
+				FUIAction(
+					FExecuteAction::CreateLambda([SetReplicationMode]()
+					{
+						SetReplicationMode(EConcertObjectReplicationMode::Realtime);
+					}),
+					FCanExecuteAction::CreateLambda([ContextObject, Clients, &InClientManager]()
+					{
+						return FrequencyUtils::CanChangeFrequencySettings(ContextObject, Clients.Get(), InClientManager);
+					}),
+					FIsActionChecked::CreateLambda([ContextObject, Clients, &InClientManager]()
+					{
+						return FrequencyUtils::AllClientsHaveMode(EConcertObjectReplicationMode::Realtime, ContextObject, Clients.Get(), InClientManager);
+					})
+				),
+				NAME_None,
+				EUserInterfaceActionType::RadioButton
+			);
+			
+			MenuBuilder.AddMenuEntry(
+				LOCTEXT("ListView_SpecifiedRate.Label", "Specified rate"),
+				GetSpecifiedRateTooltip(),
+				FSlateIcon(),
+				FUIAction(
+					FExecuteAction::CreateLambda([SetReplicationMode]()
+					{
+						SetReplicationMode(EConcertObjectReplicationMode::SpecifiedRate);
+					}),
+					FCanExecuteAction::CreateLambda([ContextObject, Clients, &InClientManager]()
+					{
+						return FrequencyUtils::CanChangeFrequencySettings(ContextObject, Clients.Get(), InClientManager);
+					}),
+					FIsActionChecked::CreateLambda([ContextObject, Clients, &InClientManager]()
+					{
+						return FrequencyUtils::AllClientsHaveMode(EConcertObjectReplicationMode::SpecifiedRate, ContextObject, Clients.Get(), InClientManager);
+					})
+				),
+				NAME_None,
+				EUserInterfaceActionType::RadioButton
+				);
+		}
+
+		static void SharedAppendFrequencyToMenu(FMenuBuilder& MenuBuilder, const FSoftObjectPath& ContextObject, TAttribute<FInlineClientArray> GetClientsAttribute, FReplicationClientManager& InClientManager)
+		{
+			// No frequency section if the object is not registered by any clients
+			const bool bHasObjectRegistered = Algo::AnyOf(GetClientsAttribute.Get(), [&ContextObject, &InClientManager](const FGuid& ClientId)
+			{
+				const FReplicationClient* Client = InClientManager.FindClient(ClientId);
+				return !Client || Client->GetStreamSynchronizer().GetServerState().HasProperties(ContextObject);
+			});
+			if (!bHasObjectRegistered)
+			{
+				return;
+			}
+			
+			MenuBuilder.BeginSection(NAME_None, LOCTEXT("Frequency", "Frequency"));
+			
+			// Change to Realtime / Specified Rate
+			AddReplicationModeSubMenu(MenuBuilder, ContextObject, GetClientsAttribute, InClientManager);
+
+			// Replication rate, only enabled if Specified Rate
+			constexpr bool bNoIndent = false;
+			constexpr bool bSearchable = true;
+			MenuBuilder.AddWidget(
+				SNew(Private::SFrequencyNumericBox, InClientManager)
+					.SelectedObject(ContextObject)
+					.Clients(MoveTemp(GetClientsAttribute)),
+				LOCTEXT("Frequency.Label", "Replication Rate"),
+				bNoIndent,
+				bSearchable,
+				GetEditBoxInstructionsTooltip()
+				);
+
+			MenuBuilder.EndSection();
+		}
+	}
+	
+	void AddFrequencyOptionsForSingleClient(FMenuBuilder& MenuBuilder, const FSoftObjectPath& ContextObject, const FGuid& ClientId, FReplicationClientManager& InClientManager)
+	{
+		Private::SharedAppendFrequencyToMenu(MenuBuilder, ContextObject, TArray{ ClientId }, InClientManager);
+	}
+
+	void AddFrequencyOptionsForMultipleClients(FMenuBuilder& MenuBuilder, const FSoftObjectPath& ContextObject, FReplicationClientManager& InClientManager)
+	{
+		using namespace Private;
+		TAttribute<FInlineClientArray> GetClientsAttribute = TAttribute<FInlineClientArray>::CreateLambda([ContextObject, &InClientManager]()
+		{
+			FInlineClientArray Result;
+			InClientManager.ForEachClient([&ContextObject, &Result](const FReplicationClient& Client)
+			{
+				// Only clients that have the object registered should be considered
+				if (Client.GetStreamSynchronizer().GetServerState().HasProperties(ContextObject))
+				{
+					Result.Add(Client.GetEndpointId());
+				}
+				return EBreakBehavior::Continue;
+			});
+			return Result;
+		});
+		
+		SharedAppendFrequencyToMenu(MenuBuilder, ContextObject, MoveTemp(GetClientsAttribute), InClientManager);
+	}
+}
+
+#undef LOCTEXT_NAMESPACE
\ No newline at end of file
diff --git a/Engine/Plugins/Developer/Concert/ConcertApp/MultiUserClient/Source/MultiUserClient/Private/Widgets/ActiveSession/Replication/Client/Context/ContextMenuUtils_Reassignment.cpp b/Engine/Plugins/Developer/Concert/ConcertApp/MultiUserClient/Source/MultiUserClient/Private/Widgets/ActiveSession/Replication/Client/Context/ContextMenuUtils_Reassignment.cpp
new file mode 100644
--- /dev/null
+++ b/Engine/Plugins/Developer/Concert/ConcertApp/MultiUserClient/Source/MultiUserClient/Private/Widgets/ActiveSession/Replication/Client/Context/ContextMenuUtils_Reassignment.cpp
@@ -0,0 +1,100 @@
+﻿// Copyright Epic Games, Inc. All Rights Reserved.
+
+#include "ContextMenuUtils.h"
+
+#include "Replication/Client/ReplicationClient.h"
+#include "Replication/Editor/Model/Object/IObjectHierarchyModel.h"
+#include "Replication/Editor/View/IMultiReplicationStreamEditor.h"
+#include "Replication/Editor/View/IReplicationStreamEditor.h"
+#include "Replication/Submission/MultiEdit/ReassignObjectPropertiesLogic.h"
+#include "Widgets/ActiveSession/Replication/Client/ClientUtils.h"
+#include "Widgets/ActiveSession/Replication/Client/Multi/Columns/MultiStreamColumns.h"
+
+#include "Framework/Commands/UIAction.h"
+#include "Framework/MultiBox/MultiBoxBuilder.h"
+#include "Textures/SlateIcon.h"
+
+#define LOCTEXT_NAMESPACE "ReassignmentContextMenuUtils"
+
+namespace UE::MultiUserClient::ContextMenuUtils
+{
+	namespace Private
+	{
+		using FInlineAllocator = TInlineAllocator<24>;
+		using FInlineObjectPathArray = TArray<FSoftObjectPath, FInlineAllocator>;
+		
+		static FInlineObjectPathArray GetChildrenOfManagedObject(ConcertSharedSlate::IObjectHierarchyModel& ObjectHierarchy, FSoftObjectPath ManagedObject)
+		{
+			return ObjectHierarchy.GetChildrenRecursive<FInlineAllocator>(ManagedObject);
+		}
+		
+		static void AddReassignSection(
+			FMenuBuilder& MenuBuilder,
+			const TArray<const FReplicationClient*>& SortedClients,
+			TAttribute<FInlineObjectPathArray> ObjectsToAssign,
+			const IConcertClient& ConcertClient,
+			FReassignObjectPropertiesLogic& ReassignmentLogic,
+			ConcertSharedSlate::IMultiReplicationStreamEditor& MultiStreamEditor
+			)
+		{
+			for (const FReplicationClient* Client : SortedClients)
+			{
+				const FGuid& ClientId = Client->GetEndpointId();
+				MenuBuilder.AddMenuEntry(
+					FText::FromString(ClientUtils::GetClientDisplayName(ConcertClient, Client->GetEndpointId())),
+					TAttribute<FText>::CreateLambda([ClientId, ObjectsToAssign, &ReassignmentLogic]()
+					{
+						FText CannotEditReason;
+						const bool bCanReassign = ReassignmentLogic.CanReassignAnyTo(ObjectsToAssign.Get(), ClientId, &CannotEditReason);
+						return bCanReassign ? LOCTEXT("DoReassign", "Reassign to this client") : CannotEditReason;
+					}),
+					FSlateIcon(),
+					FUIAction(
+						FExecuteAction::CreateLambda([ObjectsToAssign, ClientId, &ReassignmentLogic, &MultiStreamEditor]()
+						{
+							TArray<FSoftObjectPath> ObjectsNotInline;
+							Algo::Transform(ObjectsToAssign.Get(), ObjectsNotInline, [](const FSoftObjectPath& ObjectPath){ return ObjectPath; });
+							ReassignmentLogic.ReassignAllTo(ObjectsNotInline, ClientId);
+							MultiStreamEditor.GetEditorBase().RequestObjectColumnResort(MultiStreamColumns::AssignedClientsColumnId);
+						}),
+						FCanExecuteAction::CreateLambda([ObjectsToAssign, ClientId, &ReassignmentLogic](){ return ReassignmentLogic.CanReassignAnyTo(ObjectsToAssign.Get(), ClientId); }),
+						FIsActionChecked::CreateLambda([ObjectsToAssign, ClientId, &ReassignmentLogic](){ return ReassignmentLogic.OwnsAnyOf(ObjectsToAssign.Get(), ClientId); })
+						),
+						NAME_None,
+						EUserInterfaceActionType::Check
+					);
+			}
+		}
+	}
+	
+	void AddReassignmentOptions(
+		FMenuBuilder& MenuBuilder,
+		const FSoftObjectPath& ContextObject,
+		const IConcertClient& ConcertClient,
+		const FReplicationClientManager& ReplicationManager,
+		ConcertSharedSlate::IObjectHierarchyModel& ObjectHierarchy,
+		FReassignObjectPropertiesLogic& ReassignmentLogic,
+		ConcertSharedSlate::IMultiReplicationStreamEditor& MultiStreamEditor
+		)
+	{
+		const TArray<const FReplicationClient*> SortedClients = ClientUtils::GetSortedClientList(ConcertClient, ReplicationManager);
+		
+		MenuBuilder.BeginSection(TEXT("Reassign.This"), LOCTEXT("Reassign.This", "Reassign this to"));
+		Private::AddReassignSection(MenuBuilder, SortedClients, Private::FInlineObjectPathArray{ ContextObject }, ConcertClient, ReassignmentLogic, MultiStreamEditor);
+		MenuBuilder.EndSection();
+
+		// Do not distract the user with more options if children have no assigned properties
+		const bool bHasChildrenWithProperties = ReassignmentLogic.IsAnyObjectOwned(Private::GetChildrenOfManagedObject(ObjectHierarchy, ContextObject));
+		if (bHasChildrenWithProperties)
+		{
+			MenuBuilder.BeginSection(TEXT("Reassign.Children"), LOCTEXT("Reassign.Children", "Reassign children to"));
+			Private::AddReassignSection(MenuBuilder, SortedClients, TAttribute<Private::FInlineObjectPathArray>::CreateLambda([ContextObject, &ObjectHierarchy]()
+			{
+				return Private::GetChildrenOfManagedObject(ObjectHierarchy, ContextObject);
+			}), ConcertClient, ReassignmentLogic, MultiStreamEditor);
+			MenuBuilder.EndSection();
+		}
+	}
+}
+
+#undef LOCTEXT_NAMESPACE
\ No newline at end of file
diff --git a/Engine/Plugins/Developer/Concert/ConcertApp/MultiUserClient/Source/MultiUserClient/Private/Widgets/ActiveSession/Replication/Client/FrequencyContextMenuUtils.cpp b/Engine/Plugins/Developer/Concert/ConcertApp/MultiUserClient/Source/MultiUserClient/Private/Widgets/ActiveSession/Replication/Client/FrequencyContextMenuUtils.cpp
deleted file mode 100644
diff --git a/Engine/Plugins/Developer/Concert/ConcertApp/MultiUserClient/Source/MultiUserClient/Private/Widgets/ActiveSession/Replication/Client/Multi/Columns/MultiStreamColumns.h b/Engine/Plugins/Developer/Concert/ConcertApp/MultiUserClient/Source/MultiUserClient/Private/Widgets/ActiveSession/Replication/Client/Multi/Columns/MultiStreamColumns.h
--- a/Engine/Plugins/Developer/Concert/ConcertApp/MultiUserClient/Source/MultiUserClient/Private/Widgets/ActiveSession/Replication/Client/Multi/Columns/MultiStreamColumns.h
+++ b/Engine/Plugins/Developer/Concert/ConcertApp/MultiUserClient/Source/MultiUserClient/Private/Widgets/ActiveSession/Replication/Client/Multi/Columns/MultiStreamColumns.h
@@ -27,7 +27,7 @@ namespace UE::ConcertSharedSlate
 namespace UE::MultiUserClient::MultiStreamColumns
 {
 	const extern FName ReplicationToggleColumnId;
-	const extern FName ReassignOwnershipColumnId;
+	const extern FName AssignedClientsColumnId;
 	const extern FName AssignPropertyColumnId;
 	
 	/* @see ETopLevelColumnOrder and EReplicationPropertyColumnOrder */
@@ -65,7 +65,7 @@ namespace UE::MultiUserClient::MultiStreamColumns
 	 * @param ClientManager Used to access all clients for display in the combo box drop-down
 	 * @param ColumnsSortPriority The order relative to the other columns
 	 */
-	ConcertSharedSlate::FObjectColumnEntry ReassignOwnership(
+	ConcertSharedSlate::FObjectColumnEntry AssignedClientsColumn(
 		TSharedRef<IConcertClient> ConcertClient,
 		TAttribute<TSharedPtr<ConcertSharedSlate::IMultiReplicationStreamEditor>> MultiStreamModelAttribute,
 		TAttribute<ConcertSharedSlate::IObjectHierarchyModel*> ObjectHierarchyModelAttribute,
diff --git a/Engine/Plugins/Developer/Concert/ConcertApp/MultiUserClient/Source/MultiUserClient/Private/Widgets/ActiveSession/Replication/Client/Multi/Columns/MultiStreamColumns_AssignProperties.cpp b/Engine/Plugins/Developer/Concert/ConcertApp/MultiUserClient/Source/MultiUserClient/Private/Widgets/ActiveSession/Replication/Client/Multi/Columns/MultiStreamColumns_AssignProperties.cpp
--- a/Engine/Plugins/Developer/Concert/ConcertApp/MultiUserClient/Source/MultiUserClient/Private/Widgets/ActiveSession/Replication/Client/Multi/Columns/MultiStreamColumns_AssignProperties.cpp
+++ b/Engine/Plugins/Developer/Concert/ConcertApp/MultiUserClient/Source/MultiUserClient/Private/Widgets/ActiveSession/Replication/Client/Multi/Columns/MultiStreamColumns_AssignProperties.cpp
@@ -99,7 +99,7 @@ namespace UE::MultiUserClient::MultiStreamColumns
 			virtual SHeaderRow::FColumn::FArguments CreateHeaderRowArgs() const override
 			{
 				return SHeaderRow::Column(AssignPropertyColumnId)
-					.DefaultLabel(LOCTEXT("Owner.Label", "Assigned Client"))
+					.DefaultLabel(LOCTEXT("Owner.Label", "Client"))
 					.ToolTipText(LOCTEXT("Owner.ToolTip", "Client that should replicate this property"))
 					.FillSized(FMultiUserReplicationStyle::Get()->GetFloat(TEXT("AllClients.Property.OwnerSize")));
 			}
diff --git a/Engine/Plugins/Developer/Concert/ConcertApp/MultiUserClient/Source/MultiUserClient/Private/Widgets/ActiveSession/Replication/Client/Multi/Columns/MultiStreamColumns_AssignedClients.cpp b/Engine/Plugins/Developer/Concert/ConcertApp/MultiUserClient/Source/MultiUserClient/Private/Widgets/ActiveSession/Replication/Client/Multi/Columns/MultiStreamColumns_AssignedClients.cpp
new file mode 100644
--- /dev/null
+++ b/Engine/Plugins/Developer/Concert/ConcertApp/MultiUserClient/Source/MultiUserClient/Private/Widgets/ActiveSession/Replication/Client/Multi/Columns/MultiStreamColumns_AssignedClients.cpp
@@ -0,0 +1,188 @@
+// Copyright Epic Games, Inc. All Rights Reserved.
+
+#include "MultiStreamColumns.h"
+
+#include "IConcertClient.h"
+#include "MultiUserReplicationStyle.h"
+#include "Replication/Editor/View/IMultiReplicationStreamEditor.h"
+#include "Replication/Submission/MultiEdit/ReassignObjectPropertiesLogic.h"
+#include "Widgets/ActiveSession/Replication/Client/ClientUtils.h"
+#include "Widgets/Client/ClientInfoHelpers.h"
+#include "Widgets/Client/SHorizontalClientList.h"
+
+#include "Widgets/Text/STextBlock.h"
+
+#define LOCTEXT_NAMESPACE "AssignedClientsColumnId"
+
+namespace UE::MultiUserClient::MultiStreamColumns
+{
+	const FName AssignedClientsColumnId(TEXT("AssignedClientsColumn"));
+
+	namespace Private
+	{
+		static TArray<FGuid> GetDisplayedClients(const FReassignObjectPropertiesLogic& ReassignmentLogic, const FSoftObjectPath& ManagedObject)
+		{
+			TArray<FGuid> ClientsWithOwnership;
+			ReassignmentLogic.EnumerateClientOwnershipState(ManagedObject, [&ClientsWithOwnership](const FGuid& ClientId, FReassignObjectPropertiesLogic::EOwnershipState Ownership)
+			{
+				if (Ownership == FReassignObjectPropertiesLogic::EOwnershipState::HasObjectRegistered)
+				{
+					ClientsWithOwnership.Add(ClientId);
+				}
+				return EBreakBehavior::Continue;
+			});
+			return ClientsWithOwnership;
+		}
+		
+		class SAssignedClientsWidget : public SCompoundWidget
+		{
+			TSharedPtr<ConcertClientSharedSlate::SHorizontalClientList> ClientList;
+			FSoftObjectPath ManagedObject;
+			FReassignObjectPropertiesLogic* ReassignmentLogic = nullptr;
+		public:
+		
+			SLATE_BEGIN_ARGS(SAssignedClientsWidget){}
+			SLATE_END_ARGS()
+
+			void Construct(
+				const FArguments& InArgs,
+				const TSharedRef<IConcertClient>& InConcertClient,
+				FReassignObjectPropertiesLogic& InReassignmentLogic,
+				FSoftObjectPath InManagedObject,
+				TSharedPtr<FText> InHighlightText
+				)
+			{
+				ReassignmentLogic = &InReassignmentLogic;
+				ManagedObject = MoveTemp(InManagedObject);
+			
+				ChildSlot
+				[
+					SAssignNew(ClientList, ConcertClientSharedSlate::SHorizontalClientList)
+					.IsLocalClient(ConcertClientSharedSlate::MakeIsLocalClientGetter(InConcertClient))
+					.GetClientInfo(ConcertClientSharedSlate::MakeClientInfoGetter(InConcertClient))
+					.Font(FAppStyle::GetFontStyle("PropertyWindow.NormalFont"))
+					.HighlightText_Lambda([InHighlightText](){ return InHighlightText ? *InHighlightText : FText::GetEmpty(); })
+					.ListToolTipText(LOCTEXT("Clients.ToolTip", "These clients will replicate their assigned properties when replication is active.\nYou can pause & resume replication at the beginnig of this row."))
+					.EmptyListSlot()
+					[
+						SNew(STextBlock)
+						.Text(LOCTEXT("NoClients.Label", "No assigned properties"))
+						.ToolTipText(LOCTEXT("NoClients.ToolTip", "Click this row and then assign the properties to the client that should replicate them."))
+						.ColorAndOpacity(FSlateColor::UseSubduedForeground())
+					]
+				];
+			
+				UpdateComboButtonContent();
+				ReassignmentLogic->OnOwnershipChanged().AddSP(this, &SAssignedClientsWidget::UpdateComboButtonContent);
+			}
+		
+			virtual ~SAssignedClientsWidget() override
+			{
+				ReassignmentLogic->OnOwnershipChanged().RemoveAll(this);
+			}
+
+			void UpdateComboButtonContent() const
+			{
+				ClientList->RefreshList(
+					GetDisplayedClients(*ReassignmentLogic, ManagedObject)
+					);
+			}
+		};
+	}
+	
+	ConcertSharedSlate::FObjectColumnEntry AssignedClientsColumn(
+		TSharedRef<IConcertClient> ConcertClient,
+		TAttribute<TSharedPtr<ConcertSharedSlate::IMultiReplicationStreamEditor>> MultiStreamModelAttribute,
+		TAttribute<ConcertSharedSlate::IObjectHierarchyModel*> ObjectHierarchyModelAttribute,
+		FReassignObjectPropertiesLogic& ReassignmentLogic,
+		const FReplicationClientManager& ClientManager,
+		const int32 ColumnsSortPriority
+		)
+	{
+		class FObjectColumn_ReassignOwnership : public ConcertSharedSlate::IObjectTreeColumn
+		{
+		public:
+
+			FObjectColumn_ReassignOwnership(
+				TSharedRef<IConcertClient> ConcertClient,
+				TAttribute<TSharedPtr<ConcertSharedSlate::IMultiReplicationStreamEditor>> MultiStreamModelAttribute,
+				TAttribute<ConcertSharedSlate::IObjectHierarchyModel*> ObjectHierarchyModelAttribute,
+				FReassignObjectPropertiesLogic& ReassignmentLogic,
+				const FReplicationClientManager& ClientManager
+				)
+				: ConcertClient(MoveTemp(ConcertClient))
+				, MultiStreamModelAttribute(MoveTemp(MultiStreamModelAttribute))
+				, ObjectHierarchyModelAttribute(MoveTemp(ObjectHierarchyModelAttribute))
+				, ReassignmentLogic(ReassignmentLogic)
+				, ClientManager(ClientManager)
+			{}
+			
+			virtual SHeaderRow::FColumn::FArguments CreateHeaderRowArgs() const override
+			{
+				return SHeaderRow::Column(AssignedClientsColumnId)
+					.DefaultLabel(LOCTEXT("Owner.Label", "Clients"))
+					.ToolTipText(LOCTEXT("Owner.ToolTip", "Clients that have registered properties for the object"))
+					.FillSized(FMultiUserReplicationStyle::Get()->GetFloat(TEXT("AllClients.Object.OwnerSize")));
+			}
+			
+			virtual TSharedRef<SWidget> GenerateColumnWidget(const FBuildArgs& InArgs) override
+			{
+				return SNew(Private::SAssignedClientsWidget, ConcertClient, ReassignmentLogic, InArgs.RowItem.RowData.GetObjectPath(), InArgs.HighlightText);
+			}
+			virtual void PopulateSearchString(const ConcertSharedSlate::FObjectTreeRowContext& InItem, TArray<FString>& InOutSearchStrings) const override
+			{
+				for (const FGuid& ClientId : Private::GetDisplayedClients(ReassignmentLogic, InItem.RowData.GetObjectPath()))
+				{
+					InOutSearchStrings.Add(ClientUtils::GetClientDisplayName(*ConcertClient->GetCurrentSession(), ClientId));
+				}
+			}
+
+			virtual bool CanBeSorted() const override { return true; }
+			virtual bool IsLessThan(const ConcertSharedSlate::FObjectTreeRowContext& Left, const ConcertSharedSlate::FObjectTreeRowContext& Right) const override
+			{
+				const TOptional<FString> LeftClientDisplayString = GetDisplayString(Left.RowData.GetObjectPath());
+				const TOptional<FString> RightClientDisplayString = GetDisplayString(Right.RowData.GetObjectPath());
+			
+				if (LeftClientDisplayString && RightClientDisplayString)
+				{
+					return *LeftClientDisplayString < *RightClientDisplayString;
+				}
+				// Our rule: set < unset. This way unassigned appears last.
+				return LeftClientDisplayString.IsSet() && !RightClientDisplayString.IsSet();
+			}
+
+		private:
+			
+			const TSharedRef<IConcertClient> ConcertClient;
+			const TAttribute<TSharedPtr<ConcertSharedSlate::IMultiReplicationStreamEditor>> MultiStreamModelAttribute;
+			const TAttribute<ConcertSharedSlate::IObjectHierarchyModel*> ObjectHierarchyModelAttribute;
+			FReassignObjectPropertiesLogic& ReassignmentLogic;
+			const FReplicationClientManager& ClientManager;
+			
+			TOptional<FString> GetDisplayString(const FSoftObjectPath& ManagedObject) const
+			{
+				using SWidgetType = ConcertClientSharedSlate::SHorizontalClientList;
+				const TArray<FGuid> Clients = Private::GetDisplayedClients(ReassignmentLogic, ManagedObject);
+				const ConcertSharedSlate::FIsLocalClient IsLocalClientDelegate = ConcertClientSharedSlate::MakeIsLocalClientGetter(ConcertClient);
+				return SWidgetType::GetDisplayString(
+					Clients,
+					ConcertClientSharedSlate::MakeClientInfoGetter(ConcertClient),
+					SWidgetType::FSortPredicate::CreateStatic(&SWidgetType::SortLocalClientFirstThenAlphabetical, IsLocalClientDelegate),
+					IsLocalClientDelegate
+					);
+			}
+		};
+
+		return {
+			ConcertSharedSlate::TReplicationColumnDelegates<ConcertSharedSlate::FObjectTreeRowContext>::FCreateColumn::CreateLambda(
+				[ConcertClient = MoveTemp(ConcertClient), MultiStreamModelAttribute = MoveTemp(MultiStreamModelAttribute), ObjectHierarchyModelAttribute = MoveTemp(ObjectHierarchyModelAttribute), &ReassignmentLogic, &ClientManager]()
+				{
+					return MakeShared<FObjectColumn_ReassignOwnership>(ConcertClient, MultiStreamModelAttribute, ObjectHierarchyModelAttribute, ReassignmentLogic, ClientManager);
+				}),
+			AssignedClientsColumnId,
+			{ ColumnsSortPriority }
+		};
+	}
+}
+
+#undef LOCTEXT_NAMESPACE
\ No newline at end of file
diff --git a/Engine/Plugins/Developer/Concert/ConcertApp/MultiUserClient/Source/MultiUserClient/Private/Widgets/ActiveSession/Replication/Client/Multi/Columns/MultiStreamColumns_ReassignOwnership.cpp b/Engine/Plugins/Developer/Concert/ConcertApp/MultiUserClient/Source/MultiUserClient/Private/Widgets/ActiveSession/Replication/Client/Multi/Columns/MultiStreamColumns_ReassignOwnership.cpp
deleted file mode 100644
diff --git a/Engine/Plugins/Developer/Concert/ConcertApp/MultiUserClient/Source/MultiUserClient/Private/Widgets/ActiveSession/Replication/Client/Multi/Columns/SReassignObjectComboBox.cpp b/Engine/Plugins/Developer/Concert/ConcertApp/MultiUserClient/Source/MultiUserClient/Private/Widgets/ActiveSession/Replication/Client/Multi/Columns/SReassignObjectComboBox.cpp
deleted file mode 100644
diff --git a/Engine/Plugins/Developer/Concert/ConcertApp/MultiUserClient/Source/MultiUserClient/Private/Widgets/ActiveSession/Replication/Client/Multi/Columns/SReassignObjectComboBox.h b/Engine/Plugins/Developer/Concert/ConcertApp/MultiUserClient/Source/MultiUserClient/Private/Widgets/ActiveSession/Replication/Client/Multi/Columns/SReassignObjectComboBox.h
deleted file mode 100644
diff --git a/Engine/Plugins/Developer/Concert/ConcertApp/MultiUserClient/Source/MultiUserClient/Private/Widgets/ActiveSession/Replication/Client/Multi/SMultiClientView.cpp b/Engine/Plugins/Developer/Concert/ConcertApp/MultiUserClient/Source/MultiUserClient/Private/Widgets/ActiveSession/Replication/Client/Multi/SMultiClientView.cpp
--- a/Engine/Plugins/Developer/Concert/ConcertApp/MultiUserClient/Source/MultiUserClient/Private/Widgets/ActiveSession/Replication/Client/Multi/SMultiClientView.cpp
+++ b/Engine/Plugins/Developer/Concert/ConcertApp/MultiUserClient/Source/MultiUserClient/Private/Widgets/ActiveSession/Replication/Client/Multi/SMultiClientView.cpp
@@ -12,9 +12,9 @@
 #include "Replication/Editor/View/IMultiReplicationStreamEditor.h"
 #include "Replication/Editor/View/IReplicationStreamEditor.h"
 #include "Replication/Editor/Model/ObjectSource/ActorSelectionSourceModel.h"
-#include "Widgets/ActiveSession/Replication/Client/FrequencyContextMenuUtils.h"
 #include "Widgets/ActiveSession/Replication/Client/Multi/Columns/MultiStreamColumns.h"
 #include "Widgets/ActiveSession/Replication/Client/SReplicationStatus.h"
+#include "Widgets/ActiveSession/Replication/Client/Context/ContextMenuUtils.h"
 
 #include "Widgets/SBoxPanel.h"
 
@@ -26,6 +26,7 @@ namespace UE::MultiUserClient
 	{
 		StreamModel = MakeShared<FMultiStreamModel>(InDisplayClientsModel, InClientManager);
 
+		ConcertClient = MoveTemp(InConcertClient);
 		ClientManager = &InClientManager;
 		ClientManager->OnRemoteClientsChanged().AddSP(this, &SMultiClientView::RebuildClientSubscriptions);
 		SelectionModel = &InDisplayClientsModel;
@@ -40,7 +41,7 @@ namespace UE::MultiUserClient
 			+SVerticalBox::Slot()
 			.FillHeight(1.f)
 			[
-				CreateEditorContent(InConcertClient, InClientManager)
+				CreateEditorContent(ConcertClient.ToSharedRef(), InClientManager)
 			]
 		];
 		
@@ -110,7 +111,7 @@ namespace UE::MultiUserClient
 			.ObjectColumns =
 			{
 				MultiStreamColumns::ReplicationToggle(InConcertClient, ObjectHierarchyAttribute, InClientManager),
-				MultiStreamColumns::ReassignOwnership(InConcertClient, MultiStreamEditorAttribute, ObjectHierarchyAttribute, InClientManager.GetReassignmentLogic(), InClientManager)
+				MultiStreamColumns::AssignedClientsColumn(InConcertClient, MultiStreamEditorAttribute, ObjectHierarchyAttribute, InClientManager.GetReassignmentLogic(), InClientManager)
 			}
 		};
 		
@@ -172,7 +173,7 @@ namespace UE::MultiUserClient
 	void SMultiClientView::OnClientChanged() const
 	{
 		// When reassignment operations complete, the content of the columns changes so a resort is required.
-		StreamEditor->GetEditorBase().RequestObjectColumnResort(MultiStreamColumns::ReassignOwnershipColumnId);
+		StreamEditor->GetEditorBase().RequestObjectColumnResort(MultiStreamColumns::AssignedClientsColumnId);
 		StreamEditor->GetEditorBase().RequestPropertyColumnResort(MultiStreamColumns::AssignPropertyColumnId);
 	}
 
@@ -184,7 +185,20 @@ namespace UE::MultiUserClient
 
 	void SMultiClientView::ExtendObjectContextMenu(FMenuBuilder& MenuBuilder, TConstArrayView<FSoftObjectPath> ContextObjects) const
 	{
-		FrequencyContextMenuUtils::AddFrequencyOptionsIfOneContextObject_MultiClient(MenuBuilder, ContextObjects, *ClientManager);
+		ContextMenuUtils::AddFrequencyOptionsIfOneContextObject_MultiClient(MenuBuilder, ContextObjects, *ClientManager);
+
+		if (ContextObjects.Num() == 1)
+		{
+			ContextMenuUtils::AddReassignmentOptions(
+				MenuBuilder,
+				ContextObjects[0],
+				*ConcertClient,
+				*ClientManager,
+				*ObjectHierarchy,
+				ClientManager->GetReassignmentLogic(),
+				*StreamEditor
+				);
+		}
 	}
 }
 
diff --git a/Engine/Plugins/Developer/Concert/ConcertApp/MultiUserClient/Source/MultiUserClient/Private/Widgets/ActiveSession/Replication/Client/Multi/SMultiClientView.h b/Engine/Plugins/Developer/Concert/ConcertApp/MultiUserClient/Source/MultiUserClient/Private/Widgets/ActiveSession/Replication/Client/Multi/SMultiClientView.h
--- a/Engine/Plugins/Developer/Concert/ConcertApp/MultiUserClient/Source/MultiUserClient/Private/Widgets/ActiveSession/Replication/Client/Multi/SMultiClientView.h
+++ b/Engine/Plugins/Developer/Concert/ConcertApp/MultiUserClient/Source/MultiUserClient/Private/Widgets/ActiveSession/Replication/Client/Multi/SMultiClientView.h
@@ -30,11 +30,12 @@ namespace UE::MultiUserClient
 		SLATE_BEGIN_ARGS(SMultiClientView){}
 		SLATE_END_ARGS()
 
-		void Construct(const FArguments& InArgs, TSharedRef<IConcertClient> InConcertClient, FReplicationClientManager& InClientManager, IClientSelectionModel& InDisplayClientsModel);
+		void Construct(const FArguments& InArgs, TSharedRef<IConcertClient> InConcertClient, FReplicationClientManager& InClientManager UE_LIFETIMEBOUND, IClientSelectionModel& InDisplayClientsModel UE_LIFETIMEBOUND);
 		virtual ~SMultiClientView() override;
 
 	private:
 
+		TSharedPtr<IConcertClient> ConcertClient;
 		FReplicationClientManager* ClientManager = nullptr;
 		IClientSelectionModel* SelectionModel = nullptr;
 		
diff --git a/Engine/Plugins/Developer/Concert/ConcertUI/ConcertSharedSlate/Source/ConcertSharedSlate/Private/Replication/Editor/View/ObjectEditor/SBaseReplicationStreamEditor.cpp b/Engine/Plugins/Developer/Concert/ConcertUI/ConcertSharedSlate/Source/ConcertSharedSlate/Private/Replication/Editor/View/ObjectEditor/SBaseReplicationStreamEditor.cpp
--- a/Engine/Plugins/Developer/Concert/ConcertUI/ConcertSharedSlate/Source/ConcertSharedSlate/Private/Replication/Editor/View/ObjectEditor/SBaseReplicationStreamEditor.cpp
+++ b/Engine/Plugins/Developer/Concert/ConcertUI/ConcertSharedSlate/Source/ConcertSharedSlate/Private/Replication/Editor/View/ObjectEditor/SBaseReplicationStreamEditor.cpp
@@ -201,7 +201,7 @@ namespace UE::ConcertSharedSlate
 
 	TSharedPtr<SWidget> SBaseReplicationStreamEditor::OnObjectsContextMenuOpening()
 	{
-		FMenuBuilder MenuBuilder(true, nullptr);
+		FMenuBuilder MenuBuilder(false, nullptr);
 		
 		AddObjectSourceContextMenuOptions(MenuBuilder);
 		MenuBuilder.AddMenuEntry(
diff --git a/Engine/Plugins/Developer/Concert/ConcertUI/ConcertSharedSlate/Source/ConcertSharedSlate/Private/Widgets/Client/SHorizontalClientList.cpp b/Engine/Plugins/Developer/Concert/ConcertUI/ConcertSharedSlate/Source/ConcertSharedSlate/Private/Widgets/Client/SHorizontalClientList.cpp
--- a/Engine/Plugins/Developer/Concert/ConcertUI/ConcertSharedSlate/Source/ConcertSharedSlate/Private/Widgets/Client/SHorizontalClientList.cpp
+++ b/Engine/Plugins/Developer/Concert/ConcertUI/ConcertSharedSlate/Source/ConcertSharedSlate/Private/Widgets/Client/SHorizontalClientList.cpp
@@ -98,6 +98,7 @@ namespace UE::ConcertClientSharedSlate
 			[
 				SAssignNew(ScrollBox, SScrollBox)
 				.Orientation(Orient_Horizontal)
+				.ToolTipText(InArgs._ListToolTipText)
 			]
 		];
 	}
diff --git a/Engine/Plugins/Developer/Concert/ConcertUI/ConcertSharedSlate/Source/ConcertSharedSlate/Public/Widgets/Client/SHorizontalClientList.h b/Engine/Plugins/Developer/Concert/ConcertUI/ConcertSharedSlate/Source/ConcertSharedSlate/Public/Widgets/Client/SHorizontalClientList.h
--- a/Engine/Plugins/Developer/Concert/ConcertUI/ConcertSharedSlate/Source/ConcertSharedSlate/Public/Widgets/Client/SHorizontalClientList.h
+++ b/Engine/Plugins/Developer/Concert/ConcertUI/ConcertSharedSlate/Source/ConcertSharedSlate/Public/Widgets/Client/SHorizontalClientList.h
@@ -50,6 +50,9 @@ namespace UE::ConcertClientSharedSlate
 			
 			/** Defaults to placing the local client first (if contained) and sorting alphabetically otherwise. */
 			SLATE_EVENT(FSortPredicate, SortPredicate)
+
+			/** Tooltip text to display when the list is non-empty. */
+			SLATE_ATTRIBUTE(FText, ListToolTipText)
 			
 			/** The widget to display when the list is empty */
 			SLATE_NAMED_SLOT(FArguments, EmptyListSlot)
