Refactor: Split up FRegularQueryService into FRegularQueryService and FStreamAndAuthorityQueryService.h.
- FRegularQueryService exposes a unified interface for all things that need to be queried.
- FStreamAndAuthorityQueryService handles querying
- This refactor is in preparation of adding mute state querying

diff --git a/Engine/Plugins/Developer/Concert/ConcertApp/MultiUserClient/Source/MultiUserClient/Private/Replication/Authority/AuthoritySynchronizer_RemoteClient.cpp b/Engine/Plugins/Developer/Concert/ConcertApp/MultiUserClient/Source/MultiUserClient/Private/Replication/Authority/AuthoritySynchronizer_RemoteClient.cpp
--- a/Engine/Plugins/Developer/Concert/ConcertApp/MultiUserClient/Source/MultiUserClient/Private/Replication/Authority/AuthoritySynchronizer_RemoteClient.cpp
+++ b/Engine/Plugins/Developer/Concert/ConcertApp/MultiUserClient/Source/MultiUserClient/Private/Replication/Authority/AuthoritySynchronizer_RemoteClient.cpp
@@ -2,11 +2,11 @@
 
 #include "AuthoritySynchronizer_RemoteClient.h"
 
-#include "Replication/Util/RegularQueryService.h"
+#include "Replication/Util/Query/StreamAndAuthorityQueryService.h"
 
 namespace UE::MultiUserClient
 {
-	FAuthoritySynchronizer_RemoteClient::FAuthoritySynchronizer_RemoteClient(const FGuid& RemoteEndpointId, FRegularQueryService& InQueryService)
+	FAuthoritySynchronizer_RemoteClient::FAuthoritySynchronizer_RemoteClient(const FGuid& RemoteEndpointId, FStreamAndAuthorityQueryService& InQueryService)
 		: QueryService(InQueryService)
 		, QueryStreamHandle(
     		QueryService.RegisterAuthorityQuery(
diff --git a/Engine/Plugins/Developer/Concert/ConcertApp/MultiUserClient/Source/MultiUserClient/Private/Replication/Authority/AuthoritySynchronizer_RemoteClient.h b/Engine/Plugins/Developer/Concert/ConcertApp/MultiUserClient/Source/MultiUserClient/Private/Replication/Authority/AuthoritySynchronizer_RemoteClient.h
--- a/Engine/Plugins/Developer/Concert/ConcertApp/MultiUserClient/Source/MultiUserClient/Private/Replication/Authority/AuthoritySynchronizer_RemoteClient.h
+++ b/Engine/Plugins/Developer/Concert/ConcertApp/MultiUserClient/Source/MultiUserClient/Private/Replication/Authority/AuthoritySynchronizer_RemoteClient.h
@@ -5,6 +5,7 @@
 #include "IClientAuthoritySynchronizer.h"
 
 #include "Containers/Set.h"
+#include "HAL/Platform.h"
 #include "UObject/SoftObjectPath.h"
 
 struct FGuid;
@@ -12,14 +13,14 @@ struct FConcertAuthorityClientInfo;
 
 namespace UE::MultiUserClient
 {
-	class FRegularQueryService;
+	class FStreamAndAuthorityQueryService;
 
 	class FAuthoritySynchronizer_RemoteClient : public FAuthoritySynchronizer_Base
 	{
 	public:
 		
-		FAuthoritySynchronizer_RemoteClient(const FGuid& RemoteEndpointId, FRegularQueryService& InQueryService);
-		virtual ~FAuthoritySynchronizer_RemoteClient();
+		FAuthoritySynchronizer_RemoteClient(const FGuid& RemoteEndpointId, FStreamAndAuthorityQueryService& InQueryService UE_LIFETIMEBOUND);
+		virtual ~FAuthoritySynchronizer_RemoteClient() override;
 
 		//~ Begin IClientAuthoritySynchronizer Interface
 		virtual bool HasAnyAuthority() const override;
@@ -29,7 +30,7 @@ namespace UE::MultiUserClient
 	private:
 
 		/** Queries the server in regular intervals. This services outlives our object. */
-		FRegularQueryService& QueryService;
+		FStreamAndAuthorityQueryService& QueryService;
 		
 		/** Used to unregister HandleStreamQuery upon destruction. */
 		const FDelegateHandle QueryStreamHandle;
diff --git a/Engine/Plugins/Developer/Concert/ConcertApp/MultiUserClient/Source/MultiUserClient/Private/Replication/Client/RemoteReplicationClient.cpp b/Engine/Plugins/Developer/Concert/ConcertApp/MultiUserClient/Source/MultiUserClient/Private/Replication/Client/RemoteReplicationClient.cpp
--- a/Engine/Plugins/Developer/Concert/ConcertApp/MultiUserClient/Source/MultiUserClient/Private/Replication/Client/RemoteReplicationClient.cpp
+++ b/Engine/Plugins/Developer/Concert/ConcertApp/MultiUserClient/Source/MultiUserClient/Private/Replication/Client/RemoteReplicationClient.cpp
@@ -16,7 +16,7 @@ namespace UE::MultiUserClient
 		TSharedRef<IConcertClient> InClient,
 		FGlobalAuthorityCache& InAuthorityCache,
 		UMultiUserReplicationClientPreset& InSessionContent,
-		FRegularQueryService& QueryService
+		FStreamAndAuthorityQueryService& QueryService
 		)
 		: FReplicationClient(
 			InConcertClientId,
diff --git a/Engine/Plugins/Developer/Concert/ConcertApp/MultiUserClient/Source/MultiUserClient/Private/Replication/Client/RemoteReplicationClient.h b/Engine/Plugins/Developer/Concert/ConcertApp/MultiUserClient/Source/MultiUserClient/Private/Replication/Client/RemoteReplicationClient.h
--- a/Engine/Plugins/Developer/Concert/ConcertApp/MultiUserClient/Source/MultiUserClient/Private/Replication/Client/RemoteReplicationClient.h
+++ b/Engine/Plugins/Developer/Concert/ConcertApp/MultiUserClient/Source/MultiUserClient/Private/Replication/Client/RemoteReplicationClient.h
@@ -4,11 +4,13 @@
 
 #include "ReplicationClient.h"
 
+#include "HAL/Platform.h"
+
 class IConcertClient;
 
 namespace UE::MultiUserClient
 {
-	class FRegularQueryService;
+	class FStreamAndAuthorityQueryService;
 	
 	/** Holds extra information about a remote replication client. */
 	class FRemoteReplicationClient : public FReplicationClient
@@ -25,11 +27,11 @@ namespace UE::MultiUserClient
 		 */
 		FRemoteReplicationClient(
 			const FGuid& InConcertClientId,
-			FReplicationDiscoveryContainer& InDiscoveryContainer,
+			FReplicationDiscoveryContainer& InDiscoveryContainer UE_LIFETIMEBOUND,
 			TSharedRef<IConcertClient> InClient,
-			FGlobalAuthorityCache& InAuthorityCache,
-			UMultiUserReplicationClientPreset& InSessionContent,
-			FRegularQueryService& QueryService
+			FGlobalAuthorityCache& InAuthorityCache UE_LIFETIMEBOUND,
+			UMultiUserReplicationClientPreset& InSessionContent UE_LIFETIMEBOUND,
+			FStreamAndAuthorityQueryService& QueryService UE_LIFETIMEBOUND
 			);
 	};
 }
diff --git a/Engine/Plugins/Developer/Concert/ConcertApp/MultiUserClient/Source/MultiUserClient/Private/Replication/Client/ReplicationClientManager.cpp b/Engine/Plugins/Developer/Concert/ConcertApp/MultiUserClient/Source/MultiUserClient/Private/Replication/Client/ReplicationClientManager.cpp
--- a/Engine/Plugins/Developer/Concert/ConcertApp/MultiUserClient/Source/MultiUserClient/Private/Replication/Client/ReplicationClientManager.cpp
+++ b/Engine/Plugins/Developer/Concert/ConcertApp/MultiUserClient/Source/MultiUserClient/Private/Replication/Client/ReplicationClientManager.cpp
@@ -22,7 +22,7 @@ namespace UE::MultiUserClient
 		, ConcertClient(InClient)
 		, Session(InSession)
 		, RegisteredExtenders(InRegisteredExtenders)
-		, QueryService(InClient)
+		, QueryService(*ConcertClient)
 		, AuthorityCache(*this)
 		, LocalClient([this, InClient]()
 		{
@@ -169,7 +169,7 @@ namespace UE::MultiUserClient
 			ConcertClient->GetConcertClient(),
 			AuthorityCache,
 			*SessionContent->AddClient(),
-			QueryService
+			QueryService.GetStreamAndAuthorityQueryService()
 			);
 		FRemoteReplicationClient& RemoteClient = *RemoteClientPtr;
 		RemoteClients.Emplace(
diff --git a/Engine/Plugins/Developer/Concert/ConcertApp/MultiUserClient/Source/MultiUserClient/Private/Replication/Client/ReplicationClientManager.h b/Engine/Plugins/Developer/Concert/ConcertApp/MultiUserClient/Source/MultiUserClient/Private/Replication/Client/ReplicationClientManager.h
--- a/Engine/Plugins/Developer/Concert/ConcertApp/MultiUserClient/Source/MultiUserClient/Private/Replication/Client/ReplicationClientManager.h
+++ b/Engine/Plugins/Developer/Concert/ConcertApp/MultiUserClient/Source/MultiUserClient/Private/Replication/Client/ReplicationClientManager.h
@@ -9,7 +9,7 @@
 #include "Replication/Submission/MultiEdit/ReassignObjectPropertiesLogic.h"
 #include "Replication/Submission/Notification/SubmissionNotifier.h"
 #include "Replication/Util/GlobalAuthorityCache.h"
-#include "Replication/Util/RegularQueryService.h"
+#include "Replication/Util/Query/RegularQueryService.h"
 
 #include "UObject/GCObject.h"
 #include "Templates/UnrealTemplate.h"
@@ -17,16 +17,10 @@
 
 class IConcertClientSession;
 class IConcertSyncClient;
-
+namespace UE::ConcertSharedSlate { class IEditableReplicationStreamModel; }
 enum class EConcertClientStatus : uint8;
-
 struct FConcertSessionClientInfo;
 
-namespace UE::ConcertSharedSlate
-{
-	class IEditableReplicationStreamModel;
-}
-
 namespace UE::MultiUserClient
 {
 	class FRemoteReplicationClient;
diff --git a/Engine/Plugins/Developer/Concert/ConcertApp/MultiUserClient/Source/MultiUserClient/Private/Replication/Stream/StreamSynchronizer_RemoteClient.cpp b/Engine/Plugins/Developer/Concert/ConcertApp/MultiUserClient/Source/MultiUserClient/Private/Replication/Stream/StreamSynchronizer_RemoteClient.cpp
--- a/Engine/Plugins/Developer/Concert/ConcertApp/MultiUserClient/Source/MultiUserClient/Private/Replication/Stream/StreamSynchronizer_RemoteClient.cpp
+++ b/Engine/Plugins/Developer/Concert/ConcertApp/MultiUserClient/Source/MultiUserClient/Private/Replication/Stream/StreamSynchronizer_RemoteClient.cpp
@@ -2,15 +2,14 @@
 
 #include "StreamSynchronizer_RemoteClient.h"
 
-#include "ConcertLogGlobal.h"
 #include "Assets/MultiUserReplicationClientPreset.h"
-#include "Replication/Util/RegularQueryService.h"
+#include "Replication/Util/Query/StreamAndAuthorityQueryService.h"
 
 namespace UE::MultiUserClient
 {
 	constexpr float QueryTimeInterval = 1.f;
 	
-	FStreamSynchronizer_RemoteClient::FStreamSynchronizer_RemoteClient(const FGuid& RemoteEndpointId, FRegularQueryService& InQueryService)
+	FStreamSynchronizer_RemoteClient::FStreamSynchronizer_RemoteClient(const FGuid& RemoteEndpointId, FStreamAndAuthorityQueryService& InQueryService)
 		: QueryService(InQueryService)
 		, QueryStreamHandle(
 			QueryService.RegisterStreamQuery(
diff --git a/Engine/Plugins/Developer/Concert/ConcertApp/MultiUserClient/Source/MultiUserClient/Private/Replication/Stream/StreamSynchronizer_RemoteClient.h b/Engine/Plugins/Developer/Concert/ConcertApp/MultiUserClient/Source/MultiUserClient/Private/Replication/Stream/StreamSynchronizer_RemoteClient.h
--- a/Engine/Plugins/Developer/Concert/ConcertApp/MultiUserClient/Source/MultiUserClient/Private/Replication/Stream/StreamSynchronizer_RemoteClient.h
+++ b/Engine/Plugins/Developer/Concert/ConcertApp/MultiUserClient/Source/MultiUserClient/Private/Replication/Stream/StreamSynchronizer_RemoteClient.h
@@ -6,18 +6,19 @@
 #include "Replication/Data/ReplicationStream.h"
 
 #include "Containers/Array.h"
+#include "HAL/Platform.h"
 #include "Templates/UnrealTemplate.h"
 
 namespace UE::MultiUserClient
 {
-	class FRegularQueryService;
+	class FStreamAndAuthorityQueryService;
 	
 	/** Tracks the state of a remote client by querying the client's state in regular intervals. */
 	class FStreamSynchronizer_RemoteClient : public IClientStreamSynchronizer, public FNoncopyable
 	{
 	public:
 		
-		FStreamSynchronizer_RemoteClient(const FGuid& RemoteEndpointId, FRegularQueryService& InQueryService);
+		FStreamSynchronizer_RemoteClient(const FGuid& RemoteEndpointId, FStreamAndAuthorityQueryService& InQueryService UE_LIFETIMEBOUND);
 		virtual ~FStreamSynchronizer_RemoteClient() override;
 
 		//~ Begin IClientStreamSynchronizer Interface
@@ -30,7 +31,7 @@ namespace UE::MultiUserClient
 	private:
 
 		/** Queries the server in regular intervals. This services outlives our object. */
-		FRegularQueryService& QueryService;
+		FStreamAndAuthorityQueryService& QueryService;
 		/** Used to unregister HandleStreamQuery upon destruction. */
 		const FDelegateHandle QueryStreamHandle; 
 
diff --git a/Engine/Plugins/Developer/Concert/ConcertApp/MultiUserClient/Source/MultiUserClient/Private/Replication/Util/Query/RegularQueryService.cpp b/Engine/Plugins/Developer/Concert/ConcertApp/MultiUserClient/Source/MultiUserClient/Private/Replication/Util/Query/RegularQueryService.cpp
new file mode 100644
--- /dev/null
+++ b/Engine/Plugins/Developer/Concert/ConcertApp/MultiUserClient/Source/MultiUserClient/Private/Replication/Util/Query/RegularQueryService.cpp
@@ -0,0 +1,29 @@
+﻿// Copyright Epic Games, Inc. All Rights Reserved.
+
+#include "RegularQueryService.h"
+
+namespace UE::MultiUserClient
+{
+	FRegularQueryService::FRegularQueryService(const IConcertSyncClient& InOwningClient, float InInterval)
+		: OwningClient(InOwningClient)
+		, TickerDelegateHandle(FTSTicker::GetCoreTicker().AddTicker(
+			TEXT("Multi-User Replication Query"),
+			InInterval,
+			[this](float)
+			{
+				Tick();
+				return true;
+			}))
+		, StreamAndAuthorityQueryService(Token, OwningClient)
+	{}
+
+	FRegularQueryService::~FRegularQueryService()
+	{
+		FTSTicker::GetCoreTicker().RemoveTicker(TickerDelegateHandle);
+	}
+
+	void FRegularQueryService::Tick()
+	{
+		StreamAndAuthorityQueryService.SendQueryEvent();
+	}
+}
diff --git a/Engine/Plugins/Developer/Concert/ConcertApp/MultiUserClient/Source/MultiUserClient/Private/Replication/Util/Query/RegularQueryService.h b/Engine/Plugins/Developer/Concert/ConcertApp/MultiUserClient/Source/MultiUserClient/Private/Replication/Util/Query/RegularQueryService.h
new file mode 100644
--- /dev/null
+++ b/Engine/Plugins/Developer/Concert/ConcertApp/MultiUserClient/Source/MultiUserClient/Private/Replication/Util/Query/RegularQueryService.h
@@ -0,0 +1,57 @@
+﻿// Copyright Epic Games, Inc. All Rights Reserved.
+
+#pragma once
+
+#include "StreamAndAuthorityQueryService.h"
+#include "Replication/IToken.h"
+
+#include "Containers/Ticker.h"
+#include "HAL/Platform.h"
+#include "Templates/UnrealTemplate.h"
+
+class IConcertSyncClient;
+
+namespace UE::MultiUserClient
+{
+	/**
+	 * This service manages multiple sub-services that query the server state, e.g. stream, authority, etc.
+	 * Its responsibility is to tick the sub-services and manage access to them. 
+	 */
+	class FRegularQueryService : public FNoncopyable
+	{
+	public:
+		
+		/**
+		 * @param InOwningClient The client through which queries are made. The caller ensures it outlasts the constructed instance.
+		 * @param InInterval The interval at which queries are to be made.
+		 */
+		FRegularQueryService(const IConcertSyncClient& InOwningClient UE_LIFETIMEBOUND, float InInterval = 1.f);
+		~FRegularQueryService();
+
+		/** @return The service for querying stream and authority state. */
+		FStreamAndAuthorityQueryService& GetStreamAndAuthorityQueryService() { return StreamAndAuthorityQueryService; }
+
+	private:
+		
+		/** Passed to sub-services to check whether they were destroyed after a response is received. */
+		const TSharedRef<FToken> Token = FToken::Make();
+		
+		/**
+		 * Used to send queries.
+		 * 
+		 * This FRegularQueryService's owner is supposed to make sure this FRegularQueryService is destroyed
+		 * when the client shuts down.
+		 */
+		const IConcertSyncClient& OwningClient;
+
+		/** Used to unregister the ticker. */
+		const FTSTicker::FDelegateHandle TickerDelegateHandle;
+
+		/** Queries stream and authority changes. */
+		FStreamAndAuthorityQueryService StreamAndAuthorityQueryService;
+
+		/** Ticks the sub-services when it is time to send a new request. */
+		void Tick();
+	};
+}
+
diff --git a/Engine/Plugins/Developer/Concert/ConcertApp/MultiUserClient/Source/MultiUserClient/Private/Replication/Util/Query/StreamAndAuthorityQueryService.cpp b/Engine/Plugins/Developer/Concert/ConcertApp/MultiUserClient/Source/MultiUserClient/Private/Replication/Util/Query/StreamAndAuthorityQueryService.cpp
new file mode 100644
--- /dev/null
+++ b/Engine/Plugins/Developer/Concert/ConcertApp/MultiUserClient/Source/MultiUserClient/Private/Replication/Util/Query/StreamAndAuthorityQueryService.cpp
@@ -0,0 +1,186 @@
+// Copyright Epic Games, Inc. All Rights Reserved.
+
+#include "StreamAndAuthorityQueryService.h"
+
+#include "IConcertSyncClient.h"
+#include "Replication/IConcertClientReplicationManager.h"
+
+#include "Templates/UnrealTemplate.h"
+
+namespace UE::MultiUserClient
+{
+	FStreamAndAuthorityQueryService::FStreamAndAuthorityQueryService(TWeakPtr<FToken> InToken, const IConcertSyncClient& InOwningClient)
+		: Token(MoveTemp(InToken))
+		, OwningClient(InOwningClient)
+	{}
+
+	FDelegateHandle FStreamAndAuthorityQueryService::RegisterStreamQuery(const FGuid& EndpointId, FStreamQueryDelegate Delegate)
+	{
+		FStreamQueryInfo& Info = StreamQueryInfos.FindOrAdd(EndpointId);
+		const FDelegateHandle DelegateHandle = Info.Delegate.AddLambda([Delegate = MoveTemp(Delegate)](const TArray<FConcertBaseStreamInfo>& Descriptions)
+		{
+			Delegate.Execute(Descriptions);
+		});
+		Info.Handles.Add(DelegateHandle);
+		return DelegateHandle;
+	}
+
+	FDelegateHandle FStreamAndAuthorityQueryService::RegisterAuthorityQuery(const FGuid& EndpointId, FAuthorityQueryDelegate Delegate)
+	{
+		FAuthorityQueryInfo& Info = AuthorityQueryInfos.FindOrAdd(EndpointId);
+		const FDelegateHandle DelegateHandle = Info.Delegate.AddLambda([Delegate = MoveTemp(Delegate)](const TArray<FConcertAuthorityClientInfo>& Infos)
+		{
+			Delegate.Execute(Infos);
+		});
+		Info.Handles.Add(DelegateHandle);
+		return DelegateHandle;
+	}
+
+	void FStreamAndAuthorityQueryService::UnregisterStreamQuery(const FDelegateHandle& Handle)
+	{
+		for (auto StreamIt = StreamQueryInfos.CreateIterator(); StreamIt; ++StreamIt)
+		{
+			FStreamQueryInfo& Info = StreamIt->Value;
+			if (Info.Handles.Contains(Handle))
+			{
+				Info.Handles.Remove(Handle);
+				Info.Delegate.Remove(Handle);
+				
+				// If bIsHandlingQueryResponse, then HandleQueryResponse is executing delegates and a delegate is unregistering.
+				// Removing now would cause a crash since Broadcast will read some bookkeeping memory after Broadcast.
+				// Memory is cleaned up later by HandleQueryResponse.
+				if (!bIsHandlingQueryResponse && Info.Handles.IsEmpty())
+				{
+					StreamIt.RemoveCurrent();
+				}
+				
+				return;
+			}
+		}
+	}
+
+	void FStreamAndAuthorityQueryService::UnregisterAuthorityQuery(const FDelegateHandle& Handle)
+	{
+		for (auto AuthorityIt = AuthorityQueryInfos.CreateIterator(); AuthorityIt; ++AuthorityIt)
+		{
+			FAuthorityQueryInfo& Info = AuthorityIt->Value;
+			if (Info.Handles.Contains(Handle))
+			{
+				Info.Handles.Remove(Handle);
+				Info.Delegate.Remove(Handle);
+				
+				// If bIsHandlingQueryResponse, then HandleQueryResponse is executing delegates and a delegate is unregistering.
+				// Removing now would cause a crash since Broadcast will read some bookkeeping memory after Broadcast.
+				// Memory is cleaned up later by HandleQueryResponse.
+				if (!bIsHandlingQueryResponse && Info.Handles.IsEmpty())
+				{
+					AuthorityIt.RemoveCurrent();
+				}
+				
+				return;
+			}
+		}
+	}
+
+	void FStreamAndAuthorityQueryService::SendQueryEvent()
+	{
+		IConcertClientReplicationManager* ReplicationManager = OwningClient.GetReplicationManager();
+		if (!ensure(ReplicationManager))
+		{
+			return;
+		}
+		
+		FConcertReplication_QueryReplicationInfo_Request Request;
+		BuildStreamRequest(Request);
+		BuildAuthorityRequest(Request);
+
+		if (!Request.ClientEndpointIds.IsEmpty())
+		{
+			ReplicationManager->QueryClientInfo({ Request })
+				.Next([this, WeakToken = Token](FConcertReplication_QueryReplicationInfo_Response&& Response)
+				{
+					const bool bCanProcessRequest = WeakToken.Pin().IsValid();
+					if (bCanProcessRequest)
+					{
+						HandleQueryResponse(Response);
+					}
+				});
+		}
+	}
+
+	void FStreamAndAuthorityQueryService::BuildStreamRequest(FConcertReplication_QueryReplicationInfo_Request& Request) const
+	{
+		if (StreamQueryInfos.IsEmpty())
+		{
+			Request.QueryFlags |= EConcertQueryClientStreamFlags::SkipStreamInfo;
+		}
+		else
+		{
+			for (const TPair<FGuid, FStreamQueryInfo>& Query : StreamQueryInfos)
+			{
+				Request.ClientEndpointIds.Add(Query.Key);
+			}
+		}
+	}
+
+	void FStreamAndAuthorityQueryService::BuildAuthorityRequest(FConcertReplication_QueryReplicationInfo_Request& Request) const
+	{
+		if (AuthorityQueryInfos.IsEmpty())
+		{
+			Request.QueryFlags |= EConcertQueryClientStreamFlags::SkipAuthority;
+		}
+		else
+		{
+			for (const TPair<FGuid, FAuthorityQueryInfo>& Query : AuthorityQueryInfos)
+			{
+				Request.ClientEndpointIds.Add(Query.Key);
+			}
+		}
+	}
+
+	void FStreamAndAuthorityQueryService::HandleQueryResponse(
+		const FConcertReplication_QueryReplicationInfo_Response& Response
+		)
+	{
+		TGuardValue<bool> Guard(bIsHandlingQueryResponse, true);
+		
+		for (const TPair<FGuid, FConcertQueriedClientInfo>& StreamQueryPair : Response.ClientInfo)
+		{
+			const FGuid StreamId = StreamQueryPair.Key;
+			
+			// Note that the delegates may have been unsubscribed since request was sent, which is why we need to Find() them
+			if (const FStreamQueryInfo* QueryInfo = StreamQueryInfos.Find(StreamId))
+			{
+				QueryInfo->Delegate.Broadcast(StreamQueryPair.Value.Streams);
+			}
+
+			// Note that the delegates may have been unsubscribed since request was sent, which is why we need to Find() them
+			if (const FAuthorityQueryInfo* QueryInfo = AuthorityQueryInfos.Find(StreamId))
+			{
+				QueryInfo->Delegate.Broadcast(StreamQueryPair.Value.Authority);
+			}
+		}
+
+		// Delegates may have removed themselves up above.
+		CompactDelegates();
+	}
+
+	void FStreamAndAuthorityQueryService::CompactDelegates()
+	{
+		for (auto StreamIt = StreamQueryInfos.CreateIterator(); StreamIt; ++StreamIt)
+		{
+			if (StreamIt->Value.Handles.IsEmpty())
+			{
+				StreamIt.RemoveCurrent();
+			}
+		}
+
+		for (auto Authority = AuthorityQueryInfos.CreateIterator(); Authority; ++Authority)
+		{
+			if (Authority->Value.Handles.IsEmpty())
+			{
+				Authority.RemoveCurrent();
+			}
+		}
+	}
+}
diff --git a/Engine/Plugins/Developer/Concert/ConcertApp/MultiUserClient/Source/MultiUserClient/Private/Replication/Util/Query/StreamAndAuthorityQueryService.h b/Engine/Plugins/Developer/Concert/ConcertApp/MultiUserClient/Source/MultiUserClient/Private/Replication/Util/Query/StreamAndAuthorityQueryService.h
new file mode 100644
--- /dev/null
+++ b/Engine/Plugins/Developer/Concert/ConcertApp/MultiUserClient/Source/MultiUserClient/Private/Replication/Util/Query/StreamAndAuthorityQueryService.h
@@ -0,0 +1,82 @@
+// Copyright Epic Games, Inc. All Rights Reserved.
+
+#pragma once
+
+#include "Replication/Messages/ClientQuery.h"
+#include "Replication/IToken.h"
+
+#include "Delegates/Delegate.h"
+#include "HAL/Platform.h"
+#include "Templates/SharedPointer.h"
+#include "Templates/UnrealTemplate.h"
+
+class IConcertSyncClient;
+
+namespace UE::MultiUserClient
+{
+	DECLARE_DELEGATE_OneParam(FStreamQueryDelegate, const TArray<FConcertBaseStreamInfo>&);
+	DECLARE_DELEGATE_OneParam(FAuthorityQueryDelegate, const TArray<FConcertAuthorityClientInfo>&);
+	
+	/** Sends regular FConcertReplication_QueryReplicationInfo_Request to endpoints and publishes the results. */
+	class FStreamAndAuthorityQueryService : public FNoncopyable
+	{
+		friend class FRegularQueryService;
+	public:
+		
+		/**
+		 * @param InToken Used to check whether this instance was destroyed after a response is received.
+		 * @param InOwningClient The client through which queries are made. The caller ensures it outlasts the constructed instance.
+		 */
+		FStreamAndAuthorityQueryService(TWeakPtr<FToken> InToken, const IConcertSyncClient& InOwningClient UE_LIFETIMEBOUND);
+
+		/** Registers a delegate to invoke for querying an endpoint about its registered streams. */
+		FDelegateHandle RegisterStreamQuery(const FGuid& EndpointId, FStreamQueryDelegate Delegate);
+		/** Registers a delegate to invoke for querying an endpoint about its authority. */
+		FDelegateHandle RegisterAuthorityQuery(const FGuid& EndpointId, FAuthorityQueryDelegate Delegate);
+
+		void UnregisterStreamQuery(const FDelegateHandle& Handle);
+		void UnregisterAuthorityQuery(const FDelegateHandle& Handle);
+		
+	private:
+
+		/** Used to check whether we were destroyed after a response is received. */
+		const TWeakPtr<FToken> Token;
+		
+		/**
+		 * Used to send queries.
+		 * 
+		 * This FRegularQueryService's owner is supposed to make sure this FRegularQueryService is destroyed
+		 * when the client shuts down.
+		 */
+		const IConcertSyncClient& OwningClient;
+
+		// We use a multicast delegate here because it handles unsubscribing while being Broadcast
+		DECLARE_MULTICAST_DELEGATE_OneParam(FMulticastStreamQueryDelegate, const TArray<FConcertBaseStreamInfo>&);
+		DECLARE_MULTICAST_DELEGATE_OneParam(FMulticastAuthorityQueryDelegate, const TArray<FConcertAuthorityClientInfo>&);
+		struct FStreamQueryInfo
+		{
+			TSet<FDelegateHandle> Handles;
+			FMulticastStreamQueryDelegate Delegate;
+		};
+		struct FAuthorityQueryInfo
+		{
+			TSet<FDelegateHandle> Handles;
+			FMulticastAuthorityQueryDelegate Delegate;
+		};
+		TMap<FGuid, FStreamQueryInfo> StreamQueryInfos;
+		TMap<FGuid, FAuthorityQueryInfo> AuthorityQueryInfos;
+
+		bool bIsHandlingQueryResponse = false;
+		
+		/** Queries the server for the client's current state. */
+		void SendQueryEvent();
+		
+		void BuildStreamRequest(FConcertReplication_QueryReplicationInfo_Request& Request) const;
+		void BuildAuthorityRequest(FConcertReplication_QueryReplicationInfo_Request& Request) const;
+		
+		/** Handles the server's response. */
+		void HandleQueryResponse(const FConcertReplication_QueryReplicationInfo_Response& Response);
+		void CompactDelegates();
+	};
+}
+
diff --git a/Engine/Plugins/Developer/Concert/ConcertApp/MultiUserClient/Source/MultiUserClient/Private/Replication/Util/RegularQueryService.cpp b/Engine/Plugins/Developer/Concert/ConcertApp/MultiUserClient/Source/MultiUserClient/Private/Replication/Util/RegularQueryService.cpp
deleted file mode 100644
diff --git a/Engine/Plugins/Developer/Concert/ConcertApp/MultiUserClient/Source/MultiUserClient/Private/Replication/Util/RegularQueryService.h b/Engine/Plugins/Developer/Concert/ConcertApp/MultiUserClient/Source/MultiUserClient/Private/Replication/Util/RegularQueryService.h
deleted file mode 100644
