Integrate replication muting into MultiUserClient.
- Replaces replication checkbox column with pause & unpause column

diff --git a/Engine/Plugins/Developer/Concert/ConcertApp/MultiUserClient/Source/MultiUserClient/Private/MultiUserReplicationStyle.cpp b/Engine/Plugins/Developer/Concert/ConcertApp/MultiUserClient/Source/MultiUserClient/Private/MultiUserReplicationStyle.cpp
--- a/Engine/Plugins/Developer/Concert/ConcertApp/MultiUserClient/Source/MultiUserClient/Private/MultiUserReplicationStyle.cpp
+++ b/Engine/Plugins/Developer/Concert/ConcertApp/MultiUserClient/Source/MultiUserClient/Private/MultiUserReplicationStyle.cpp
@@ -9,14 +9,33 @@
 
 #include "Styling/SlateTypes.h"
 #include "Styling/SlateStyleRegistry.h"
-#include "Styling/StarshipCoreStyle.h"
 #include "Styling/SlateStyle.h"
+#include "Styling/StyleColors.h"
+
+#define IMAGE_PLUGIN_BRUSH( RelativePath, ... ) FSlateImageBrush( FMultiUserReplicationStyle::InContent( RelativePath, ".png" ), __VA_ARGS__ )
 
 namespace UE::MultiUserClient
 {
+	struct FButtonColor
+	{
+		FLinearColor Normal;
+		FLinearColor Hovered;
+		FLinearColor Pressed;
+
+		FButtonColor(const FLinearColor& Color)
+		{
+			Normal = Color * 0.8f;
+			Normal.A = Color.A;
+			Hovered = Color * 1.0f;
+			Hovered.A = Color.A;
+			Pressed = Color * 0.6f;
+			Pressed.A = Color.A;
+		}
+	};
+	
 	FString FMultiUserReplicationStyle::InContent(const FString& RelativePath, const ANSICHAR* Extension)
 	{
-		static FString ContentDir = IPluginManager::Get().FindPlugin(TEXT("ConcertSharedSlate"))->GetContentDir();
+		static FString ContentDir = IPluginManager::Get().FindPlugin(TEXT("MultiUserClient"))->GetContentDir();
 		return (ContentDir / RelativePath) + Extension;
 	}
 
@@ -38,14 +57,32 @@ namespace UE::MultiUserClient
 		StyleSet = MakeShared<FSlateStyleSet>(GetStyleSetName());
 		StyleSet->SetContentRoot(FPaths::EngineContentDir() / TEXT("Editor/Slate"));
 		StyleSet->SetCoreContentRoot(FPaths::EngineContentDir() / TEXT("Slate"));
+		
+		const FVector2D Icon16x16(16.0f, 16.0f); 
 
 		// Column widths - see FConcertFrontendStyle::Initialize() also ("Default Concert Replication Columns Widths")
-		StyleSet->Set("AllClients.Object.ReplicationToggle", 45.f);
+		StyleSet->Set("AllClients.Object.MuteToggle", 25.f);
 		StyleSet->Set("AllClients.Object.OwnerSize", 200.f);
 		StyleSet->Set("AllClients.Property.OwnerSize", 200.f);
 
 		// Timing
 		StyleSet->Set("AllClients.Reassignment.DisplayThrobberAfterSeconds", 0.2f);
+
+
+		FSlateImageBrush* PlayBrush = new IMAGE_PLUGIN_BRUSH("generic_play_16x", Icon16x16);
+		FSlateImageBrush* PauseBrush = new IMAGE_PLUGIN_BRUSH("generic_pause_16x", Icon16x16);
+		StyleSet->Set("MultiUser.Icons.Play", PlayBrush);
+		StyleSet->Set("MultiUser.Icons.Pause", PauseBrush);
+
+		StyleSet->Set("AllClients.MuteToggle.Style", FCheckBoxStyle()		
+			.SetCheckBoxType(ESlateCheckBoxType::CheckBox)
+			.SetUncheckedImage(IMAGE_PLUGIN_BRUSH("generic_pause_16x", Icon16x16, FStyleColors::Foreground))
+			.SetUncheckedHoveredImage(IMAGE_PLUGIN_BRUSH("generic_pause_16x", Icon16x16, FStyleColors::ForegroundHover))  
+			.SetUncheckedPressedImage(IMAGE_PLUGIN_BRUSH("generic_pause_16x", Icon16x16, FStyleColors::ForegroundHover)) 
+			.SetCheckedImage(IMAGE_PLUGIN_BRUSH("generic_play_16x", Icon16x16, FStyleColors::Foreground)) 
+			.SetCheckedHoveredImage(IMAGE_PLUGIN_BRUSH("generic_play_16x", Icon16x16, FStyleColors::ForegroundHover)) 
+			.SetCheckedPressedImage(IMAGE_PLUGIN_BRUSH("generic_play_16x", Icon16x16, FStyleColors::ForegroundHover)) 
+		);
 		
 		FSlateStyleRegistry::RegisterSlateStyle(*StyleSet.Get());
 	};
@@ -66,4 +103,4 @@ namespace UE::MultiUserClient
 	}
 }
 
-
+#undef IMAGE_PLUGIN_BRUSH
diff --git a/Engine/Plugins/Developer/Concert/ConcertApp/MultiUserClient/Source/MultiUserClient/Private/Replication/Client/ReplicationClientManager.cpp b/Engine/Plugins/Developer/Concert/ConcertApp/MultiUserClient/Source/MultiUserClient/Private/Replication/Client/ReplicationClientManager.cpp
--- a/Engine/Plugins/Developer/Concert/ConcertApp/MultiUserClient/Source/MultiUserClient/Private/Replication/Client/ReplicationClientManager.cpp
+++ b/Engine/Plugins/Developer/Concert/ConcertApp/MultiUserClient/Source/MultiUserClient/Private/Replication/Client/ReplicationClientManager.cpp
@@ -7,6 +7,7 @@
 #include "RemoteReplicationClient.h"
 #include "Assets/MultiUserReplicationSessionPreset.h"
 #include "Replication/Stream/StreamSynchronizer_LocalClient.h"
+#include "Replication/Util/Query/StreamAndAuthorityQueryService.h"
 
 #include "UObject/Package.h"
 #include "UObject/UObjectGlobals.h"
@@ -16,20 +17,20 @@ namespace UE::MultiUserClient
 	FReplicationClientManager::FReplicationClientManager(
 		const TSharedRef<IConcertSyncClient>& InClient,
 		const TSharedRef<IConcertClientSession>& InSession,
-		FReplicationDiscoveryContainer& InRegisteredExtenders
+		FReplicationDiscoveryContainer& InRegisteredExtenders,
+		FStreamAndAuthorityQueryService& InQueryService
 		)
 		: SessionContent(NewObject<UMultiUserReplicationSessionPreset>(GetTransientPackage(), NAME_None, RF_Transient))
 		, ConcertClient(InClient)
 		, Session(InSession)
 		, RegisteredExtenders(InRegisteredExtenders)
-		, QueryService(*ConcertClient)
+		, QueryService(InQueryService)
 		, AuthorityCache(*this)
 		, LocalClient([this, InClient]()
 		{
 			UMultiUserReplicationClientPreset* ClientPreset = SessionContent->AddClient();
 			return FLocalReplicationClient(RegisteredExtenders, AuthorityCache, *ClientPreset, MakeUnique<FStreamSynchronizer_LocalClient>(InClient, ClientPreset->Stream->StreamId), InClient);
 		}())
-		, SubmissionNotifier(*this)
 		, ReassignmentLogic(*this)
 	{
 		AuthorityCache.RegisterEvents();
@@ -169,7 +170,7 @@ namespace UE::MultiUserClient
 			ConcertClient->GetConcertClient(),
 			AuthorityCache,
 			*SessionContent->AddClient(),
-			QueryService.GetStreamAndAuthorityQueryService()
+			QueryService
 			);
 		FRemoteReplicationClient& RemoteClient = *RemoteClientPtr;
 		RemoteClients.Emplace(
diff --git a/Engine/Plugins/Developer/Concert/ConcertApp/MultiUserClient/Source/MultiUserClient/Private/Replication/Client/ReplicationClientManager.h b/Engine/Plugins/Developer/Concert/ConcertApp/MultiUserClient/Source/MultiUserClient/Private/Replication/Client/ReplicationClientManager.h
--- a/Engine/Plugins/Developer/Concert/ConcertApp/MultiUserClient/Source/MultiUserClient/Private/Replication/Client/ReplicationClientManager.h
+++ b/Engine/Plugins/Developer/Concert/ConcertApp/MultiUserClient/Source/MultiUserClient/Private/Replication/Client/ReplicationClientManager.h
@@ -9,7 +9,6 @@
 #include "Replication/Submission/MultiEdit/ReassignObjectPropertiesLogic.h"
 #include "Replication/Submission/Notification/SubmissionNotifier.h"
 #include "Replication/Util/GlobalAuthorityCache.h"
-#include "Replication/Util/Query/RegularQueryService.h"
 
 #include "UObject/GCObject.h"
 #include "Templates/UnrealTemplate.h"
@@ -26,6 +25,7 @@ namespace UE::MultiUserClient
 	class FRemoteReplicationClient;
 	class FReplicationDiscoveryContainer;
 	class FReplicationClient;
+	class FStreamAndAuthorityQueryService;
 	class FStreamChangeTracker;
 
 	/**
@@ -44,11 +44,13 @@ namespace UE::MultiUserClient
 		 * @param InClient The local client. The owning FMultiUserReplicationManager ensures it outlives the constructed instance.
 		 * @param InSession The session to observe. The owning FMultiUserReplicationManager ensures it outlives the constructed instance.
 		 * @param InRegisteredExtenders Used for auto-discovering properties added to this client's stream. Passed to the clients. The owning FMultiUserReplicationManager ensures it outlives the constructed instance.
+		 * @param InQueryService Used for querying info about clients' streams. Passed to the clients. The owning FMultiUserReplicationManager ensures it outlives the constructed instance.
 		 */
 		FReplicationClientManager(
 			const TSharedRef<IConcertSyncClient>& InClient,
 			const TSharedRef<IConcertClientSession>& InSession,
-			FReplicationDiscoveryContainer& InRegisteredExtenders
+			FReplicationDiscoveryContainer& InRegisteredExtenders UE_LIFETIMEBOUND,
+			FStreamAndAuthorityQueryService& InQueryService UE_LIFETIMEBOUND
 			);
 		virtual ~FReplicationClientManager() override;
 
@@ -132,7 +134,7 @@ namespace UE::MultiUserClient
 		 * Sends FConcertReplication_QueryReplicationInfo_Request in regular intervals.
 		 * Shared by all remote clients so all requests are bundled reducing the number of network requests. 
 		 */
-		FRegularQueryService QueryService;
+		FStreamAndAuthorityQueryService& QueryService;
 		/** Keeps a cache of object to owning clients. */
 		FGlobalAuthorityCache AuthorityCache;
 		
@@ -150,12 +152,6 @@ namespace UE::MultiUserClient
 		FRemoteClientDelegate OnPostRemoteClientAddedDelegate;
 		/** Called just before a remote client is about to be removed from RemoteClients. */
 		FRemoteClientDelegate OnPreRemoteClientRemovedDelegate; 
-		
-		/**
-		 * Manages SNotificationItems when submission to the server fails.
-		 * TODO UE-200925: This should be moved to the client so we can customize messages depending on whether it is a remote or local client
-		 */
-		FSubmissionNotifier SubmissionNotifier;
 
 		/** Used for transferring ownership from multiple clients to one. Used by multi client view. */
 		FReassignObjectPropertiesLogic ReassignmentLogic;
diff --git a/Engine/Plugins/Developer/Concert/ConcertApp/MultiUserClient/Source/MultiUserClient/Private/Replication/MultiUserReplicationManager.cpp b/Engine/Plugins/Developer/Concert/ConcertApp/MultiUserClient/Source/MultiUserClient/Private/Replication/MultiUserReplicationManager.cpp
--- a/Engine/Plugins/Developer/Concert/ConcertApp/MultiUserClient/Source/MultiUserClient/Private/Replication/MultiUserReplicationManager.cpp
+++ b/Engine/Plugins/Developer/Concert/ConcertApp/MultiUserClient/Source/MultiUserClient/Private/Replication/MultiUserReplicationManager.cpp
@@ -196,7 +196,11 @@ namespace UE::MultiUserClient
 	}
 
 	FMultiUserReplicationManager::FConnectedState::FConnectedState(TSharedRef<IConcertSyncClient> InClient, FReplicationDiscoveryContainer& InDiscoveryContainer)
-		: ClientManager(InClient, InClient->GetConcertClient()->GetCurrentSession().ToSharedRef(), InDiscoveryContainer)
+		: Client(InClient)
+		, QueryService(*InClient)
+		, ClientManager(InClient, InClient->GetConcertClient()->GetCurrentSession().ToSharedRef(), InDiscoveryContainer, QueryService.GetStreamAndAuthorityQueryService())
+		, MuteManager(*InClient, QueryService.GetMuteStateQueryService(), ClientManager.GetAuthorityCache())
 		, ChangeLevelHandler(ClientManager.GetLocalClient().GetClientEditModel().Get())
+		, UserNotifier(ClientManager, MuteManager)
 	{}
 }
diff --git a/Engine/Plugins/Developer/Concert/ConcertApp/MultiUserClient/Source/MultiUserClient/Private/Replication/MultiUserReplicationManager.h b/Engine/Plugins/Developer/Concert/ConcertApp/MultiUserClient/Source/MultiUserClient/Private/Replication/MultiUserReplicationManager.h
--- a/Engine/Plugins/Developer/Concert/ConcertApp/MultiUserClient/Source/MultiUserClient/Private/Replication/MultiUserReplicationManager.h
+++ b/Engine/Plugins/Developer/Concert/ConcertApp/MultiUserClient/Source/MultiUserClient/Private/Replication/MultiUserReplicationManager.h
@@ -7,8 +7,11 @@
 #include "Replication/IMultiUserReplication.h"
 #include "Replication/Stream/Discovery/ReplicationDiscoveryContainer.h"
 #include "UnrealEditor/ChangeLevelHandler.h"
+#include "Util/Query/RegularQueryService.h"
 
 #include "Misc/Optional.h"
+#include "Muting/MuteStateManager.h"
+#include "Notification/ReplicationUserNotifier.h"
 #include "Templates/SharedPointer.h"
 #include "Templates/UnrealTemplate.h"
 
@@ -16,15 +19,8 @@ class IConcertClientSession;
 class IConcertSyncClient;
 enum class EConcertConnectionStatus : uint8;
 
-namespace UE::ConcertSyncClient::Replication
-{
-	struct FJoinReplicatedSessionResult;
-}
-
-namespace UE::ConcertSharedSlate
-{
-	class IEditableReplicationStreamModel;
-}
+namespace UE::ConcertSharedSlate { class IEditableReplicationStreamModel; }
+namespace UE::ConcertSyncClient::Replication { struct FJoinReplicatedSessionResult; }
 
 namespace UE::MultiUserClient
 {
@@ -65,6 +61,10 @@ namespace UE::MultiUserClient
 		/** @note You're not supposed to keep any reference to the ClientManager since it can become invalid depending on connection state. */
 		FReplicationClientManager* GetClientManager() { return ConnectedState ? &ConnectedState->ClientManager : nullptr; }
 		const FReplicationClientManager* GetClientManager() const { return ConnectedState ? &ConnectedState->ClientManager : nullptr; }
+		
+		/** @note You're not supposed to keep any reference to the MuteManager since it can become invalid depending on connection state. */
+		FMuteStateManager* GetMuteManager() { return ConnectedState ? &ConnectedState->MuteManager : nullptr; }
+		const FMuteStateManager* GetMuteManager() const { return ConnectedState ? &ConnectedState->MuteManager : nullptr; }
 
 		/** Called when the connection to the replication system changes. */
 		DECLARE_MULTICAST_DELEGATE_OneParam(FOnReplicationConnectionStateChanged, EMultiUserReplicationConnectionState /*NewState*/);
@@ -92,8 +92,17 @@ namespace UE::MultiUserClient
 
 		struct FConnectedState
 		{
-			// The order of the below members matters so they are destroyed in the right order!
+			// The order of the below members matters, so they are destroyed in the right order!
 			// Rule: Lower systems can only reference higher systems (reminder: C++ destroys in reverse declaration order).
+
+			/** Keep the client alive because it is referenced by the lower systems. */
+			TSharedRef<IConcertSyncClient> Client;
+			
+			/**
+			 * Regularly queries server state (only if a system is subscribed to it), such as remote client streams, authority, and global mute state.
+			 * Shared by subsystems, such as remote clients. This ensures requests are bundled, reducing the number of network requests. 
+			 */
+			FRegularQueryService QueryService;
 			
 			/**
 			 * Creates UMultiUserReplicationSessionPreset which is displayed by UI.
@@ -103,8 +112,14 @@ namespace UE::MultiUserClient
 			 */
 			FReplicationClientManager ClientManager;
 
+			/** Interacts with the mute global server mute system. */
+			FMuteStateManager MuteManager;
+
 			/** Clears local client's registered objects when leaving map. */
 			FChangeLevelHandler ChangeLevelHandler;
+
+			/** This system notifies users when requests go wrong */
+			FReplicationUserNotifier UserNotifier;
 			
 			FConnectedState(TSharedRef<IConcertSyncClient> InClient, FReplicationDiscoveryContainer& InDiscoveryContainer);
 		};
diff --git a/Engine/Plugins/Developer/Concert/ConcertApp/MultiUserClient/Source/MultiUserClient/Private/Replication/Muting/MuteChangeTracker.cpp b/Engine/Plugins/Developer/Concert/ConcertApp/MultiUserClient/Source/MultiUserClient/Private/Replication/Muting/MuteChangeTracker.cpp
new file mode 100644
--- /dev/null
+++ b/Engine/Plugins/Developer/Concert/ConcertApp/MultiUserClient/Source/MultiUserClient/Private/Replication/Muting/MuteChangeTracker.cpp
@@ -0,0 +1,73 @@
+﻿// Copyright Epic Games, Inc. All Rights Reserved.
+
+#include "MuteChangeTracker.h"
+
+#include "Replication/MultiUserReplicationManager.h"
+
+#include "UObject/SoftObjectPath.h"
+
+namespace UE::MultiUserClient
+{
+	FMuteChangeTracker::FMuteChangeTracker(FMuteStateSynchronizer& InMuteStateSynchronizer, const FGlobalAuthorityCache& InAuthorityCache)
+		: MuteStateSynchronizer(InMuteStateSynchronizer)
+		, AuthorityCache(InAuthorityCache)
+	{
+		MuteStateSynchronizer.OnMuteStateChanged().AddRaw(this, &FMuteChangeTracker::RefreshChanges);
+	}
+
+	FMuteChangeTracker::~FMuteChangeTracker()
+	{
+		MuteStateSynchronizer.OnMuteStateChanged().RemoveAll(this);
+	}
+
+	void FMuteChangeTracker::ToggleMuteState(const FSoftObjectPath& ObjectPath)
+	{
+		if (CanChangeMuteState(ObjectPath))
+		{
+			const bool bIsCurrentlyMuted = MuteStateSynchronizer.IsMuted(ObjectPath);
+			LocalChanges.Add(ObjectPath, !bIsCurrentlyMuted);
+			OnLocalMuteStateOverridenDelegate.Broadcast();
+		}
+	}
+
+	bool FMuteChangeTracker::CanChangeMuteState(const FSoftObjectPath& ObjectPath) const
+	{
+		return AuthorityCache.IsObjectOrChildReferenced(ObjectPath);
+	}
+
+	bool FMuteChangeTracker::IsMuted(const FSoftObjectPath& ObjectPath) const
+	{
+		return MuteStateSynchronizer.IsMuted(ObjectPath);
+	}
+
+	void FMuteChangeTracker::RefreshChanges()
+	{
+		for (auto It = LocalChanges.CreateIterator(); It; ++It)
+		{
+			const FSoftObjectPath& ObjectPath = It->Key;
+			const bool bWillBeMuted = It->Value;
+		
+			if (!CanChangeMuteState(ObjectPath) || bWillBeMuted == MuteStateSynchronizer.IsMuted(ObjectPath))
+			{
+				It.RemoveCurrent();
+			}
+		}
+	}
+
+	FConcertReplication_ChangeMuteState_Request FMuteChangeTracker::BuildChangeRequest() const
+	{
+		FConcertReplication_ChangeMuteState_Request Request;
+		for (const TPair<FSoftObjectPath, bool>& Change : LocalChanges)
+		{
+			// May have changed since local mute change was made
+			if (!CanChangeMuteState(Change.Key))
+			{
+				continue;
+			}
+
+			TMap<FSoftObjectPath, FConcertReplication_ObjectMuteSetting>& MuteMap = Change.Value ? Request.ObjectsToMute : Request.ObjectsToUnmute;
+			MuteMap.Add(Change.Key, FConcertReplication_ObjectMuteSetting{ EConcertReplicationMuteOption::ObjectAndSubobjects });
+		}
+		return Request;
+	}
+}
diff --git a/Engine/Plugins/Developer/Concert/ConcertApp/MultiUserClient/Source/MultiUserClient/Private/Replication/Muting/MuteChangeTracker.h b/Engine/Plugins/Developer/Concert/ConcertApp/MultiUserClient/Source/MultiUserClient/Private/Replication/Muting/MuteChangeTracker.h
new file mode 100644
--- /dev/null
+++ b/Engine/Plugins/Developer/Concert/ConcertApp/MultiUserClient/Source/MultiUserClient/Private/Replication/Muting/MuteChangeTracker.h
@@ -0,0 +1,60 @@
+﻿// Copyright Epic Games, Inc. All Rights Reserved.
+
+#pragma once
+
+#include "Containers/Map.h"
+#include "Delegates/Delegate.h"
+#include "Templates/UnrealTemplate.h"
+#include "UObject/SoftObjectPath.h"
+
+struct FConcertReplication_ChangeMuteState_Request;
+
+namespace UE::MultiUserClient
+{
+	class FGlobalAuthorityCache;
+	class FMuteStateSynchronizer;
+
+	DECLARE_MULTICAST_DELEGATE(FOnLocalMuteStateOverriden);
+
+	/** Keeps track of changes to the mute state the local application has made but not yet submitted. */
+	class FMuteChangeTracker : public FNoncopyable
+	{
+	public:
+		FMuteChangeTracker(FMuteStateSynchronizer& InMuteStateSynchronizer UE_LIFETIMEBOUND, const FGlobalAuthorityCache& InAuthorityCache UE_LIFETIMEBOUND);
+		~FMuteChangeTracker();
+		
+		/** @return Whether the mute state of this object can be changed. True if the object is at least in one client stream. */
+		bool CanChangeMuteState(const FSoftObjectPath& ObjectPath) const;
+		/** @return Returns whether ObjectPath according to the local, unsubmitted changes or returns the state on the server if no local changes have been made. */
+		bool IsMuted(const FSoftObjectPath& ObjectPath) const;
+		
+		/** Toggles the mute state of ObjectPath if allowed. The mute setting will apply to all subobjects as well (unless they are overriden, too). */
+		void ToggleMuteState(const FSoftObjectPath& ObjectPath);
+
+		/** Clears all local changes, so they do not show up the next time BuildChangeRequest is called. */
+		void ClearLocalMuteOverrides() { LocalChanges.Reset(); }
+		/** Compares the local changes against the server state and kicks out any (now) invalid changes. */
+		void RefreshChanges();
+
+		/** @return A request contains all the pending local changes. */
+		FConcertReplication_ChangeMuteState_Request BuildChangeRequest() const;
+
+		/** Broadcasts when there a local change has been made that can be submitted. */
+		FOnLocalMuteStateOverriden& OnLocalMuteStateOverriden() { return OnLocalMuteStateOverridenDelegate; }
+
+	private:
+
+		/** Knows of the server state. Used to diff state. */
+		FMuteStateSynchronizer& MuteStateSynchronizer;
+
+		/** Used to check whether an object is referenced in any stream. Needed so we don't build requests containing objects the server does not know about. */
+		const FGlobalAuthorityCache& AuthorityCache;
+		
+		/** Maps objects to whether they should be muted. */
+		TMap<FSoftObjectPath, bool> LocalChanges;
+
+		/** Broadcasts when there a local change has been made that can be submitted. */
+		FOnLocalMuteStateOverriden OnLocalMuteStateOverridenDelegate;
+	};
+}
+
diff --git a/Engine/Plugins/Developer/Concert/ConcertApp/MultiUserClient/Source/MultiUserClient/Private/Replication/Muting/MuteStateManager.cpp b/Engine/Plugins/Developer/Concert/ConcertApp/MultiUserClient/Source/MultiUserClient/Private/Replication/Muting/MuteStateManager.cpp
new file mode 100644
--- /dev/null
+++ b/Engine/Plugins/Developer/Concert/ConcertApp/MultiUserClient/Source/MultiUserClient/Private/Replication/Muting/MuteStateManager.cpp
@@ -0,0 +1,78 @@
+﻿// Copyright Epic Games, Inc. All Rights Reserved.
+
+#include "MuteStateManager.h"
+
+#include "IConcertSyncClient.h"
+#include "Misc/CoreDelegates.h"
+#include "Replication/MultiUserReplicationManager.h"
+
+namespace UE::MultiUserClient
+{
+	FMuteStateManager::FMuteStateManager(const IConcertSyncClient& InClient, FMuteStateQueryService& InMuteQueryService, const FGlobalAuthorityCache& InAuthorityCache)
+		: Client(InClient)
+		, MuteQueryService(InMuteQueryService)
+		, MuteStateSynchronizer(InMuteQueryService)
+		, ChangeTracker(MuteStateSynchronizer, InAuthorityCache)
+	{
+		ChangeTracker.OnLocalMuteStateOverriden().AddLambda([this]()
+		{
+			if (!FCoreDelegates::OnEndFrame.IsBoundToObject(this))
+			{
+				FCoreDelegates::OnEndFrame.AddRaw(this, &FMuteStateManager::OnEndOfFrame);
+			}
+		});
+	}
+
+	FMuteStateManager::~FMuteStateManager()
+	{
+		FCoreDelegates::OnEndFrame.RemoveAll(this);
+	}
+
+	void FMuteStateManager::OnEndOfFrame()
+	{
+		if (!bIsMuteChangeInProgress)
+		{
+			FCoreDelegates::OnEndFrame.RemoveAll(this);
+			SendChangeRequest();
+		}
+	}
+
+	void FMuteStateManager::SendChangeRequest()
+	{
+		IConcertClientReplicationManager* ReplicationManager = Client.GetReplicationManager();
+		const FConcertReplication_ChangeMuteState_Request Request = ChangeTracker.BuildChangeRequest();
+		if (Request.IsEmpty() || !ensure(ReplicationManager))
+		{
+			return;
+		}
+
+		bIsMuteChangeInProgress = true;
+		ReplicationManager->ChangeMuteState(Request)
+			.Next([this, Request, WeakToken = Token.ToWeakPtr()](FConcertReplication_ChangeMuteState_Response&& Response)
+			{
+				if (!WeakToken.IsValid())
+				{
+					return;
+				}
+					
+				bIsMuteChangeInProgress = false;
+				if (Response.IsSuccess())
+				{
+					// This will implicitly refresh ChangeTracker's local state because FMuteStateSynchronizer::OnMuteStateChanged is broadcast.
+					MuteStateSynchronizer.UpdateStateFromSuccessfulChange(Request);
+				}
+				else
+				{
+					// If the response fails, keep the local changes. The state will be cleansed next time we attempt to submit or the FMuteStateQueryService receives the new server state.
+
+					// This can e.g. notify the user of the failure
+					OnMuteRequestFailureDelegate.Broadcast(Request, Response);
+				}
+			});
+
+		// This is a hacky way of instantly getting the UI to refresh.
+		// It's a hack because it generates another network request.
+		// Instead, the local application should predict the new mute state locally.
+		MuteQueryService.RequestInstantRefresh();
+	}
+}
diff --git a/Engine/Plugins/Developer/Concert/ConcertApp/MultiUserClient/Source/MultiUserClient/Private/Replication/Muting/MuteStateManager.h b/Engine/Plugins/Developer/Concert/ConcertApp/MultiUserClient/Source/MultiUserClient/Private/Replication/Muting/MuteStateManager.h
new file mode 100644
--- /dev/null
+++ b/Engine/Plugins/Developer/Concert/ConcertApp/MultiUserClient/Source/MultiUserClient/Private/Replication/Muting/MuteStateManager.h
@@ -0,0 +1,62 @@
+﻿// Copyright Epic Games, Inc. All Rights Reserved.
+
+#pragma once
+
+#include "MuteChangeTracker.h"
+#include "MuteStateSynchronizer.h"
+#include "Replication/IToken.h"
+
+#include "Templates/UnrealTemplate.h"
+
+class IConcertSyncClient;
+struct FConcertReplication_ChangeMuteState_Request;
+struct FConcertReplication_ChangeMuteState_Response;
+
+namespace UE::MultiUserClient
+{
+	class FGlobalAuthorityCache;
+	class FMuteStateQueryService;
+
+	/** Manages all interaction with the application instance to the server regarding muting. */
+	class FMuteStateManager : public FNoncopyable
+	{
+	public:
+
+		FMuteStateManager(const IConcertSyncClient& InClient UE_LIFETIMEBOUND, FMuteStateQueryService& InMuteQueryService UE_LIFETIMEBOUND, const FGlobalAuthorityCache& InAuthorityCache UE_LIFETIMEBOUND);
+		~FMuteStateManager();
+
+		/** @return Object with which the objects' mute state can be changed and queried. */
+		FMuteChangeTracker& GetChangeTracker() { return ChangeTracker; }
+		
+		/** Broadcasts when a mute request fails. */
+		DECLARE_MULTICAST_DELEGATE_TwoParams(FOnMuteRequestFailure, const FConcertReplication_ChangeMuteState_Request&, const FConcertReplication_ChangeMuteState_Response&);
+		FOnMuteRequestFailure& OnMuteRequestFailure() { return OnMuteRequestFailureDelegate; }
+
+	private:
+
+		/** When the mute request is answered, we use this to detect whether we've been destroyed. */
+		const TSharedRef<FToken> Token = FToken::Make();
+
+		/** Used to send submit mute changes to the server. */
+		const IConcertSyncClient& Client;
+		/** We'll request an instant refresh of the server's mute state after submitting a request.. */
+		FMuteStateQueryService& MuteQueryService; 
+		
+		/** Knows of the current mute state on the server. */
+		FMuteStateSynchronizer MuteStateSynchronizer;
+		/** Tracks locally made changes that still need to be submitted to the server. */
+		FMuteChangeTracker ChangeTracker;
+
+		/** Broadcasts when a mute request fails. */
+		FOnMuteRequestFailure OnMuteRequestFailureDelegate;
+
+		/** Whether we're currently waiting for a response to a mute change. */
+		bool bIsMuteChangeInProgress = false;
+
+		/** When a change is made to ChangeTracker, we subscribe FCoreDelegates::OnEndFrame to this function. */
+		void OnEndOfFrame();
+		/** Checks the changes pending in ChangeTracker and sends them to the server, if any. */
+		void SendChangeRequest();
+	};
+}
+
diff --git a/Engine/Plugins/Developer/Concert/ConcertApp/MultiUserClient/Source/MultiUserClient/Private/Replication/Muting/MuteStateSynchronizer.cpp b/Engine/Plugins/Developer/Concert/ConcertApp/MultiUserClient/Source/MultiUserClient/Private/Replication/Muting/MuteStateSynchronizer.cpp
new file mode 100644
--- /dev/null
+++ b/Engine/Plugins/Developer/Concert/ConcertApp/MultiUserClient/Source/MultiUserClient/Private/Replication/Muting/MuteStateSynchronizer.cpp
@@ -0,0 +1,51 @@
+﻿// Copyright Epic Games, Inc. All Rights Reserved.
+
+#include "MuteStateSynchronizer.h"
+
+#include "Replication/MultiUserReplicationManager.h"
+
+#include "UObject/SoftObjectPath.h"
+
+namespace UE::MultiUserClient
+{
+	FMuteStateSynchronizer::FMuteStateSynchronizer(FMuteStateQueryService& InMuteQueryService)
+		: MuteQueryService(InMuteQueryService)
+	{
+		MuteQueryService.OnMuteStateQueried().AddRaw(this, &FMuteStateSynchronizer::UpdateFromServerState);
+	}
+
+	FMuteStateSynchronizer::~FMuteStateSynchronizer()
+	{
+		MuteQueryService.OnMuteStateQueried().RemoveAll(this);
+	}
+
+	void FMuteStateSynchronizer::UpdateStateFromSuccessfulChange(const FConcertReplication_ChangeMuteState_Request& Request)
+	{
+		if (!ensure(!Request.IsEmpty()))
+		{
+			return;
+		}
+
+		for (const TPair<FSoftObjectPath, FConcertReplication_ObjectMuteSetting>& MutedObject : Request.ObjectsToMute)
+		{
+			MutedObjects.Add(MutedObject.Key);
+		}
+		for (const TPair<FSoftObjectPath, FConcertReplication_ObjectMuteSetting>& UnmutedObject : Request.ObjectsToUnmute)
+		{
+			MutedObjects.Remove(UnmutedObject.Key);
+		}
+		
+		OnMuteStateChangedDelegate.Broadcast();
+	}
+
+	void FMuteStateSynchronizer::UpdateFromServerState(const FConcertReplication_QueryMuteState_Response& NewMuteState)
+	{
+		const int32 NumExpectedElements = NewMuteState.ExplicitlyMutedObjects.Num() + NewMuteState.ImplicitlyMutedObjects.Num();
+		MutedObjects.Empty(NumExpectedElements);
+
+		Algo::Transform(NewMuteState.ExplicitlyMutedObjects, MutedObjects, [](const TPair<FSoftObjectPath, FConcertReplication_ObjectMuteSetting>& Pair){ return Pair.Key; });
+		MutedObjects.Append(NewMuteState.ImplicitlyMutedObjects);
+
+		OnMuteStateChangedDelegate.Broadcast();
+	}
+}
diff --git a/Engine/Plugins/Developer/Concert/ConcertApp/MultiUserClient/Source/MultiUserClient/Private/Replication/Muting/MuteStateSynchronizer.h b/Engine/Plugins/Developer/Concert/ConcertApp/MultiUserClient/Source/MultiUserClient/Private/Replication/Muting/MuteStateSynchronizer.h
new file mode 100644
--- /dev/null
+++ b/Engine/Plugins/Developer/Concert/ConcertApp/MultiUserClient/Source/MultiUserClient/Private/Replication/Muting/MuteStateSynchronizer.h
@@ -0,0 +1,58 @@
+﻿// Copyright Epic Games, Inc. All Rights Reserved.
+
+#pragma once
+
+#include "Containers/Set.h"
+#include "Delegates/Delegate.h"
+#include "HAL/Platform.h"
+#include "Templates/UnrealTemplate.h"
+#include "UObject/SoftObjectPath.h"
+
+struct FConcertReplication_ChangeMuteState_Request;
+struct FConcertReplication_QueryMuteState_Response;
+
+namespace UE::MultiUserClient
+{
+	class FMuteStateQueryService;
+
+	DECLARE_MULTICAST_DELEGATE(FOnMuteStateChanged);
+
+	/**
+	 * Responsible for answering questions about whether an object is muted.
+	 * In the future, this may be extended with reasons (e.g. object Actor.Component is muted because Actor is muted).
+	 */
+	class FMuteStateSynchronizer : public FNoncopyable
+	{
+	public:
+
+		FMuteStateSynchronizer(FMuteStateQueryService& InMuteQueryService UE_LIFETIMEBOUND);
+		~FMuteStateSynchronizer();
+
+		/** @return Whether ObjectPath is muted. */
+		bool IsMuted(const FSoftObjectPath& ObjectPath) const { return MutedObjects.Contains(ObjectPath); }
+
+		/**
+		 * Update the mute state after the local application has successfully changed mute state.
+		 * FMuteStateQueryService will eventually notify us of the change but this applies it instantaneously.
+		 */
+		void UpdateStateFromSuccessfulChange(const FConcertReplication_ChangeMuteState_Request& Request);
+
+		/** Broadcasts whent the mute state changes (either because FMuteStateQueryService received updated state or because the local application has successfully made a change request). */
+		FOnMuteStateChanged& OnMuteStateChanged() { return OnMuteStateChangedDelegate; }
+		
+	private:
+
+		/** Updates us with the new mute state from the server regularily. */
+		FMuteStateQueryService& MuteQueryService;
+
+		/** All the objects that are muted. */
+		TSet<FSoftObjectPath> MutedObjects;
+
+		/** Broadcasts whent the mute state changes (either because FMuteStateQueryService received updated state or because the local application has successfully made a change request). */
+		FOnMuteStateChanged OnMuteStateChangedDelegate;
+
+		/** Updates MutedObjects from the mute state on the server. */
+		void UpdateFromServerState(const FConcertReplication_QueryMuteState_Response& NewMuteState);
+	};
+}
+
diff --git a/Engine/Plugins/Developer/Concert/ConcertApp/MultiUserClient/Source/MultiUserClient/Private/Replication/Notification/MutingNotifier.cpp b/Engine/Plugins/Developer/Concert/ConcertApp/MultiUserClient/Source/MultiUserClient/Private/Replication/Notification/MutingNotifier.cpp
new file mode 100644
--- /dev/null
+++ b/Engine/Plugins/Developer/Concert/ConcertApp/MultiUserClient/Source/MultiUserClient/Private/Replication/Notification/MutingNotifier.cpp
@@ -0,0 +1,41 @@
+﻿// Copyright Epic Games, Inc. All Rights Reserved.
+
+#include "MutingNotifier.h"
+
+#include "Replication/Muting/MuteStateManager.h"
+
+#include "Framework/Application/SlateApplication.h"
+#include "Framework/Notifications/NotificationManager.h"
+#include "Widgets/Notifications/SNotificationList.h"
+
+#define LOCTEXT_NAMESPACE "FMutingNotifier"
+
+namespace UE::MultiUserClient
+{
+	FMutingNotifier::FMutingNotifier(FMuteStateManager& InMuteManager)
+		: MuteManager(InMuteManager)
+	{
+		MuteManager.OnMuteRequestFailure().AddRaw(this, &FMutingNotifier::OnMuteRequestFailed);
+	}
+	
+	FMutingNotifier::~FMutingNotifier()
+	{
+		MuteManager.OnMuteRequestFailure().RemoveAll(this);
+	}
+
+	void FMutingNotifier::OnMuteRequestFailed(const FConcertReplication_ChangeMuteState_Request& Request, const FConcertReplication_ChangeMuteState_Response& Response)
+	{
+		if (FSlateApplication::IsInitialized())
+		{
+			FNotificationInfo Info(LOCTEXT("MuteRequestFailed.Title", "Pause / resume rejected by server."));
+			Info.SubText = LOCTEXT("MuteRequestFailed.Subtext", "Server-client state was likely de-synched.\nTry again.");
+			Info.bFireAndForget = true;
+			Info.ExpireDuration = 4.f;
+			Info.bUseSuccessFailIcons = true;
+			FSlateNotificationManager::Get().AddNotification(Info)
+				->SetCompletionState(SNotificationItem::CS_Fail);
+		}
+	}
+}
+
+#undef LOCTEXT_NAMESPACE
\ No newline at end of file
diff --git a/Engine/Plugins/Developer/Concert/ConcertApp/MultiUserClient/Source/MultiUserClient/Private/Replication/Notification/MutingNotifier.h b/Engine/Plugins/Developer/Concert/ConcertApp/MultiUserClient/Source/MultiUserClient/Private/Replication/Notification/MutingNotifier.h
new file mode 100644
--- /dev/null
+++ b/Engine/Plugins/Developer/Concert/ConcertApp/MultiUserClient/Source/MultiUserClient/Private/Replication/Notification/MutingNotifier.h
@@ -0,0 +1,31 @@
+﻿// Copyright Epic Games, Inc. All Rights Reserved.
+
+#pragma once
+
+#include "HAL/Platform.h"
+#include "Replication/Messages/Muting.h"
+#include "Templates/UnrealTemplate.h"
+
+struct FConcertReplication_ChangeMuteState_Request;
+struct FConcertReplication_ChangeMuteState_Response;
+
+namespace UE::MultiUserClient
+{
+	class FMuteStateManager;
+
+	/** Informs the user of things that went wrong with muting requests. */
+	class FMutingNotifier : public FNoncopyable
+	{
+	public:
+		FMutingNotifier(FMuteStateManager& InMuteManager UE_LIFETIMEBOUND);
+		~FMutingNotifier();
+
+	private:
+
+		/** Informs us when a request goes wrong. */
+		FMuteStateManager& MuteManager;
+		
+		void OnMuteRequestFailed(const FConcertReplication_ChangeMuteState_Request& Request, const FConcertReplication_ChangeMuteState_Response& Response);
+	};
+}
+
diff --git a/Engine/Plugins/Developer/Concert/ConcertApp/MultiUserClient/Source/MultiUserClient/Private/Replication/Notification/ReplicationUserNotifier.cpp b/Engine/Plugins/Developer/Concert/ConcertApp/MultiUserClient/Source/MultiUserClient/Private/Replication/Notification/ReplicationUserNotifier.cpp
new file mode 100644
--- /dev/null
+++ b/Engine/Plugins/Developer/Concert/ConcertApp/MultiUserClient/Source/MultiUserClient/Private/Replication/Notification/ReplicationUserNotifier.cpp
@@ -0,0 +1,11 @@
+﻿// Copyright Epic Games, Inc. All Rights Reserved.
+
+#include "ReplicationUserNotifier.h"
+
+namespace UE::MultiUserClient
+{
+	FReplicationUserNotifier::FReplicationUserNotifier(FReplicationClientManager& InReplicationClientManager, FMuteStateManager& InMuteManager)
+		: SubmissionNotifier(InReplicationClientManager)
+		, MutingNotifier(InMuteManager)
+	{}
+}
diff --git a/Engine/Plugins/Developer/Concert/ConcertApp/MultiUserClient/Source/MultiUserClient/Private/Replication/Notification/ReplicationUserNotifier.h b/Engine/Plugins/Developer/Concert/ConcertApp/MultiUserClient/Source/MultiUserClient/Private/Replication/Notification/ReplicationUserNotifier.h
new file mode 100644
--- /dev/null
+++ b/Engine/Plugins/Developer/Concert/ConcertApp/MultiUserClient/Source/MultiUserClient/Private/Replication/Notification/ReplicationUserNotifier.h
@@ -0,0 +1,31 @@
+﻿// Copyright Epic Games, Inc. All Rights Reserved.
+
+#pragma once
+
+#include "MutingNotifier.h"
+#include "Replication/Submission/Notification/SubmissionNotifier.h"
+
+#include "HAL/Platform.h"
+#include "Templates/UnrealTemplate.h"
+
+namespace UE::MultiUserClient
+{
+	class FMuteStateManager;
+
+	/** This informs the user of things that went wrong during replication. */
+	class FReplicationUserNotifier : public FNoncopyable
+	{
+	public:
+
+		FReplicationUserNotifier(FReplicationClientManager& InReplicationClientManager UE_LIFETIMEBOUND, FMuteStateManager& InMuteManager UE_LIFETIMEBOUND);
+
+	private:
+		
+		/** Notifies the user when stream or authority submission to the server fails. */
+		const FSubmissionNotifier SubmissionNotifier;
+		
+		/** Informs the user of things that went wrong with muting requests. */
+		const FMutingNotifier MutingNotifier;
+	};
+}
+
diff --git a/Engine/Plugins/Developer/Concert/ConcertApp/MultiUserClient/Source/MultiUserClient/Private/Replication/Util/GlobalAuthorityCache.cpp b/Engine/Plugins/Developer/Concert/ConcertApp/MultiUserClient/Source/MultiUserClient/Private/Replication/Util/GlobalAuthorityCache.cpp
--- a/Engine/Plugins/Developer/Concert/ConcertApp/MultiUserClient/Source/MultiUserClient/Private/Replication/Util/GlobalAuthorityCache.cpp
+++ b/Engine/Plugins/Developer/Concert/ConcertApp/MultiUserClient/Source/MultiUserClient/Private/Replication/Util/GlobalAuthorityCache.cpp
@@ -38,6 +38,19 @@ namespace UE::MultiUserClient
 		}
 	}
 
+	bool FGlobalAuthorityCache::IsObjectOrChildReferenced(const FSoftObjectPath& Object) const
+	{
+		const FString ObjectPathString = Object.ToString();
+		for (const TPair<FSoftObjectPath, TSet<FGuid>>& Pair : RegisteredObjectsToClients)
+		{
+			if (Pair.Key.ToString().Contains(ObjectPathString))
+			{
+				return true;
+			}
+		}
+		return false;
+	}
+
 	void FGlobalAuthorityCache::ForEachClientWithAuthorityOverObject(const FSoftObjectPath& Object, TFunctionRef<EBreakBehavior(const FGuid& ClientId)> Callback) const
 	{
 		const TSet<FGuid>* Clients = OwnedObjectsToClients.Find(Object);
diff --git a/Engine/Plugins/Developer/Concert/ConcertApp/MultiUserClient/Source/MultiUserClient/Private/Replication/Util/GlobalAuthorityCache.h b/Engine/Plugins/Developer/Concert/ConcertApp/MultiUserClient/Source/MultiUserClient/Private/Replication/Util/GlobalAuthorityCache.h
--- a/Engine/Plugins/Developer/Concert/ConcertApp/MultiUserClient/Source/MultiUserClient/Private/Replication/Util/GlobalAuthorityCache.h
+++ b/Engine/Plugins/Developer/Concert/ConcertApp/MultiUserClient/Source/MultiUserClient/Private/Replication/Util/GlobalAuthorityCache.h
@@ -39,6 +39,9 @@ namespace UE::MultiUserClient
 		/** Iterates every client that Object in its stream. */
 		void ForEachClientWithObjectInStream(const FSoftObjectPath& Object, TFunctionRef<EBreakBehavior(const FGuid& ClientId)> Callback) const;
 
+		/** @return Whether the object is referenced by at least one client stream. */
+		bool IsObjectOrChildReferenced(const FSoftObjectPath& Object) const;
+
 		/** Iterates every client that has authority over Object. */
 		void ForEachClientWithAuthorityOverObject(const FSoftObjectPath& Object, TFunctionRef<EBreakBehavior(const FGuid& ClientId)> Callback) const;
 		/** Util that uses ForEachClientWithAuthorityOverObject to make an array. */
diff --git a/Engine/Plugins/Developer/Concert/ConcertApp/MultiUserClient/Source/MultiUserClient/Private/Replication/Util/Query/MuteStateQueryService.cpp b/Engine/Plugins/Developer/Concert/ConcertApp/MultiUserClient/Source/MultiUserClient/Private/Replication/Util/Query/MuteStateQueryService.cpp
new file mode 100644
--- /dev/null
+++ b/Engine/Plugins/Developer/Concert/ConcertApp/MultiUserClient/Source/MultiUserClient/Private/Replication/Util/Query/MuteStateQueryService.cpp
@@ -0,0 +1,35 @@
+﻿// Copyright Epic Games, Inc. All Rights Reserved.
+
+#include "MuteStateQueryService.h"
+
+#include "IConcertSyncClient.h"
+#include "Replication/IConcertClientReplicationManager.h"
+
+namespace UE::MultiUserClient
+{
+	FMuteStateQueryService::FMuteStateQueryService(TWeakPtr<FToken> InToken, const IConcertSyncClient& InOwningClient)
+		: Token(MoveTemp(InToken))
+		, OwningClient(InOwningClient)
+	{}
+
+	void FMuteStateQueryService::SendQueryEvent()
+	{
+		// Avoid network calls when nobody is subscribed.
+		IConcertClientReplicationManager* ReplicationManager = OwningClient.GetReplicationManager();
+		if (!OnMuteStateQueriedDelegate.IsBound() || !ensure(ReplicationManager))
+		{
+			return;
+		}
+
+		// For now, we'll query for ALL objects - in the future we could restrict it to the local application's GWorld
+		ReplicationManager->QueryMuteState()
+			.Next([this, WeakToken = Token](FConcertReplication_QueryMuteState_Response&& Response)
+			{
+				// If WeakToken is not valid, our this pointer has been destroyed (there's currently no API for cancelling started requests).
+				if (WeakToken.IsValid())
+				{
+					OnMuteStateQueriedDelegate.Broadcast(Response);
+				}
+			});
+	}
+}
diff --git a/Engine/Plugins/Developer/Concert/ConcertApp/MultiUserClient/Source/MultiUserClient/Private/Replication/Util/Query/MuteStateQueryService.h b/Engine/Plugins/Developer/Concert/ConcertApp/MultiUserClient/Source/MultiUserClient/Private/Replication/Util/Query/MuteStateQueryService.h
new file mode 100644
--- /dev/null
+++ b/Engine/Plugins/Developer/Concert/ConcertApp/MultiUserClient/Source/MultiUserClient/Private/Replication/Util/Query/MuteStateQueryService.h
@@ -0,0 +1,56 @@
+﻿// Copyright Epic Games, Inc. All Rights Reserved.
+
+#pragma once
+
+#include "Replication/IToken.h"
+
+#include "Delegates/Delegate.h"
+#include "HAL/Platform.h"
+#include "Templates/SharedPointer.h"
+#include "Templates/UnrealTemplate.h"
+
+class IConcertSyncClient;
+struct FConcertReplication_QueryMuteState_Response;
+
+namespace UE::MultiUserClient
+{
+	DECLARE_MULTICAST_DELEGATE_OneParam(FOnMuteStateQueried, const FConcertReplication_QueryMuteState_Response& NewMuteState);
+	
+	/** Sends regular FConcertReplication_QueryMuteState_Request to endpoints and publishes the results. */
+	class FMuteStateQueryService : public FNoncopyable
+	{
+		friend class FRegularQueryService;
+	public:
+		
+		/**
+		 * @param InToken Used to check whether this instance was destroyed after a response is received.
+		 * @param InOwningClient The client through which queries are made. The caller ensures it outlasts the constructed instance.
+		 */
+		FMuteStateQueryService(TWeakPtr<FToken> InToken, const IConcertSyncClient& InOwningClient UE_LIFETIMEBOUND);
+
+		/** @return Sends a request to the server immediately */
+		void RequestInstantRefresh() { SendQueryEvent(); }
+
+		FOnMuteStateQueried& OnMuteStateQueried() { return OnMuteStateQueriedDelegate; }
+		
+	private:
+
+		/** Used to check whether we were destroyed after a response is received. */
+		const TWeakPtr<FToken> Token;
+		
+		/**
+		 * Used to send queries.
+		 * 
+		 * This FRegularQueryService's owner is supposed to make sure this FRegularQueryService is destroyed
+		 * when the client shuts down.
+		 */
+		const IConcertSyncClient& OwningClient;
+
+		/** Broadcasts when the mute state has been queried. */
+		FOnMuteStateQueried OnMuteStateQueriedDelegate;
+		
+		/** Queries the server for the client's current state. */
+		void SendQueryEvent();
+	};
+}
+
diff --git a/Engine/Plugins/Developer/Concert/ConcertApp/MultiUserClient/Source/MultiUserClient/Private/Replication/Util/Query/RegularQueryService.cpp b/Engine/Plugins/Developer/Concert/ConcertApp/MultiUserClient/Source/MultiUserClient/Private/Replication/Util/Query/RegularQueryService.cpp
--- a/Engine/Plugins/Developer/Concert/ConcertApp/MultiUserClient/Source/MultiUserClient/Private/Replication/Util/Query/RegularQueryService.cpp
+++ b/Engine/Plugins/Developer/Concert/ConcertApp/MultiUserClient/Source/MultiUserClient/Private/Replication/Util/Query/RegularQueryService.cpp
@@ -15,6 +15,7 @@ namespace UE::MultiUserClient
 				return true;
 			}))
 		, StreamAndAuthorityQueryService(Token, OwningClient)
+		, MuteStateQueryService(Token, OwningClient)
 	{}
 
 	FRegularQueryService::~FRegularQueryService()
@@ -25,5 +26,6 @@ namespace UE::MultiUserClient
 	void FRegularQueryService::Tick()
 	{
 		StreamAndAuthorityQueryService.SendQueryEvent();
+		MuteStateQueryService.SendQueryEvent();
 	}
 }
diff --git a/Engine/Plugins/Developer/Concert/ConcertApp/MultiUserClient/Source/MultiUserClient/Private/Replication/Util/Query/RegularQueryService.h b/Engine/Plugins/Developer/Concert/ConcertApp/MultiUserClient/Source/MultiUserClient/Private/Replication/Util/Query/RegularQueryService.h
--- a/Engine/Plugins/Developer/Concert/ConcertApp/MultiUserClient/Source/MultiUserClient/Private/Replication/Util/Query/RegularQueryService.h
+++ b/Engine/Plugins/Developer/Concert/ConcertApp/MultiUserClient/Source/MultiUserClient/Private/Replication/Util/Query/RegularQueryService.h
@@ -2,6 +2,7 @@
 
 #pragma once
 
+#include "MuteStateQueryService.h"
 #include "StreamAndAuthorityQueryService.h"
 #include "Replication/IToken.h"
 
@@ -30,6 +31,8 @@ namespace UE::MultiUserClient
 
 		/** @return The service for querying stream and authority state. */
 		FStreamAndAuthorityQueryService& GetStreamAndAuthorityQueryService() { return StreamAndAuthorityQueryService; }
+		/** @return The service for querying global mute state. */
+		FMuteStateQueryService& GetMuteStateQueryService() { return MuteStateQueryService; }
 
 	private:
 		
@@ -49,6 +52,8 @@ namespace UE::MultiUserClient
 
 		/** Queries stream and authority changes. */
 		FStreamAndAuthorityQueryService StreamAndAuthorityQueryService;
+		/** Queries global mute state at regular intervals. */
+		FMuteStateQueryService MuteStateQueryService;
 
 		/** Ticks the sub-services when it is time to send a new request. */
 		void Tick();
diff --git a/Engine/Plugins/Developer/Concert/ConcertApp/MultiUserClient/Source/MultiUserClient/Private/Widgets/ActiveSession/Replication/Client/Multi/Columns/MultiStreamColumns.h b/Engine/Plugins/Developer/Concert/ConcertApp/MultiUserClient/Source/MultiUserClient/Private/Widgets/ActiveSession/Replication/Client/Multi/Columns/MultiStreamColumns.h
--- a/Engine/Plugins/Developer/Concert/ConcertApp/MultiUserClient/Source/MultiUserClient/Private/Widgets/ActiveSession/Replication/Client/Multi/Columns/MultiStreamColumns.h
+++ b/Engine/Plugins/Developer/Concert/ConcertApp/MultiUserClient/Source/MultiUserClient/Private/Widgets/ActiveSession/Replication/Client/Multi/Columns/MultiStreamColumns.h
@@ -6,12 +6,14 @@
 #include "Replication/Editor/View/Column/IObjectTreeColumn.h"
 
 #include "Delegates/Delegate.h"
+#include "HAL/Platform.h"
 #include "Misc/EnumClassFlags.h"
 
 class IConcertClient;
 
 namespace UE::MultiUserClient
 {
+	class FMuteChangeTracker;
 	class FReassignObjectPropertiesLogic;
 	class FReplicationClientManager;
 }
@@ -26,33 +28,27 @@ namespace UE::ConcertSharedSlate
 
 namespace UE::MultiUserClient::MultiStreamColumns
 {
-	const extern FName ReplicationToggleColumnId;
+	const extern FName MuteToggleColumnId;
 	const extern FName AssignedClientsColumnId;
 	const extern FName AssignPropertyColumnId;
 	
 	/* @see ETopLevelColumnOrder and EReplicationPropertyColumnOrder */
 	enum class EColumnSortOrder
 	{
-		ReplicationToggle = 0,
+		MuteToggle = 0,
 		AssignPropertyColumn = 30,
 		ReassignOwnership = 40
 	};
 
 	/**
-	 * Toggles replication for all clients assigned to the object (and optionally all children).
+	 * Mutes and unmutes the object and its child objects.
 	 * 
-	 * @param ConcertClient Used to look up client names
-	 * @param ObjectHierarchyModelAttribute Used to get child objects
-	 * @param ClientManager Used to access all clients for toggling authority
-	 * @param ColumnsSortPriority The order relative to the other columns
-	 * 
-	 * @return A checkbox for controlling the authority of the object in the row
+	 * @param MuteChangeTracker Tells us the mute state and changes it.
+	 * @return A checkbox with pause and unpause icons.
 	 */
-	ConcertSharedSlate::FObjectColumnEntry ReplicationToggle(
-		TSharedRef<IConcertClient> ConcertClient,
-		TAttribute<ConcertSharedSlate::IObjectHierarchyModel*> ObjectHierarchyModelAttribute,
-		FReplicationClientManager& ClientManager,
-		const int32 ColumnsSortPriority = static_cast<int32>(EColumnSortOrder::ReplicationToggle)
+	ConcertSharedSlate::FObjectColumnEntry MuteToggleColumn(
+		FMuteChangeTracker& MuteChangeTracker UE_LIFETIMEBOUND,
+		int32 ColumnsSortPriority = static_cast<int32>(EColumnSortOrder::MuteToggle)
 		);
 
 	/**
@@ -70,7 +66,7 @@ namespace UE::MultiUserClient::MultiStreamColumns
 		TAttribute<TSharedPtr<ConcertSharedSlate::IMultiReplicationStreamEditor>> MultiStreamModelAttribute,
 		TAttribute<ConcertSharedSlate::IObjectHierarchyModel*> ObjectHierarchyModelAttribute,
 		FReassignObjectPropertiesLogic& ReassignmentLogic,
-		const FReplicationClientManager& ClientManager,
+		const FReplicationClientManager& ClientManager UE_LIFETIMEBOUND,
 		int32 ColumnsSortPriority = static_cast<int32>(EColumnSortOrder::ReassignOwnership)
 		);
 	
@@ -87,7 +83,7 @@ namespace UE::MultiUserClient::MultiStreamColumns
 	ConcertSharedSlate::FPropertyColumnEntry AssignPropertyColumn(
 		TAttribute<TSharedPtr<ConcertSharedSlate::IMultiReplicationStreamEditor>> MultiStreamEditor,
 		TSharedRef<IConcertClient> ConcertClient,
-		FReplicationClientManager& ClientManager,
+		FReplicationClientManager& ClientManager UE_LIFETIMEBOUND,
 		const int32 ColumnsSortPriority = static_cast<int32>(EColumnSortOrder::AssignPropertyColumn)
 		);
 }
diff --git a/Engine/Plugins/Developer/Concert/ConcertApp/MultiUserClient/Source/MultiUserClient/Private/Widgets/ActiveSession/Replication/Client/Multi/Columns/MultiStreamColumns_AssignProperties.cpp b/Engine/Plugins/Developer/Concert/ConcertApp/MultiUserClient/Source/MultiUserClient/Private/Widgets/ActiveSession/Replication/Client/Multi/Columns/MultiStreamColumns_AssignProperties.cpp
--- a/Engine/Plugins/Developer/Concert/ConcertApp/MultiUserClient/Source/MultiUserClient/Private/Widgets/ActiveSession/Replication/Client/Multi/Columns/MultiStreamColumns_AssignProperties.cpp
+++ b/Engine/Plugins/Developer/Concert/ConcertApp/MultiUserClient/Source/MultiUserClient/Private/Widgets/ActiveSession/Replication/Client/Multi/Columns/MultiStreamColumns_AssignProperties.cpp
@@ -99,7 +99,7 @@ namespace UE::MultiUserClient::MultiStreamColumns
 			virtual SHeaderRow::FColumn::FArguments CreateHeaderRowArgs() const override
 			{
 				return SHeaderRow::Column(AssignPropertyColumnId)
-					.DefaultLabel(LOCTEXT("Owner.Label", "Client"))
+					.DefaultLabel(LOCTEXT("Owner.Label", "Author"))
 					.ToolTipText(LOCTEXT("Owner.ToolTip", "Client that should replicate this property"))
 					.FillSized(FMultiUserReplicationStyle::Get()->GetFloat(TEXT("AllClients.Property.OwnerSize")));
 			}
diff --git a/Engine/Plugins/Developer/Concert/ConcertApp/MultiUserClient/Source/MultiUserClient/Private/Widgets/ActiveSession/Replication/Client/Multi/Columns/MultiStreamColumns_AssignedClients.cpp b/Engine/Plugins/Developer/Concert/ConcertApp/MultiUserClient/Source/MultiUserClient/Private/Widgets/ActiveSession/Replication/Client/Multi/Columns/MultiStreamColumns_AssignedClients.cpp
--- a/Engine/Plugins/Developer/Concert/ConcertApp/MultiUserClient/Source/MultiUserClient/Private/Widgets/ActiveSession/Replication/Client/Multi/Columns/MultiStreamColumns_AssignedClients.cpp
+++ b/Engine/Plugins/Developer/Concert/ConcertApp/MultiUserClient/Source/MultiUserClient/Private/Widgets/ActiveSession/Replication/Client/Multi/Columns/MultiStreamColumns_AssignedClients.cpp
@@ -89,7 +89,7 @@ namespace UE::MultiUserClient::MultiStreamColumns
 			}
 		};
 	}
-	
+
 	ConcertSharedSlate::FObjectColumnEntry AssignedClientsColumn(
 		TSharedRef<IConcertClient> ConcertClient,
 		TAttribute<TSharedPtr<ConcertSharedSlate::IMultiReplicationStreamEditor>> MultiStreamModelAttribute,
@@ -120,7 +120,7 @@ namespace UE::MultiUserClient::MultiStreamColumns
 			virtual SHeaderRow::FColumn::FArguments CreateHeaderRowArgs() const override
 			{
 				return SHeaderRow::Column(AssignedClientsColumnId)
-					.DefaultLabel(LOCTEXT("Owner.Label", "Clients"))
+					.DefaultLabel(LOCTEXT("Owner.Label", "Author"))
 					.ToolTipText(LOCTEXT("Owner.ToolTip", "Clients that have registered properties for the object"))
 					.FillSized(FMultiUserReplicationStyle::Get()->GetFloat(TEXT("AllClients.Object.OwnerSize")));
 			}
diff --git a/Engine/Plugins/Developer/Concert/ConcertApp/MultiUserClient/Source/MultiUserClient/Private/Widgets/ActiveSession/Replication/Client/Multi/Columns/MultiStreamColumns_MuteToggle.cpp b/Engine/Plugins/Developer/Concert/ConcertApp/MultiUserClient/Source/MultiUserClient/Private/Widgets/ActiveSession/Replication/Client/Multi/Columns/MultiStreamColumns_MuteToggle.cpp
new file mode 100644
--- /dev/null
+++ b/Engine/Plugins/Developer/Concert/ConcertApp/MultiUserClient/Source/MultiUserClient/Private/Widgets/ActiveSession/Replication/Client/Multi/Columns/MultiStreamColumns_MuteToggle.cpp
@@ -0,0 +1,80 @@
+﻿// Copyright Epic Games, Inc. All Rights Reserved.
+
+#include "MultiStreamColumns.h"
+
+#include "MultiUserReplicationStyle.h"
+#include "Replication/Editor/View/Column/IObjectTreeColumn.h"
+#include "SMuteToggle.h"
+
+#include "Widgets/Images/SImage.h"
+#include "Widgets/Layout/SBox.h"
+#include "Widgets/Views/SHeaderRow.h"
+
+#define LOCTEXT_NAMESPACE "ReplicationToggle"
+
+namespace UE::MultiUserClient::MultiStreamColumns
+{
+	const FName MuteToggleColumnId(TEXT("MuteToggleColumnId"));
+	
+	ConcertSharedSlate::FObjectColumnEntry MuteToggleColumn(
+		FMuteChangeTracker& MuteChangeTracker UE_LIFETIMEBOUND,
+		int32 ColumnsSortPriority
+		)
+	{
+		class FObjectColumn_MuteToggle : public ConcertSharedSlate::IObjectTreeColumn
+		{
+		public:
+
+			FObjectColumn_MuteToggle(FMuteChangeTracker& MuteChangeTracker UE_LIFETIMEBOUND)
+				: MuteChangeTracker(MuteChangeTracker)
+			{}
+			
+			virtual SHeaderRow::FColumn::FArguments CreateHeaderRowArgs() const override
+			{
+				return SHeaderRow::Column(MuteToggleColumnId)
+					.DefaultLabel(FText::GetEmpty())
+					.FixedWidth(FMultiUserReplicationStyle::Get()->GetFloat(TEXT("AllClients.Object.MuteToggle")))
+					[
+						SNew(SBox)
+						.WidthOverride(16.f)
+						.HeightOverride(16.f)
+						.HAlign(HAlign_Center)
+						.VAlign(VAlign_Center)
+						.ToolTipText(LOCTEXT("Mute.ToolTip", "Whether an object is supposed to replicate or not.\nTip: When you pause an actor, you can selectively unpause components."))
+						[
+							SNew(SImage)
+							.Image(FMultiUserReplicationStyle::Get()->GetBrush(TEXT("MultiUser.Icons.Play")))
+							.ColorAndOpacity(FSlateColor::UseForeground())
+						]
+					];
+			}
+			
+			virtual TSharedRef<SWidget> GenerateColumnWidget(const FBuildArgs& InArgs) override
+			{
+				return SNew(SBox)
+					.HAlign(HAlign_Center)
+					.VAlign(VAlign_Center)
+					.Padding(2.f)
+					[
+						SNew(SMuteToggle, InArgs.RowItem.RowData.GetObjectPath(), MuteChangeTracker)
+					];
+			}
+
+		private:
+
+			FMuteChangeTracker& MuteChangeTracker;
+		};
+
+		return {
+			ConcertSharedSlate::TReplicationColumnDelegates<ConcertSharedSlate::FObjectTreeRowContext>::FCreateColumn::CreateLambda(
+				[&MuteChangeTracker]()
+				{
+					return MakeShared<FObjectColumn_MuteToggle>(MuteChangeTracker);
+				}),
+			MuteToggleColumnId,
+			{ ColumnsSortPriority }
+		};
+	}
+}
+
+#undef LOCTEXT_NAMESPACE
\ No newline at end of file
diff --git a/Engine/Plugins/Developer/Concert/ConcertApp/MultiUserClient/Source/MultiUserClient/Private/Widgets/ActiveSession/Replication/Client/Multi/Columns/MultiStreamColumns_ReplicationToggle.cpp b/Engine/Plugins/Developer/Concert/ConcertApp/MultiUserClient/Source/MultiUserClient/Private/Widgets/ActiveSession/Replication/Client/Multi/Columns/MultiStreamColumns_ReplicationToggle.cpp
deleted file mode 100644
diff --git a/Engine/Plugins/Developer/Concert/ConcertApp/MultiUserClient/Source/MultiUserClient/Private/Widgets/ActiveSession/Replication/Client/Multi/Columns/SMuteToggle.cpp b/Engine/Plugins/Developer/Concert/ConcertApp/MultiUserClient/Source/MultiUserClient/Private/Widgets/ActiveSession/Replication/Client/Multi/Columns/SMuteToggle.cpp
new file mode 100644
--- /dev/null
+++ b/Engine/Plugins/Developer/Concert/ConcertApp/MultiUserClient/Source/MultiUserClient/Private/Widgets/ActiveSession/Replication/Client/Multi/Columns/SMuteToggle.cpp
@@ -0,0 +1,62 @@
+﻿// Copyright Epic Games, Inc. All Rights Reserved.
+
+#include "SMuteToggle.h"
+
+#include "MultiUserReplicationStyle.h"
+#include "Replication/Muting/MuteChangeTracker.h"
+
+#include "Styling/SlateTypes.h"
+#include "Widgets/Input/SCheckBox.h"
+
+#define LOCTEXT_NAMESPACE "SMuteToggle"
+
+namespace UE::MultiUserClient
+{
+	void SMuteToggle::Construct(const FArguments& InArgs, FSoftObjectPath InObjectPath, FMuteChangeTracker& InMuteChangeTracker)
+	{
+		ObjectPath = InObjectPath;
+		MuteChangeTracker = &InMuteChangeTracker;
+		ChildSlot
+		[
+			SNew(SCheckBox)
+			.IsChecked(this, &SMuteToggle::IsMuted)
+			.Visibility(this, &SMuteToggle::GetMuteVisibility)
+			.ToolTipText(this, &SMuteToggle::GetToolTipText)
+			.OnCheckStateChanged(this, &SMuteToggle::OnCheckboxStateChanged)
+			.Style(FMultiUserReplicationStyle::Get(), TEXT("AllClients.MuteToggle.Style"))
+		];
+	}
+
+	ECheckBoxState SMuteToggle::IsMuted() const
+	{
+		return MuteChangeTracker->IsMuted(ObjectPath)
+			? ECheckBoxState::Checked
+			: ECheckBoxState::Unchecked;
+	}
+
+	EVisibility SMuteToggle::GetMuteVisibility() const
+	{
+		return MuteChangeTracker->CanChangeMuteState(ObjectPath)
+			? EVisibility::Visible
+			: EVisibility::Hidden;
+	}
+
+	FText SMuteToggle::GetToolTipText() const
+	{
+		if (!MuteChangeTracker->CanChangeMuteState(ObjectPath))
+		{
+			return LOCTEXT("Mute.ToolTip.CannotMute", "Assign some properties to this object or a child object first.");
+		}
+		
+		return MuteChangeTracker->IsMuted(ObjectPath)
+			? LOCTEXT("Mute.ToolTip.Muted", "Paused. Press to resume replication.")
+			: LOCTEXT("Mute.ToolTip.Unmuted", "Replicating. Press to pause replication.");
+	}
+
+	void SMuteToggle::OnCheckboxStateChanged(ECheckBoxState)
+	{
+		MuteChangeTracker->ToggleMuteState(ObjectPath);
+	}
+}
+
+#undef LOCTEXT_NAMESPACE
\ No newline at end of file
diff --git a/Engine/Plugins/Developer/Concert/ConcertApp/MultiUserClient/Source/MultiUserClient/Private/Widgets/ActiveSession/Replication/Client/Multi/Columns/SMuteToggle.h b/Engine/Plugins/Developer/Concert/ConcertApp/MultiUserClient/Source/MultiUserClient/Private/Widgets/ActiveSession/Replication/Client/Multi/Columns/SMuteToggle.h
new file mode 100644
--- /dev/null
+++ b/Engine/Plugins/Developer/Concert/ConcertApp/MultiUserClient/Source/MultiUserClient/Private/Widgets/ActiveSession/Replication/Client/Multi/Columns/SMuteToggle.h
@@ -0,0 +1,42 @@
+﻿// Copyright Epic Games, Inc. All Rights Reserved.
+
+#pragma once
+
+#include "HAL/Platform.h"
+#include "UObject/SoftObjectPath.h"
+#include "Widgets/DeclarativeSyntaxSupport.h"
+#include "Widgets/SCompoundWidget.h"
+
+class FText;
+enum class ECheckBoxState : uint8;
+struct EVisibility;
+
+namespace UE::MultiUserClient
+{
+	class FMuteChangeTracker;
+	
+	/** Pauses and resumes replication for an object */
+	class SMuteToggle : public SCompoundWidget
+	{
+	public:
+
+		SLATE_BEGIN_ARGS(SMuteToggle) {}
+		SLATE_END_ARGS()
+
+		void Construct(const FArguments& InArgs, FSoftObjectPath InObjectPath, FMuteChangeTracker& InMuteChangeTracker UE_LIFETIMEBOUND);
+
+	private:
+
+		/** The object that is being muted by this widget. */
+		FSoftObjectPath ObjectPath; 
+		
+		/** Knows the mute state that should be displayed and is used to change it. */
+		FMuteChangeTracker* MuteChangeTracker = nullptr;
+		
+		ECheckBoxState IsMuted() const;
+		EVisibility GetMuteVisibility() const;
+		FText GetToolTipText() const;
+		void OnCheckboxStateChanged(ECheckBoxState CheckBoxState);
+	};
+}
+
diff --git a/Engine/Plugins/Developer/Concert/ConcertApp/MultiUserClient/Source/MultiUserClient/Private/Widgets/ActiveSession/Replication/Client/Multi/Columns/SReplicationMultiToggleCheckbox.cpp b/Engine/Plugins/Developer/Concert/ConcertApp/MultiUserClient/Source/MultiUserClient/Private/Widgets/ActiveSession/Replication/Client/Multi/Columns/SReplicationMultiToggleCheckbox.cpp
deleted file mode 100644
diff --git a/Engine/Plugins/Developer/Concert/ConcertApp/MultiUserClient/Source/MultiUserClient/Private/Widgets/ActiveSession/Replication/Client/Multi/Columns/SReplicationMultiToggleCheckbox.h b/Engine/Plugins/Developer/Concert/ConcertApp/MultiUserClient/Source/MultiUserClient/Private/Widgets/ActiveSession/Replication/Client/Multi/Columns/SReplicationMultiToggleCheckbox.h
deleted file mode 100644
diff --git a/Engine/Plugins/Developer/Concert/ConcertApp/MultiUserClient/Source/MultiUserClient/Private/Widgets/ActiveSession/Replication/Client/Multi/SAllClientsView.cpp b/Engine/Plugins/Developer/Concert/ConcertApp/MultiUserClient/Source/MultiUserClient/Private/Widgets/ActiveSession/Replication/Client/Multi/SAllClientsView.cpp
--- a/Engine/Plugins/Developer/Concert/ConcertApp/MultiUserClient/Source/MultiUserClient/Private/Widgets/ActiveSession/Replication/Client/Multi/SAllClientsView.cpp
+++ b/Engine/Plugins/Developer/Concert/ConcertApp/MultiUserClient/Source/MultiUserClient/Private/Widgets/ActiveSession/Replication/Client/Multi/SAllClientsView.cpp
@@ -7,17 +7,19 @@
 #include "Replication/Client/ReplicationClientManager.h"
 
 #include "Algo/Transform.h"
+#include "Replication/MultiUserReplicationManager.h"
 
 namespace UE::MultiUserClient
 {
-	void SAllClientsView::Construct(const FArguments& InArgs, TSharedRef<IConcertClient> InConcertClient, FReplicationClientManager& InClientManager)
+	void SAllClientsView::Construct(const FArguments&, TSharedRef<IConcertClient> InConcertClient, FMultiUserReplicationManager& InMultiUserReplicationManager)
 	{
-		ClientManager = &InClientManager;
-		AllClientsModel = MakeUnique<FAllClientsSelectionModel>(InClientManager);
+		ClientManager = InMultiUserReplicationManager.GetClientManager();
+		check(ClientManager);
+		AllClientsModel = MakeUnique<FAllClientsSelectionModel>(*ClientManager);
 		
 		ChildSlot
 		[
-			SNew(SMultiClientView, InConcertClient, InClientManager, *AllClientsModel)
+			SNew(SMultiClientView, InConcertClient, InMultiUserReplicationManager, *AllClientsModel)
 		];
 	}
 
diff --git a/Engine/Plugins/Developer/Concert/ConcertApp/MultiUserClient/Source/MultiUserClient/Private/Widgets/ActiveSession/Replication/Client/Multi/SAllClientsView.h b/Engine/Plugins/Developer/Concert/ConcertApp/MultiUserClient/Source/MultiUserClient/Private/Widgets/ActiveSession/Replication/Client/Multi/SAllClientsView.h
--- a/Engine/Plugins/Developer/Concert/ConcertApp/MultiUserClient/Source/MultiUserClient/Private/Widgets/ActiveSession/Replication/Client/Multi/SAllClientsView.h
+++ b/Engine/Plugins/Developer/Concert/ConcertApp/MultiUserClient/Source/MultiUserClient/Private/Widgets/ActiveSession/Replication/Client/Multi/SAllClientsView.h
@@ -3,6 +3,8 @@
 #pragma once
 
 #include "AllClientsSelectionModel.h"
+
+#include "HAL/Platform.h"
 #include "Widgets/DeclarativeSyntaxSupport.h"
 #include "Widgets/SCompoundWidget.h"
 
@@ -10,6 +12,7 @@ class IConcertClient;
 
 namespace UE::MultiUserClient
 {
+	class FMultiUserReplicationManager;
 	class FReplicationClient;
 	class FReplicationClientManager;
 
@@ -19,7 +22,7 @@ namespace UE::MultiUserClient
 		SLATE_BEGIN_ARGS(SAllClientsView){}
 		SLATE_END_ARGS()
 
-		void Construct(const FArguments& InArgs, TSharedRef<IConcertClient> InConcertClient, FReplicationClientManager& InClientManager);
+		void Construct(const FArguments&, TSharedRef<IConcertClient> InConcertClient, FMultiUserReplicationManager& InMultiUserReplicationManager UE_LIFETIMEBOUND);
 
 	private:
 
diff --git a/Engine/Plugins/Developer/Concert/ConcertApp/MultiUserClient/Source/MultiUserClient/Private/Widgets/ActiveSession/Replication/Client/Multi/SMultiClientView.cpp b/Engine/Plugins/Developer/Concert/ConcertApp/MultiUserClient/Source/MultiUserClient/Private/Widgets/ActiveSession/Replication/Client/Multi/SMultiClientView.cpp
--- a/Engine/Plugins/Developer/Concert/ConcertApp/MultiUserClient/Source/MultiUserClient/Private/Widgets/ActiveSession/Replication/Client/Multi/SMultiClientView.cpp
+++ b/Engine/Plugins/Developer/Concert/ConcertApp/MultiUserClient/Source/MultiUserClient/Private/Widgets/ActiveSession/Replication/Client/Multi/SMultiClientView.cpp
@@ -5,6 +5,7 @@
 #include "IClientSelectionModel.h"
 #include "MultiStreamModel.h"
 #include "Replication/ClientReplicationWidgetFactories.h"
+#include "Replication/MultiUserReplicationManager.h"
 #include "Replication/ReplicationWidgetFactories.h"
 #include "Replication/Client/ReplicationClient.h"
 #include "Replication/Client/ReplicationClientManager.h"
@@ -22,12 +23,17 @@
 
 namespace UE::MultiUserClient
 {
-	void SMultiClientView::Construct(const FArguments& InArgs, TSharedRef<IConcertClient> InConcertClient, FReplicationClientManager& InClientManager, IClientSelectionModel& InDisplayClientsModel)
+	void SMultiClientView::Construct(
+		const FArguments&,
+		TSharedRef<IConcertClient> InConcertClient,
+		FMultiUserReplicationManager& InMultiUserReplicationManager,
+		IClientSelectionModel& InDisplayClientsModel
+		)
 	{
-		StreamModel = MakeShared<FMultiStreamModel>(InDisplayClientsModel, InClientManager);
+		ClientManager = InMultiUserReplicationManager.GetClientManager();
+		StreamModel = MakeShared<FMultiStreamModel>(InDisplayClientsModel, *ClientManager);
 
 		ConcertClient = MoveTemp(InConcertClient);
-		ClientManager = &InClientManager;
 		ClientManager->OnRemoteClientsChanged().AddSP(this, &SMultiClientView::RebuildClientSubscriptions);
 		SelectionModel = &InDisplayClientsModel;
 		SelectionModel->OnSelectionChanged().AddSP(this, &SMultiClientView::RebuildClientSubscriptions);
@@ -41,11 +47,11 @@ namespace UE::MultiUserClient
 			+SVerticalBox::Slot()
 			.FillHeight(1.f)
 			[
-				CreateEditorContent(ConcertClient.ToSharedRef(), InClientManager)
+				CreateEditorContent(ConcertClient.ToSharedRef(), InMultiUserReplicationManager)
 			]
 		];
 		
-		SReplicationStatus::AppendReplicationStatus(*Content, InClientManager.GetAuthorityCache(),
+		SReplicationStatus::AppendReplicationStatus(*Content, ClientManager->GetAuthorityCache(),
 			SReplicationStatus::FArguments()
 			.DisplayedClients(this, &SMultiClientView::GetDisplayClientIds)
 			.ForEachReplicatedObject(this, &SMultiClientView::EnumerateObjectsInStreams)
@@ -66,10 +72,12 @@ namespace UE::MultiUserClient
 		CleanClientSubscriptions();
 	}
 
-	TSharedRef<SWidget> SMultiClientView::CreateEditorContent(const TSharedRef<IConcertClient>& InConcertClient, FReplicationClientManager& InClientManager)
+	TSharedRef<SWidget> SMultiClientView::CreateEditorContent(const TSharedRef<IConcertClient>& InConcertClient, FMultiUserReplicationManager& InMultiUserReplicationManager)
 	{
 		using namespace UE::ConcertSharedSlate;
 
+		FMuteStateManager& MuteManager = *InMultiUserReplicationManager.GetMuteManager();
+		
 		TAttribute<TSharedPtr<IMultiReplicationStreamEditor>> MultiStreamEditorAttribute =
 		   TAttribute<TSharedPtr<IMultiReplicationStreamEditor>>::CreateLambda([this]()
 		   {
@@ -80,9 +88,9 @@ namespace UE::MultiUserClient
 		   {
 			   return ObjectHierarchy.Get();
 		   });
-		FGetAutoAssignTarget GetAutoAssignTargetDelegate = FGetAutoAssignTarget::CreateLambda([this, &InClientManager](TConstArrayView<UObject*>)
+		FGetAutoAssignTarget GetAutoAssignTargetDelegate = FGetAutoAssignTarget::CreateLambda([this](TConstArrayView<UObject*>)
 		{
-			const TSharedRef<IEditableReplicationStreamModel>& LocalStream = InClientManager.GetLocalClient().GetClientEditModel();
+			const TSharedRef<IEditableReplicationStreamModel>& LocalStream = ClientManager->GetLocalClient().GetClientEditModel();
 			return StreamModel->GetEditableStreams().Contains(LocalStream) ? LocalStream.ToSharedPtr() : nullptr;
 		});
 		const TSharedRef<ConcertClientSharedSlate::FSelectPropertyFromUClassModel> PropertySourceModel = MakeShared<ConcertClientSharedSlate::FSelectPropertyFromUClassModel>();
@@ -92,7 +100,7 @@ namespace UE::MultiUserClient
 			.AdditionalPropertyColumns =
 			{
 				ReplicationColumns::Property::LabelColumn(),
-				MultiStreamColumns::AssignPropertyColumn(MultiStreamEditorAttribute, InConcertClient, InClientManager)
+				MultiStreamColumns::AssignPropertyColumn(MultiStreamEditorAttribute, InConcertClient, *ClientManager)
 			}
 		};
 		TSharedRef<IPropertyTreeView> PropertyTreeView = CreateFilterablePropertyTreeView(MoveTemp(TreeViewParams));
@@ -116,8 +124,8 @@ namespace UE::MultiUserClient
 			.OnExtendObjectsContextMenu = FExtendObjectMenu::CreateSP(this, &SMultiClientView::ExtendObjectContextMenu),
 			.ObjectColumns =
 			{
-				MultiStreamColumns::ReplicationToggle(InConcertClient, ObjectHierarchyAttribute, InClientManager),
-				MultiStreamColumns::AssignedClientsColumn(InConcertClient, MultiStreamEditorAttribute, ObjectHierarchyAttribute, InClientManager.GetReassignmentLogic(), InClientManager)
+				MultiStreamColumns::MuteToggleColumn(MuteManager.GetChangeTracker()),
+				MultiStreamColumns::AssignedClientsColumn(InConcertClient, MultiStreamEditorAttribute, ObjectHierarchyAttribute, ClientManager->GetReassignmentLogic(), *ClientManager)
 			},
 			.ShouldDisplayObjectDelegate = FShouldDisplayObject::CreateSP(this, &SMultiClientView::ShouldDisplayObject)
 		};
diff --git a/Engine/Plugins/Developer/Concert/ConcertApp/MultiUserClient/Source/MultiUserClient/Private/Widgets/ActiveSession/Replication/Client/Multi/SMultiClientView.h b/Engine/Plugins/Developer/Concert/ConcertApp/MultiUserClient/Source/MultiUserClient/Private/Widgets/ActiveSession/Replication/Client/Multi/SMultiClientView.h
--- a/Engine/Plugins/Developer/Concert/ConcertApp/MultiUserClient/Source/MultiUserClient/Private/Widgets/ActiveSession/Replication/Client/Multi/SMultiClientView.h
+++ b/Engine/Plugins/Developer/Concert/ConcertApp/MultiUserClient/Source/MultiUserClient/Private/Widgets/ActiveSession/Replication/Client/Multi/SMultiClientView.h
@@ -20,6 +20,7 @@ namespace UE::MultiUserClient
 {
 	class FGlobalAuthorityCache;
 	class FMultiStreamModel;
+	class FMultiUserReplicationManager;
 	class FReplicationClient;
 	class FReplicationClientManager;
 	class IClientSelectionModel;
@@ -31,7 +32,12 @@ namespace UE::MultiUserClient
 		SLATE_BEGIN_ARGS(SMultiClientView){}
 		SLATE_END_ARGS()
 
-		void Construct(const FArguments& InArgs, TSharedRef<IConcertClient> InConcertClient, FReplicationClientManager& InClientManager UE_LIFETIMEBOUND, IClientSelectionModel& InDisplayClientsModel UE_LIFETIMEBOUND);
+		void Construct(
+			const FArguments& InArgs,
+			TSharedRef<IConcertClient> InConcertClient,
+			FMultiUserReplicationManager& InMultiUserReplicationManager UE_LIFETIMEBOUND,
+			IClientSelectionModel& InDisplayClientsModel UE_LIFETIMEBOUND
+			);
 		virtual ~SMultiClientView() override;
 
 	private:
@@ -51,7 +57,7 @@ namespace UE::MultiUserClient
 		ConcertClientSharedSlate::FHideObjectsNotInWorldLogic HideObjectsNotInEditorWorld;
 
 		/** Creates this widget's editor content */
-		TSharedRef<SWidget> CreateEditorContent(const TSharedRef<IConcertClient>& InConcertClient, FReplicationClientManager& InClientManager);
+		TSharedRef<SWidget> CreateEditorContent(const TSharedRef<IConcertClient>& InConcertClient, FMultiUserReplicationManager& InMultiUserReplicationManager);
 
 		// SClientToolbar attributes
 		TSet<FGuid> GetDisplayClientIds() const;
diff --git a/Engine/Plugins/Developer/Concert/ConcertApp/MultiUserClient/Source/MultiUserClient/Private/Widgets/ActiveSession/Replication/Joined/SReplicationJoinedView.cpp b/Engine/Plugins/Developer/Concert/ConcertApp/MultiUserClient/Source/MultiUserClient/Private/Widgets/ActiveSession/Replication/Joined/SReplicationJoinedView.cpp
--- a/Engine/Plugins/Developer/Concert/ConcertApp/MultiUserClient/Source/MultiUserClient/Private/Widgets/ActiveSession/Replication/Joined/SReplicationJoinedView.cpp
+++ b/Engine/Plugins/Developer/Concert/ConcertApp/MultiUserClient/Source/MultiUserClient/Private/Widgets/ActiveSession/Replication/Joined/SReplicationJoinedView.cpp
@@ -16,7 +16,7 @@ namespace UE::MultiUserClient
 	{
 		ChildSlot
 		[
-			SNew(SAllClientsView, InClient->GetConcertClient(), *InReplicationManager->GetClientManager())
+			SNew(SAllClientsView, InClient->GetConcertClient(), *InReplicationManager)
 		];
 	}
 }
\ No newline at end of file
diff --git a/Engine/Plugins/Developer/Concert/ConcertUI/ConcertSharedSlate/Source/ConcertSharedSlate/Private/Replication/Editor/View/Tree/SReplicationColumnRow.h b/Engine/Plugins/Developer/Concert/ConcertUI/ConcertSharedSlate/Source/ConcertSharedSlate/Private/Replication/Editor/View/Tree/SReplicationColumnRow.h
--- a/Engine/Plugins/Developer/Concert/ConcertUI/ConcertSharedSlate/Source/ConcertSharedSlate/Private/Replication/Editor/View/Tree/SReplicationColumnRow.h
+++ b/Engine/Plugins/Developer/Concert/ConcertUI/ConcertSharedSlate/Source/ConcertSharedSlate/Private/Replication/Editor/View/Tree/SReplicationColumnRow.h
@@ -99,12 +99,11 @@ namespace UE::ConcertSharedSlate
 					[
 						SNew(SHorizontalBox)
 						
-						+SHorizontalBox::Slot() 
-						.Padding(InPadding)
+						+SHorizontalBox::Slot()
 						.HAlign(HAlign)
 						[
 							SNew(SBox)
-							.Visibility_Lambda([this](){ return this->IsHovered() ? EVisibility::Visible : EVisibility::Collapsed; })
+							.Visibility_Lambda([this](){ return this->IsHovered() ? EVisibility::Visible : EVisibility::Hidden; })
 							[
 								HoverContent
 							]
