In active Multi-User UI, replace tabs with fake "tabs" that are controlled via toggle buttons

diff --git a/Engine/Plugins/Developer/Concert/ConcertApp/MultiUserClient/Source/MultiUserClient/Private/MultiUserClientModule.cpp b/Engine/Plugins/Developer/Concert/ConcertApp/MultiUserClient/Source/MultiUserClient/Private/MultiUserClientModule.cpp
--- a/Engine/Plugins/Developer/Concert/ConcertApp/MultiUserClient/Source/MultiUserClient/Private/MultiUserClientModule.cpp
+++ b/Engine/Plugins/Developer/Concert/ConcertApp/MultiUserClient/Source/MultiUserClient/Private/MultiUserClientModule.cpp
@@ -1054,10 +1054,11 @@ private:
 	 */
 	TSharedRef<SDockTab> SpawnConcertBrowserTab(const FSpawnTabArgs& SpawnTabArgs)
 	{
-		const TSharedRef<SDockTab> DockTab = SNew(SDockTab)
-			.TabRole(NomadTab);
-		DockTab->SetContent(SNew(SConcertBrowser, DockTab, MultiUserClient.ToSharedRef(), ReplicationManager.ToSharedRef()));
-		return DockTab;
+		return SNew(SDockTab)
+			.TabRole(NomadTab)
+			[
+				SNew(SConcertBrowser, MultiUserClient.ToSharedRef(), ReplicationManager.ToSharedRef())
+			];
 	}
 
 	static FString GetConfiguredMultiUserServerExePathname()
diff --git a/Engine/Plugins/Developer/Concert/ConcertApp/MultiUserClient/Source/MultiUserClient/Private/MultiUserReplicationStyle.cpp b/Engine/Plugins/Developer/Concert/ConcertApp/MultiUserClient/Source/MultiUserClient/Private/MultiUserReplicationStyle.cpp
--- a/Engine/Plugins/Developer/Concert/ConcertApp/MultiUserClient/Source/MultiUserClient/Private/MultiUserReplicationStyle.cpp
+++ b/Engine/Plugins/Developer/Concert/ConcertApp/MultiUserClient/Source/MultiUserClient/Private/MultiUserReplicationStyle.cpp
@@ -68,12 +68,11 @@ namespace UE::MultiUserClient
 		// Timing
 		StyleSet->Set("AllClients.Reassignment.DisplayThrobberAfterSeconds", 0.2f);
 
-
+		// Muting
 		FSlateImageBrush* PlayBrush = new IMAGE_PLUGIN_BRUSH("generic_play_16x", Icon16x16);
 		FSlateImageBrush* PauseBrush = new IMAGE_PLUGIN_BRUSH("generic_pause_16x", Icon16x16);
 		StyleSet->Set("MultiUser.Icons.Play", PlayBrush);
 		StyleSet->Set("MultiUser.Icons.Pause", PauseBrush);
-
 		StyleSet->Set("AllClients.MuteToggle.Style", FCheckBoxStyle()		
 			.SetCheckBoxType(ESlateCheckBoxType::CheckBox)
 			.SetUncheckedImage(IMAGE_PLUGIN_BRUSH("generic_pause_16x", Icon16x16, FStyleColors::Foreground))
diff --git a/Engine/Plugins/Developer/Concert/ConcertApp/MultiUserClient/Source/MultiUserClient/Private/Widgets/ActiveSession/SActiveSessionRoot.cpp b/Engine/Plugins/Developer/Concert/ConcertApp/MultiUserClient/Source/MultiUserClient/Private/Widgets/ActiveSession/SActiveSessionRoot.cpp
--- a/Engine/Plugins/Developer/Concert/ConcertApp/MultiUserClient/Source/MultiUserClient/Private/Widgets/ActiveSession/SActiveSessionRoot.cpp
+++ b/Engine/Plugins/Developer/Concert/ConcertApp/MultiUserClient/Source/MultiUserClient/Private/Widgets/ActiveSession/SActiveSessionRoot.cpp
@@ -5,12 +5,14 @@
 #include "Overview/SActiveSessionOverviewTab.h"
 #include "Replication/SReplicationRootWidget.h"
 #include "SActiveSessionToolbar.h"
+#include "STabArea.h"
 
-#include "Framework/Docking/TabManager.h"
 #include "Styling/AppStyle.h"
 #include "Styling/SlateTypes.h"
-#include "Widgets/Docking/SDockTab.h"
+#include "Widgets/Images/SImage.h"
+#include "Widgets/Layout/SWidgetSwitcher.h"
 #include "Widgets/SBoxPanel.h"
+#include "Widgets/Text/STextBlock.h"
 
 #define LOCTEXT_NAMESPACE "SActiveSessionRoot"
 
@@ -21,34 +23,10 @@ namespace UE::MultiUserClient
 	
 	void SActiveSessionRoot::Construct(
 		const FArguments& InArgs,
-		const TSharedRef<SDockTab>& ConstructUnderMajorTab,
 		TSharedPtr<IConcertSyncClient> InConcertSyncClient,
 		TSharedRef<FMultiUserReplicationManager> InReplicationManager
 		)
 	{
-		ConcertSyncClient = InConcertSyncClient;
-		
-		TabManager = FGlobalTabmanager::Get()->NewTabManager(ConstructUnderMajorTab);
-		TSharedRef<FWorkspaceItem> AppMenuGroup = TabManager->AddLocalWorkspaceMenuCategory(LOCTEXT("ConcertActiveSession", "Active Session"));
-		TabManager->SetAllowWindowMenuBar(true);
-		RegisterTabSpawners(TabManager.ToSharedRef(), AppMenuGroup, InReplicationManager);
-
-		// Create our content
-		const TSharedRef<FTabManager::FLayout> Layout =
-			FTabManager::NewLayout("ConcertActiveSession_Layout_v3")
-			->AddArea
-			(
-				FTabManager::NewPrimaryArea()
-				->Split
-				(
-					FTabManager::NewStack()
-					->AddTab(SessionOverviewTabId, ETabState::OpenedTab)
-					->AddTab(ReplicationTabId, ETabState::OpenedTab)
-					->SetForegroundTab(SessionOverviewTabId)
-					->SetHideTabWell(false)
-					)
-			);
-		
 		ChildSlot
 		[
 			SNew(SVerticalBox)
@@ -62,6 +40,10 @@ namespace UE::MultiUserClient
 				.Padding(0.0f)
 				[
 					SNew(SActiveSessionToolbar, InConcertSyncClient)
+					.TabArea()
+					[
+						CreateTabArea()
+					]
 				]
 			]
 
@@ -71,44 +53,59 @@ namespace UE::MultiUserClient
 			.HAlign(HAlign_Fill)
 			.VAlign(VAlign_Fill)
 			[
-				TabManager->RestoreFrom(Layout, nullptr).ToSharedRef()
+				SAssignNew(TabSwitcher, SWidgetSwitcher)
+
+				+SWidgetSwitcher::Slot()
+				[
+					SNew(SActiveSessionOverviewTab, InConcertSyncClient)
+				]
+				+SWidgetSwitcher::Slot()
+				[
+					SNew(SReplicationRootWidget, InReplicationManager, InConcertSyncClient.ToSharedRef())
+				]
 			]
 		];
 	}
 
-	void SActiveSessionRoot::RegisterTabSpawners(
-		const TSharedRef<FTabManager>& InTabManager,
-		const TSharedRef<FWorkspaceItem>& AppMenuGroup,
-		TSharedRef<FMultiUserReplicationManager> InReplicationManager
-		)
+	TSharedRef<SWidget> SActiveSessionRoot::CreateTabArea()
 	{
-		InTabManager->RegisterTabSpawner(SessionOverviewTabId, FOnSpawnTab::CreateSP(this, &SActiveSessionRoot::SpawnTab_SessionOverview))
-			.SetDisplayName(LOCTEXT("OverviewTab.DisplayName", "Overview"))
-			.SetGroup(AppMenuGroup);
-	
-		InTabManager->RegisterTabSpawner(ReplicationTabId, FOnSpawnTab::CreateSP(this, &SActiveSessionRoot::SpawnTab_ReplicationControls, InReplicationManager))
-			.SetDisplayName(LOCTEXT("ReplicationTab.DisplayName", "Replication"))
-			.SetGroup(AppMenuGroup);
-	}
+		const auto CreateTabEntry = [this](int32 Index, const ANSICHAR* ImageBrush, FText Label, FText ToolTipText)
+		{
+			FTabEntry Entry;
 
-	TSharedRef<SDockTab> SActiveSessionRoot::SpawnTab_SessionOverview(const FSpawnTabArgs& Args)
-	{
-		return SNew(SDockTab)
-			.Label(LOCTEXT("SessionOverviewTab.Label", "Session"))
-			.ToolTipText(LOCTEXT("SessionOverviewTab.Tooltip", "Displays active session clients and activity."))
-			[
-				SNew(SActiveSessionOverviewTab, ConcertSyncClient)
-			];
-	}
+			Entry.ButtonContent.Widget = SNew(SHorizontalBox)
+				.ToolTipText(ToolTipText)
+				
+				+SHorizontalBox::Slot()
+				.AutoWidth()
+				[
+					SNew(SImage)
+					.Image(FConcertFrontendStyle::Get()->GetBrush(ImageBrush))
+				]
+				+SHorizontalBox::Slot()
+				.AutoWidth()
+				.Padding(2., 0.f, 0.f, 0.f)
+				[
+					SNew(STextBlock)
+					.Text(Label)
+				];
+			Entry.OnTabSelected.BindLambda([this, Index]()
+			{
+				TabSwitcher->SetActiveWidgetIndex(Index);
+			});
 
-	TSharedRef<SDockTab> SActiveSessionRoot::SpawnTab_ReplicationControls(const FSpawnTabArgs& Args, TSharedRef<FMultiUserReplicationManager> InReplicationManager)
-	{
-		return SNew(SDockTab)
-			.Label(LOCTEXT("ReplicationTab.Label", "Replication"))
-			.ToolTipText(LOCTEXT("ReplicationTab.Tooltip", "Manage real-time object replication"))
-			[
-				SNew(SReplicationRootWidget, InReplicationManager, ConcertSyncClient.ToSharedRef())
-			];
+			return Entry;
+		};
+
+		TArray<FTabEntry> Tabs
+		{
+			CreateTabEntry(0, "Concert.ActiveSession.Icon", LOCTEXT("OverviewTab.DisplayName", "Overview"), LOCTEXT("SessionOverviewTab.Tooltip", "Displays active session clients and activity.")),
+			CreateTabEntry(1, "Concert.MultiUser", LOCTEXT("ReplicationTab.Label", "Replication"), LOCTEXT("ReplicationTab.Tooltip", "Manage real-time object replication"))
+		};
+		
+		return SNew(STabArea)
+			.Tabs(Tabs)
+			.ActiveTabIndex(0);
 	}
 }
 
diff --git a/Engine/Plugins/Developer/Concert/ConcertApp/MultiUserClient/Source/MultiUserClient/Private/Widgets/ActiveSession/SActiveSessionRoot.h b/Engine/Plugins/Developer/Concert/ConcertApp/MultiUserClient/Source/MultiUserClient/Private/Widgets/ActiveSession/SActiveSessionRoot.h
--- a/Engine/Plugins/Developer/Concert/ConcertApp/MultiUserClient/Source/MultiUserClient/Private/Widgets/ActiveSession/SActiveSessionRoot.h
+++ b/Engine/Plugins/Developer/Concert/ConcertApp/MultiUserClient/Source/MultiUserClient/Private/Widgets/ActiveSession/SActiveSessionRoot.h
@@ -10,6 +10,7 @@ class FTabManager;
 class FWorkspaceItem;
 class IConcertSyncClient;
 class SDockTab;
+class SWidgetSwitcher;
 class SWindow;
 
 namespace UE::MultiUserClient
@@ -31,17 +32,14 @@ namespace UE::MultiUserClient
 		{}
 		SLATE_END_ARGS()
 
-		void Construct(const FArguments& InArgs, const TSharedRef<SDockTab>& ConstructUnderMajorTab, TSharedPtr<IConcertSyncClient> InConcertSyncClient, TSharedRef<FMultiUserReplicationManager> InReplicationManager);
+		void Construct(const FArguments& InArgs, TSharedPtr<IConcertSyncClient> InConcertSyncClient, TSharedRef<FMultiUserReplicationManager> InReplicationManager);
 
 	private:
 
-		TSharedPtr<IConcertSyncClient> ConcertSyncClient;
-		/** Holds the child content: SActiveSessionOverviewTab and SReplicationControlsTab. */
-		TSharedPtr<FTabManager> TabManager;
-		
-		void RegisterTabSpawners(const TSharedRef<FTabManager>& InTabManager, const TSharedRef<FWorkspaceItem>& AppMenuGroup, TSharedRef<FMultiUserReplicationManager> InReplicationManager);
-		TSharedRef<SDockTab> SpawnTab_SessionOverview(const FSpawnTabArgs& Args);
-		TSharedRef<SDockTab> SpawnTab_ReplicationControls(const FSpawnTabArgs& Args, TSharedRef<FMultiUserReplicationManager> InReplicationManager);
+		/** This switches "tabs" when a button in the "tab" area is changed. */
+		TSharedPtr<SWidgetSwitcher> TabSwitcher;
+
+		TSharedRef<SWidget> CreateTabArea();
 	};
 }
 
diff --git a/Engine/Plugins/Developer/Concert/ConcertApp/MultiUserClient/Source/MultiUserClient/Private/Widgets/ActiveSession/SActiveSessionToolbar.cpp b/Engine/Plugins/Developer/Concert/ConcertApp/MultiUserClient/Source/MultiUserClient/Private/Widgets/ActiveSession/SActiveSessionToolbar.cpp
--- a/Engine/Plugins/Developer/Concert/ConcertApp/MultiUserClient/Source/MultiUserClient/Private/Widgets/ActiveSession/SActiveSessionToolbar.cpp
+++ b/Engine/Plugins/Developer/Concert/ConcertApp/MultiUserClient/Source/MultiUserClient/Private/Widgets/ActiveSession/SActiveSessionToolbar.cpp
@@ -12,6 +12,8 @@
 #include "Widgets/Input/SButton.h"
 #include "Widgets/Input/SComboButton.h"
 #include "Widgets/Layout/SBorder.h"
+#include "Widgets/Layout/SSeparator.h"
+#include "Widgets/Layout/SSpacer.h"
 #include "Widgets/SBoxPanel.h"
 #include "Widgets/Text/STextBlock.h"
 
@@ -71,6 +73,7 @@ namespace UE::MultiUserClient
 
 			// Status Message
 			+SHorizontalBox::Slot()
+			.AutoWidth()
 			.VAlign(VAlign_Center)
 			.Padding(FMargin(4.0f, 1.0f))
 			[
@@ -84,6 +87,38 @@ namespace UE::MultiUserClient
 					.Text(this, &SActiveSessionToolbar::GetConnectionStatusText)
 				]
 			]
+
+			+SHorizontalBox::Slot()
+			.AutoWidth()
+			.Padding(4, 3)
+			[
+				SNew(SSeparator)
+				.Orientation(Orient_Vertical)
+			]
+
+			// Tab area
+			+SHorizontalBox::Slot()
+			.AutoWidth()
+			.VAlign(VAlign_Center)
+			[
+				InArgs._TabArea.Widget
+			]
+			
+			+SHorizontalBox::Slot()
+			.AutoWidth()
+			.Padding(4, 3)
+			[
+				SNew(SSeparator)
+				.Orientation(Orient_Vertical)
+			]
+
+			// Move all the other content to the right if the panel is big
+			+SHorizontalBox::Slot()
+			.FillWidth(1.f)
+			[
+				SNew(SSpacer)
+			]
+			
 			+SHorizontalBox::Slot()
 			.AutoWidth()
 			.VAlign(VAlign_Center)
diff --git a/Engine/Plugins/Developer/Concert/ConcertApp/MultiUserClient/Source/MultiUserClient/Private/Widgets/ActiveSession/SActiveSessionToolbar.h b/Engine/Plugins/Developer/Concert/ConcertApp/MultiUserClient/Source/MultiUserClient/Private/Widgets/ActiveSession/SActiveSessionToolbar.h
--- a/Engine/Plugins/Developer/Concert/ConcertApp/MultiUserClient/Source/MultiUserClient/Private/Widgets/ActiveSession/SActiveSessionToolbar.h
+++ b/Engine/Plugins/Developer/Concert/ConcertApp/MultiUserClient/Source/MultiUserClient/Private/Widgets/ActiveSession/SActiveSessionToolbar.h
@@ -36,6 +36,8 @@ namespace UE::MultiUserClient
 
 		SLATE_BEGIN_ARGS(SActiveSessionToolbar)
 		{}
+			/** Contains the "tab" buttons. */
+			SLATE_NAMED_SLOT(FArguments, TabArea)
 		SLATE_END_ARGS()
 
 		void Construct(const FArguments& InArgs, TSharedPtr<IConcertSyncClient> InConcertSyncClient);
diff --git a/Engine/Plugins/Developer/Concert/ConcertApp/MultiUserClient/Source/MultiUserClient/Private/Widgets/ActiveSession/STabArea.cpp b/Engine/Plugins/Developer/Concert/ConcertApp/MultiUserClient/Source/MultiUserClient/Private/Widgets/ActiveSession/STabArea.cpp
new file mode 100644
--- /dev/null
+++ b/Engine/Plugins/Developer/Concert/ConcertApp/MultiUserClient/Source/MultiUserClient/Private/Widgets/ActiveSession/STabArea.cpp
@@ -0,0 +1,57 @@
+﻿// Copyright Epic Games, Inc. All Rights Reserved.
+
+#include "STabArea.h"
+
+#include "STabButton.h"
+
+#include "Widgets/SBoxPanel.h"
+
+namespace UE::MultiUserClient
+{
+	void STabArea::Construct(const FArguments& InArgs)
+	{
+		const TSharedRef<SHorizontalBox> Content = SNew(SHorizontalBox);
+		for (const FTabEntry& Entry : InArgs._Tabs)
+		{
+			const int32 ButtonIndex = TabButtons.Num();
+			
+			TSharedPtr<STabButton> Button;
+			Content->AddSlot()
+				.Padding(InArgs._Padding)
+				[
+					SAssignNew(Button, STabButton)
+					.OnActivated_Lambda([this, OnTabSelected = Entry.OnTabSelected, ButtonIndex]()
+					{
+						OnButtonActivated(ButtonIndex);
+						OnTabSelected.ExecuteIfBound();
+					})
+					[
+						Entry.ButtonContent.Widget
+					]
+				];
+
+			TabButtons.Emplace(Button.ToSharedRef());
+		}
+
+		if (ensure(TabButtons.IsValidIndex(InArgs._ActiveTabIndex)))
+		{
+			TabButtons[InArgs._ActiveTabIndex]->Activate();
+		}
+		
+		ChildSlot
+		[
+			Content
+		];
+	}
+	
+	void STabArea::OnButtonActivated(int32 ButtonIndex)
+	{
+		for (int32 i = 0; i < TabButtons.Num(); ++i)
+		{
+			if (ButtonIndex != i)
+			{
+				TabButtons[i]->Deactivate();
+			}
+		}
+	}
+}
diff --git a/Engine/Plugins/Developer/Concert/ConcertApp/MultiUserClient/Source/MultiUserClient/Private/Widgets/ActiveSession/STabArea.h b/Engine/Plugins/Developer/Concert/ConcertApp/MultiUserClient/Source/MultiUserClient/Private/Widgets/ActiveSession/STabArea.h
new file mode 100644
--- /dev/null
+++ b/Engine/Plugins/Developer/Concert/ConcertApp/MultiUserClient/Source/MultiUserClient/Private/Widgets/ActiveSession/STabArea.h
@@ -0,0 +1,49 @@
+﻿// Copyright Epic Games, Inc. All Rights Reserved.
+
+#pragma once
+
+#include "Widgets/DeclarativeSyntaxSupport.h"
+#include "Widgets/SCompoundWidget.h"
+
+namespace UE::MultiUserClient
+{
+	class STabButton;
+}
+
+namespace UE::MultiUserClient
+{
+	struct FTabEntry
+	{
+		TAlwaysValidWidget ButtonContent;
+		FSimpleDelegate OnTabSelected;
+	};
+	
+	/** Manages multiple STabButtons, making sure exactly 1 is active at a time. */
+	class STabArea : public SCompoundWidget
+	{
+	public:
+
+		SLATE_BEGIN_ARGS(STabArea)
+			: _Padding(3.f)
+			{}
+			/** The tabs to create */
+			SLATE_ARGUMENT(TArray<FTabEntry>, Tabs)
+			/** The tab that should be active by default. */
+			SLATE_ARGUMENT(int32, ActiveTabIndex)
+			
+			/** Padding between buttons */
+			SLATE_ARGUMENT(FMargin, Padding)
+		SLATE_END_ARGS()
+
+		void Construct(const FArguments& InArgs);
+
+	private:
+
+		/** The buttons in this area. */
+		TArray<TSharedRef<STabButton>> TabButtons;
+
+		/** Called when the button at TabButtons[ButtonIndex] is activated. */
+		void OnButtonActivated(int32 ButtonIndex);
+	};
+}
+
diff --git a/Engine/Plugins/Developer/Concert/ConcertApp/MultiUserClient/Source/MultiUserClient/Private/Widgets/ActiveSession/STabButton.cpp b/Engine/Plugins/Developer/Concert/ConcertApp/MultiUserClient/Source/MultiUserClient/Private/Widgets/ActiveSession/STabButton.cpp
new file mode 100644
--- /dev/null
+++ b/Engine/Plugins/Developer/Concert/ConcertApp/MultiUserClient/Source/MultiUserClient/Private/Widgets/ActiveSession/STabButton.cpp
@@ -0,0 +1,40 @@
+﻿// Copyright Epic Games, Inc. All Rights Reserved.
+
+#include "STabButton.h"
+
+#include "ConcertFrontendStyle.h"
+
+#include "Widgets/Input/SCheckBox.h"
+
+namespace UE::MultiUserClient
+{
+	void STabButton::Construct(const FArguments& InArgs)
+	{
+		OnActivatedDelegate = InArgs._OnActivated;
+		
+		ChildSlot
+		[
+			SNew(SCheckBox)
+			.Style(FConcertFrontendStyle::Get(), TEXT("Concert.TabButton"))
+			.IsChecked(this, &STabButton::IsChecked)
+            .OnCheckStateChanged(this, &STabButton::OnButtonClicked)
+			[
+				InArgs._ButtonContent.Widget
+			]
+		];
+	}
+
+	ECheckBoxState STabButton::IsChecked() const
+	{
+		return bIsActive ? ECheckBoxState::Checked : ECheckBoxState::Unchecked;
+	}
+
+	void STabButton::OnButtonClicked(ECheckBoxState NewState)
+	{
+		if (!bIsActive)
+		{
+			bIsActive = true;
+			OnActivatedDelegate.ExecuteIfBound();
+		}
+	}
+}
diff --git a/Engine/Plugins/Developer/Concert/ConcertApp/MultiUserClient/Source/MultiUserClient/Private/Widgets/ActiveSession/STabButton.h b/Engine/Plugins/Developer/Concert/ConcertApp/MultiUserClient/Source/MultiUserClient/Private/Widgets/ActiveSession/STabButton.h
new file mode 100644
--- /dev/null
+++ b/Engine/Plugins/Developer/Concert/ConcertApp/MultiUserClient/Source/MultiUserClient/Private/Widgets/ActiveSession/STabButton.h
@@ -0,0 +1,44 @@
+﻿// Copyright Epic Games, Inc. All Rights Reserved.
+
+#pragma once
+
+#include "HAL/Platform.h"
+#include "Widgets/DeclarativeSyntaxSupport.h"
+#include "Widgets/SCompoundWidget.h"
+
+enum class ECheckBoxState : uint8;
+
+namespace UE::MultiUserClient
+{
+	/** Used to switch "tabs" in the active session UI. It is a blue button when active and gray when inactive. */
+	class STabButton : public SCompoundWidget
+	{
+	public:
+
+		SLATE_BEGIN_ARGS(STabButton){}
+			/** Called when this button is clicked and becomes active. Not trigger if it was already active. */
+			SLATE_EVENT(FSimpleDelegate, OnActivated)
+			
+			SLATE_DEFAULT_SLOT(FArguments, ButtonContent)
+		SLATE_END_ARGS()
+
+		void Construct(const FArguments& InArgs);
+
+		/** Makes the button appear active. */
+		void Activate() { bIsActive = true; }
+		/** Makes the button appear inactive. */
+		void Deactivate() { bIsActive = false; }
+		
+	private:
+
+		/** Whether the button is active (the tab content is supposed to be shown). */
+		bool bIsActive = false;
+
+		/** Called when this button is clicked and becomes active. Not trigger if it was already active. */
+		FSimpleDelegate OnActivatedDelegate;
+		
+		ECheckBoxState IsChecked() const;
+		void OnButtonClicked(ECheckBoxState NewState);
+	};
+}
+
diff --git a/Engine/Plugins/Developer/Concert/ConcertApp/MultiUserClient/Source/MultiUserClient/Private/Widgets/SConcertBrowser.cpp b/Engine/Plugins/Developer/Concert/ConcertApp/MultiUserClient/Source/MultiUserClient/Private/Widgets/SConcertBrowser.cpp
--- a/Engine/Plugins/Developer/Concert/ConcertApp/MultiUserClient/Source/MultiUserClient/Private/Widgets/SConcertBrowser.cpp
+++ b/Engine/Plugins/Developer/Concert/ConcertApp/MultiUserClient/Source/MultiUserClient/Private/Widgets/SConcertBrowser.cpp
@@ -19,7 +19,6 @@
 
 void SConcertBrowser::Construct(
 	const FArguments& InArgs,
-	TSharedRef<SDockTab> InConstructUnderMajorTab,
 	TSharedRef<IConcertSyncClient> InSyncClient,
 	TSharedRef<UE::MultiUserClient::FMultiUserReplicationManager> InReplicationManager)
 {
@@ -37,7 +36,6 @@ void SConcertBrowser::Construct(
 
 	WeakConcertSyncClient = InSyncClient;
 	WeakReplicationManager = InReplicationManager;
-	ConstructedUnderMajorTab = MoveTemp(InConstructUnderMajorTab);
 	if (TSharedPtr<IConcertSyncClient> ConcertSyncClient = WeakConcertSyncClient.Pin())
 	{
 		SearchedText = MakeShared<FText>(); // Will keep in memory the session browser search text between join/leave UI transitions.
@@ -59,12 +57,6 @@ void SConcertBrowser::HandleSessionConnectionChanged(IConcertClientSession& InSe
 
 void SConcertBrowser::AttachChildWidget(EConcertConnectionStatus ConnectionStatus)
 {
-	const TSharedPtr<SDockTab> OwningMajorTab = ConstructedUnderMajorTab.Pin();
-	if (!ensure(OwningMajorTab))
-	{
-		return;
-	}
-	
 	if (const TSharedPtr<IConcertSyncClient> ConcertSyncClient = WeakConcertSyncClient.Pin()
 		; ensure(ConcertSyncClient))
 	{
@@ -75,7 +67,6 @@ void SConcertBrowser::AttachChildWidget(EConcertConnectionStatus ConnectionStatu
 			{
 				ChildSlot.AttachWidget(
 					SNew(UE::MultiUserClient::SActiveSessionRoot,
-						OwningMajorTab.ToSharedRef(),
 						ConcertSyncClient,
 						ReplicationManager.ToSharedRef()
 						)
diff --git a/Engine/Plugins/Developer/Concert/ConcertApp/MultiUserClient/Source/MultiUserClient/Private/Widgets/SConcertBrowser.h b/Engine/Plugins/Developer/Concert/ConcertApp/MultiUserClient/Source/MultiUserClient/Private/Widgets/SConcertBrowser.h
--- a/Engine/Plugins/Developer/Concert/ConcertApp/MultiUserClient/Source/MultiUserClient/Private/Widgets/SConcertBrowser.h
+++ b/Engine/Plugins/Developer/Concert/ConcertApp/MultiUserClient/Source/MultiUserClient/Private/Widgets/SConcertBrowser.h
@@ -33,11 +33,10 @@ public:
 	* Constructs the Browser.
 	*
 	* @param InArgs The Slate argument list.
-	* @param InConstructUnderMajorTab The major tab which will contain the session front-end.
-	* @param InConstructUnderWindow The window in which this widget is being constructed.
 	* @param InSyncClient The sync client.
+	* @param InReplicationManager Used to create replication UI when in an active session
 	*/
-	void Construct(const FArguments& InArgs, TSharedRef<SDockTab> InConstructUnderMajorTab, TSharedRef<IConcertSyncClient> InSyncClient, TSharedRef<UE::MultiUserClient::FMultiUserReplicationManager> InReplicationManager);
+	void Construct(const FArguments& InArgs, TSharedRef<IConcertSyncClient> InSyncClient, TSharedRef<UE::MultiUserClient::FMultiUserReplicationManager> InReplicationManager);
 
 private:
 
@@ -46,12 +45,6 @@ private:
 	/** Interacts with the replication system on behalf of Multi-User. */
 	TWeakPtr<UE::MultiUserClient::FMultiUserReplicationManager> WeakReplicationManager;
 
-	/**
-	 * Kept so it can be passed on to SActiveSessionRoot.
-	 * Important: since this a pointer to the top-level widget that contains us, we must keep a weak ptr or we'll cause a memory leak.
-	 */
-	TWeakPtr<SDockTab> ConstructedUnderMajorTab;
-
 	/** Keeps the session browser searched text in memory to reapply it when a user leaves a session and goes back to the session browser. */
 	TSharedPtr<FText> SearchedText;
 	
diff --git a/Engine/Plugins/Developer/Concert/ConcertUI/ConcertSharedSlate/Source/ConcertSharedSlate/Private/ConcertFrontendStyle.cpp b/Engine/Plugins/Developer/Concert/ConcertUI/ConcertSharedSlate/Source/ConcertSharedSlate/Private/ConcertFrontendStyle.cpp
--- a/Engine/Plugins/Developer/Concert/ConcertUI/ConcertSharedSlate/Source/ConcertSharedSlate/Private/ConcertFrontendStyle.cpp
+++ b/Engine/Plugins/Developer/Concert/ConcertUI/ConcertSharedSlate/Source/ConcertSharedSlate/Private/ConcertFrontendStyle.cpp
@@ -1,14 +1,17 @@
 // Copyright Epic Games, Inc. All Rights Reserved.
 
 #include "ConcertFrontendStyle.h"
+
 #include "Brushes/SlateBoxBrush.h"
-#include "Styling/SlateStyleRegistry.h"
 #include "Brushes/SlateImageBrush.h"
-#include "Styling/SlateTypes.h"
 #include "Interfaces/IPluginManager.h"
 #include "Misc/Paths.h"
 #include "Styling/StarshipCoreStyle.h"
+#include "Styling/SlateTypes.h"
 #include "Styling/SlateStyle.h"
+#include "Styling/SlateStyleRegistry.h"
+#include "Styling/SlateTypes.h"
+#include "Styling/StyleColors.h"
 
 LLM_DEFINE_TAG(Concert_ConcertFrontendStyle);
 
@@ -105,7 +108,15 @@ void FConcertFrontendStyle::Initialize()
 	StyleSet->Set("Concert.JumpToLocation",     new IMAGE_PLUGIN_BRUSH("Icons/icon_PresenceLocation_32x",   Icon16x16, IconColorAndOpacity));
 	StyleSet->Set("Concert.HidePresence",       new IMAGE_PLUGIN_BRUSH("Icons/icon_PresenceEyeOff_32x",     Icon16x16, IconColorAndOpacity));
 	StyleSet->Set("Concert.ShowPresence",       new IMAGE_PLUGIN_BRUSH("Icons/icon_PresenceEyeOn_32x",      Icon16x16, IconColorAndOpacity));
-
+	// Tab section of ActiveSession
+	FCheckBoxStyle BaseCheckBoxStyle = FAppStyle::Get().GetWidgetStyle<FCheckBoxStyle>("ToggleButtonCheckbox");
+	StyleSet->Set( "Concert.TabButton", BaseCheckBoxStyle
+		.SetCheckedImage(FSlateRoundedBoxBrush(FStyleColors::Primary, 4.0f, FStyleColors::Primary, 1.0f))
+		.SetCheckedHoveredImage(FSlateRoundedBoxBrush(FStyleColors::PrimaryHover, 4.0f, FStyleColors::PrimaryHover, 1.0f))
+		.SetCheckedPressedImage(FSlateRoundedBoxBrush(FStyleColors::PrimaryHover, 4.0f, FStyleColors::PrimaryHover, 1.0f))
+		.SetPadding(FMargin(16, 6))
+	);
+	
 	// Multi-user archived session browser
 	StyleSet->Set("Concert.MuteActivities",     new IMAGE_PLUGIN_BRUSH("Icons/icon_PresenceEyeOff_32x",      Icon16x16, IconColorAndOpacity));
 	StyleSet->Set("Concert.UnmuteActivities",     new IMAGE_PLUGIN_BRUSH("Icons/icon_PresenceEyeOn_32x",      Icon16x16, IconColorAndOpacity));
