Exclude objects from transactions if they're being replicated in Multi-User.

diff --git a/Engine/Plugins/Developer/Concert/ConcertApp/MultiUserClient/Source/MultiUserClient/Private/Replication/Misc/ChangeLevelHandler.cpp b/Engine/Plugins/Developer/Concert/ConcertApp/MultiUserClient/Source/MultiUserClient/Private/Replication/Misc/ChangeLevelHandler.cpp
new file mode 100644
--- /dev/null
+++ b/Engine/Plugins/Developer/Concert/ConcertApp/MultiUserClient/Source/MultiUserClient/Private/Replication/Misc/ChangeLevelHandler.cpp
@@ -0,0 +1,48 @@
+// Copyright Epic Games, Inc. All Rights Reserved.
+
+#include "ChangeLevelHandler.h"
+
+#include "Replication/Editor/Model/IEditableReplicationStreamModel.h"
+
+#include "Engine/Engine.h"
+#include "Engine/World.h"
+
+namespace UE::MultiUserClient
+{
+	FChangeLevelHandler::FChangeLevelHandler(ConcertSharedSlate::IEditableReplicationStreamModel& UpdatedModel)
+		: UpdatedModel(UpdatedModel)
+	{
+		if (ensure(GEngine))
+		{
+			GEngine->OnWorldDestroyed().AddRaw(this, &FChangeLevelHandler::OnWorldDestroyed);
+			GEngine->OnWorldAdded().AddRaw(this, &FChangeLevelHandler::OnWorldAdded);
+		}
+	}
+
+	FChangeLevelHandler::~FChangeLevelHandler()
+	{
+		if (GEngine)
+		{
+			GEngine->OnWorldDestroyed().RemoveAll(this);
+			GEngine->OnWorldAdded().RemoveAll(this);
+		}
+	}
+
+	void FChangeLevelHandler::OnWorldDestroyed(UWorld* World)
+	{
+		// Remember the current map so when the next world
+		if (World->WorldType == EWorldType::Editor)
+		{
+			PreviousWorldPath = FSoftObjectPath(World);
+		}
+	}
+
+	void FChangeLevelHandler::OnWorldAdded(UWorld* World) const
+	{
+		// User reloaded the same map - keep around all settings
+		if (PreviousWorldPath != World)
+		{
+			UpdatedModel.Clear();
+		}
+	}
+}
diff --git a/Engine/Plugins/Developer/Concert/ConcertApp/MultiUserClient/Source/MultiUserClient/Private/Replication/Misc/ChangeLevelHandler.h b/Engine/Plugins/Developer/Concert/ConcertApp/MultiUserClient/Source/MultiUserClient/Private/Replication/Misc/ChangeLevelHandler.h
new file mode 100644
--- /dev/null
+++ b/Engine/Plugins/Developer/Concert/ConcertApp/MultiUserClient/Source/MultiUserClient/Private/Replication/Misc/ChangeLevelHandler.h
@@ -0,0 +1,40 @@
+// Copyright Epic Games, Inc. All Rights Reserved.
+
+#pragma once
+
+#include "HAL/Platform.h"
+#include "Templates/UnrealTemplate.h"
+#include "UObject/SoftObjectPath.h"
+
+class UWorld;
+
+namespace UE::ConcertSharedSlate { class IEditableReplicationStreamModel; }
+
+namespace UE::MultiUserClient
+{
+	/**
+	 * Handles a client opening a new level.
+	 * Upon leaving a level, all replicated objects assigned to that client in that level are removed.
+	 */
+	class FChangeLevelHandler : public FNoncopyable
+	{
+	public:
+		
+		/**
+		 * @param UpdatedModel The client model to update when the local editor changes maps. The caller ensures that it outlives the constructed object.
+		 */
+		FChangeLevelHandler(ConcertSharedSlate::IEditableReplicationStreamModel& UpdatedModel UE_LIFETIMEBOUND);
+		~FChangeLevelHandler();
+
+	private:
+
+		/** The client model to update when the local editor changes maps. */
+		ConcertSharedSlate::IEditableReplicationStreamModel& UpdatedModel;
+
+		/** Path to the previously open world */
+		FSoftObjectPath PreviousWorldPath;
+		
+		void OnWorldDestroyed(UWorld* World);
+		void OnWorldAdded(UWorld* World) const;
+	};
+}
diff --git a/Engine/Plugins/Developer/Concert/ConcertApp/MultiUserClient/Source/MultiUserClient/Private/Replication/Misc/PreventReplicatedPropertyTransaction.cpp b/Engine/Plugins/Developer/Concert/ConcertApp/MultiUserClient/Source/MultiUserClient/Private/Replication/Misc/PreventReplicatedPropertyTransaction.cpp
new file mode 100644
--- /dev/null
+++ b/Engine/Plugins/Developer/Concert/ConcertApp/MultiUserClient/Source/MultiUserClient/Private/Replication/Misc/PreventReplicatedPropertyTransaction.cpp
@@ -0,0 +1,76 @@
+﻿// Copyright Epic Games, Inc. All Rights Reserved.
+
+#include "PreventReplicatedPropertyTransaction.h"
+
+#include "IConcertSyncClient.h"
+#include "Replication/Client/ReplicationClientManager.h"
+#include "Replication/Util/GlobalAuthorityCache.h"
+
+#include "Algo/AnyOf.h"
+#include "Replication/Muting/MuteStateManager.h"
+#include "UObject/SoftObjectPath.h"
+
+namespace UE::MultiUserClient
+{
+	namespace Private
+	{
+		FORCEINLINE FName GetFilterName() { return TEXT("MultiUserReplicationFilter"); }
+	}
+	
+	FPreventReplicatedPropertyTransaction::FPreventReplicatedPropertyTransaction(
+		IConcertSyncClient& InSyncClient,
+		FReplicationClientManager& InClientManager,
+		FMuteStateManager& InMuteManager
+		)
+		: SyncClient(InSyncClient)
+		, ClientManager(InClientManager)
+		, MuteManager(InMuteManager)
+	{
+		SyncClient.GetTransactionBridge()->RegisterTransactionFilter(
+			Private::GetFilterName(),
+			FOnFilterTransactionDelegate::CreateRaw(this, &FPreventReplicatedPropertyTransaction::FilterTransactionAffectedByReplication)
+			);
+	}
+
+	FPreventReplicatedPropertyTransaction::~FPreventReplicatedPropertyTransaction()
+	{
+		SyncClient.GetTransactionBridge()->UnregisterTransactionFilter(Private::GetFilterName());
+	}
+
+	ETransactionFilterResult FPreventReplicatedPropertyTransaction::FilterTransactionAffectedByReplication(const FConcertTransactionFilterArgs& FilterArgs)
+	{
+		// 1. If an object is muted, we can save ourselves the work of analysing client streams.
+		// 2. We do not know whether remote clients have sync control (see below) but we can approximate that they don't have it if the object is muted.
+		const bool bIsMuted = MuteManager.GetSynchronizer().IsMuted(FilterArgs.ObjectToFilter);
+		
+		// FTransactionObjectEvent::GetChangedProperties() only contains root properties, like RelativeLocation; sub-properties, like RelativeLocation.X, are not listed.
+		// Luckily, replication streams list every parent property (so if RelativeLocation.X is being replicated then RelativeLocation is also in the stream).
+		const bool bIsReplicating = !bIsMuted && Algo::AnyOf(FilterArgs.TransactionEvent.GetChangedProperties(), [this, &FilterArgs](const FName& RootProperty)
+		{
+			const FGlobalAuthorityCache& AuthorityCache = ClientManager.GetAuthorityCache();
+			const TOptional<FGuid> AuthoringClient = AuthorityCache.GetClientWithAuthorityOverProperty(FilterArgs.ObjectToFilter, { RootProperty });
+
+			const bool bHasAuthoringClient = AuthoringClient.IsSet();
+			const bool bIsLocalClient = bHasAuthoringClient && *AuthoringClient == ClientManager.GetLocalClient().GetEndpointId();
+			if (bIsLocalClient)
+			{
+				const IConcertClientReplicationManager* ReplicationManager = SyncClient.GetReplicationManager();
+				check(ReplicationManager);
+				
+				// Just because a client has authority, it does not mean they're replicating.
+				// Replication is only happening when a client also has sync control.
+				// Since we only query sync control of the local client, we only know for our local client whether it has sync control...
+				const FGuid ClientStreamId = ClientManager.GetLocalClient().GetStreamSynchronizer().GetStreamId();
+				return ReplicationManager->HasSyncControl({ ClientStreamId, FilterArgs.ObjectToFilter});
+			}
+
+			// ... and for remote clients we'll just assume they're replicating when they have authority.
+			return bHasAuthoringClient;
+		});
+
+		// If the property is being replicated, it should not be transacted because the two systems may interfere.
+		return bIsReplicating
+			? ETransactionFilterResult::ExcludeObject
+			: ETransactionFilterResult::UseDefault;
+	}
+}
diff --git a/Engine/Plugins/Developer/Concert/ConcertApp/MultiUserClient/Source/MultiUserClient/Private/Replication/Misc/PreventReplicatedPropertyTransaction.h b/Engine/Plugins/Developer/Concert/ConcertApp/MultiUserClient/Source/MultiUserClient/Private/Replication/Misc/PreventReplicatedPropertyTransaction.h
new file mode 100644
--- /dev/null
+++ b/Engine/Plugins/Developer/Concert/ConcertApp/MultiUserClient/Source/MultiUserClient/Private/Replication/Misc/PreventReplicatedPropertyTransaction.h
@@ -0,0 +1,44 @@
+﻿// Copyright Epic Games, Inc. All Rights Reserved.
+
+#pragma once
+
+#include "IConcertClientTransactionBridge.h"
+#include "HAL/Platform.h"
+#include "Templates/UnrealTemplate.h"
+
+
+class IConcertSyncClient;
+enum class ETransactionFilterResult : uint8;
+struct FConcertTransactionFilterArgs;
+
+namespace UE::MultiUserClient
+{
+	class FMuteStateManager;
+	class FReplicationClientManager;
+	
+	/** Prevents generating transactions with properties that are being replicated by any other client. */
+	class FPreventReplicatedPropertyTransaction : public FNoncopyable
+	{
+	public:
+		
+		FPreventReplicatedPropertyTransaction(
+			IConcertSyncClient& InSyncClient UE_LIFETIMEBOUND,
+			FReplicationClientManager& InClientManager UE_LIFETIMEBOUND,
+			FMuteStateManager& InMuteManager UE_LIFETIMEBOUND
+			);
+		~FPreventReplicatedPropertyTransaction();
+
+	private:
+
+		/** Used to unsubscribe from the transaction bridge. */
+		IConcertSyncClient& SyncClient;
+
+		/** Used to obtain up-to-date registered properties from all clients. */
+		FReplicationClientManager& ClientManager;
+		/** Used to approximate whether remote clients may have sync control */
+		FMuteStateManager& MuteManager;
+		
+		ETransactionFilterResult FilterTransactionAffectedByReplication(const FConcertTransactionFilterArgs& FilterArgs);
+	};
+}
+
diff --git a/Engine/Plugins/Developer/Concert/ConcertApp/MultiUserClient/Source/MultiUserClient/Private/Replication/MultiUserReplicationManager.cpp b/Engine/Plugins/Developer/Concert/ConcertApp/MultiUserClient/Source/MultiUserClient/Private/Replication/MultiUserReplicationManager.cpp
--- a/Engine/Plugins/Developer/Concert/ConcertApp/MultiUserClient/Source/MultiUserClient/Private/Replication/MultiUserReplicationManager.cpp
+++ b/Engine/Plugins/Developer/Concert/ConcertApp/MultiUserClient/Source/MultiUserClient/Private/Replication/MultiUserReplicationManager.cpp
@@ -201,6 +201,7 @@ namespace UE::MultiUserClient
 		, ClientManager(InClient, InClient->GetConcertClient()->GetCurrentSession().ToSharedRef(), InDiscoveryContainer, QueryService.GetStreamAndAuthorityQueryService())
 		, MuteManager(*InClient, QueryService.GetMuteStateQueryService(), ClientManager.GetAuthorityCache())
 		, ChangeLevelHandler(ClientManager.GetLocalClient().GetClientEditModel().Get())
+		, PreventReplicatedPropertyTransaction(*InClient, ClientManager, MuteManager)
 		, UserNotifier(ClientManager, MuteManager)
 	{}
 }
diff --git a/Engine/Plugins/Developer/Concert/ConcertApp/MultiUserClient/Source/MultiUserClient/Private/Replication/MultiUserReplicationManager.h b/Engine/Plugins/Developer/Concert/ConcertApp/MultiUserClient/Source/MultiUserClient/Private/Replication/MultiUserReplicationManager.h
--- a/Engine/Plugins/Developer/Concert/ConcertApp/MultiUserClient/Source/MultiUserClient/Private/Replication/MultiUserReplicationManager.h
+++ b/Engine/Plugins/Developer/Concert/ConcertApp/MultiUserClient/Source/MultiUserClient/Private/Replication/MultiUserReplicationManager.h
@@ -4,14 +4,15 @@
 
 #include "Client/ReplicationClientManager.h"
 #include "IConcertSession.h"
+#include "Misc/ChangeLevelHandler.h"
+#include "Muting/MuteStateManager.h"
+#include "Notification/ReplicationUserNotifier.h"
 #include "Replication/IMultiUserReplication.h"
 #include "Replication/Stream/Discovery/ReplicationDiscoveryContainer.h"
-#include "UnrealEditor/ChangeLevelHandler.h"
 #include "Util/Query/RegularQueryService.h"
 
 #include "Misc/Optional.h"
-#include "Muting/MuteStateManager.h"
-#include "Notification/ReplicationUserNotifier.h"
+#include "Misc/PreventReplicatedPropertyTransaction.h"
 #include "Templates/SharedPointer.h"
 #include "Templates/UnrealTemplate.h"
 
@@ -111,13 +112,13 @@ namespace UE::MultiUserClient
 			 * Only valid when ConnectionState == EMultiUserReplicationConnectionState::Connected.
 			 */
 			FReplicationClientManager ClientManager;
-
 			/** Interacts with the mute global server mute system. */
 			FMuteStateManager MuteManager;
 
 			/** Clears local client's registered objects when leaving map. */
 			FChangeLevelHandler ChangeLevelHandler;
-
+			/** Prevents recording of transactions that change properties that are being replicated by a client. */
+			FPreventReplicatedPropertyTransaction PreventReplicatedPropertyTransaction;
 			/** This system notifies users when requests go wrong */
 			FReplicationUserNotifier UserNotifier;
 			
diff --git a/Engine/Plugins/Developer/Concert/ConcertApp/MultiUserClient/Source/MultiUserClient/Private/Replication/Muting/MuteStateManager.h b/Engine/Plugins/Developer/Concert/ConcertApp/MultiUserClient/Source/MultiUserClient/Private/Replication/Muting/MuteStateManager.h
--- a/Engine/Plugins/Developer/Concert/ConcertApp/MultiUserClient/Source/MultiUserClient/Private/Replication/Muting/MuteStateManager.h
+++ b/Engine/Plugins/Developer/Concert/ConcertApp/MultiUserClient/Source/MultiUserClient/Private/Replication/Muting/MuteStateManager.h
@@ -27,6 +27,8 @@ namespace UE::MultiUserClient
 
 		/** @return Object with which the objects' mute state can be changed and queried. */
 		FMuteChangeTracker& GetChangeTracker() { return ChangeTracker; }
+		/** @return Object with which you can query whether an object is muted. */
+		const FMuteStateSynchronizer& GetSynchronizer() const { return MuteStateSynchronizer; }
 		
 		/** Broadcasts when a mute request fails. */
 		DECLARE_MULTICAST_DELEGATE_TwoParams(FOnMuteRequestFailure, const FConcertReplication_ChangeMuteState_Request&, const FConcertReplication_ChangeMuteState_Response&);
diff --git a/Engine/Plugins/Developer/Concert/ConcertApp/MultiUserClient/Source/MultiUserClient/Private/Replication/UnrealEditor/ChangeLevelHandler.cpp b/Engine/Plugins/Developer/Concert/ConcertApp/MultiUserClient/Source/MultiUserClient/Private/Replication/UnrealEditor/ChangeLevelHandler.cpp
deleted file mode 100644
diff --git a/Engine/Plugins/Developer/Concert/ConcertApp/MultiUserClient/Source/MultiUserClient/Private/Replication/UnrealEditor/ChangeLevelHandler.h b/Engine/Plugins/Developer/Concert/ConcertApp/MultiUserClient/Source/MultiUserClient/Private/Replication/UnrealEditor/ChangeLevelHandler.h
deleted file mode 100644
diff --git a/Engine/Plugins/Developer/Concert/ConcertApp/MultiUserClient/Source/MultiUserClient/Private/Replication/Util/GlobalAuthorityCache.cpp b/Engine/Plugins/Developer/Concert/ConcertApp/MultiUserClient/Source/MultiUserClient/Private/Replication/Util/GlobalAuthorityCache.cpp
--- a/Engine/Plugins/Developer/Concert/ConcertApp/MultiUserClient/Source/MultiUserClient/Private/Replication/Util/GlobalAuthorityCache.cpp
+++ b/Engine/Plugins/Developer/Concert/ConcertApp/MultiUserClient/Source/MultiUserClient/Private/Replication/Util/GlobalAuthorityCache.cpp
@@ -137,10 +137,12 @@ namespace UE::MultiUserClient
 		return Conflict == EAuthorityConflict::Allowed;
 	}
 
-	TOptional<FGuid> FGlobalAuthorityCache::GetClientWithAuthorityOverProperty(const FSoftObjectPath& Object, const FConcertPropertyChain& Property) const
+	TOptional<FGuid> FGlobalAuthorityCache::GetClientWithAuthorityOverProperty(const FSoftObjectPath& Object, const TArray<FName>& PropertyChain) const
 	{
 		TOptional<FGuid> Result;
-		ForEachClientWithAuthorityOverObject(Object, [this, &Object, &Property, &Result](const FGuid& ClientId)
+		const uint32 Hash = ConcertSyncCore::ComputeHashForPropertyChainContent(PropertyChain);
+		
+		ForEachClientWithAuthorityOverObject(Object, [this, &Object, &PropertyChain, &Result, Hash](const FGuid& ClientId)
 		{
 			const FReplicationClient* Client = ClientManager.FindClient(ClientId);
 			if (!ensureMsgf(Client, TEXT("OnPreRemoteClientRemoved should have updated OwnedObjectsToClients")))
@@ -150,7 +152,7 @@ namespace UE::MultiUserClient
 			
 			const FConcertReplicatedObjectInfo* ObjectInfo = Client->GetStreamSynchronizer().GetServerState().ReplicatedObjects.Find(Object);
 			const bool bHasObjectRegistered = ensureMsgf(ObjectInfo, TEXT("OnStreamChanged should have updated OwnedObjectsToClients"))
-				&& ObjectInfo->PropertySelection.ReplicatedProperties.Contains(Property);
+				&& ObjectInfo->PropertySelection.ReplicatedProperties.ContainsByHash(Hash, PropertyChain);
 			const bool bHasAuthority = Client->GetAuthoritySynchronizer().HasAuthorityOver(Object);
 			if (bHasObjectRegistered && bHasAuthority)
 			{
diff --git a/Engine/Plugins/Developer/Concert/ConcertApp/MultiUserClient/Source/MultiUserClient/Private/Replication/Util/GlobalAuthorityCache.h b/Engine/Plugins/Developer/Concert/ConcertApp/MultiUserClient/Source/MultiUserClient/Private/Replication/Util/GlobalAuthorityCache.h
--- a/Engine/Plugins/Developer/Concert/ConcertApp/MultiUserClient/Source/MultiUserClient/Private/Replication/Util/GlobalAuthorityCache.h
+++ b/Engine/Plugins/Developer/Concert/ConcertApp/MultiUserClient/Source/MultiUserClient/Private/Replication/Util/GlobalAuthorityCache.h
@@ -55,7 +55,9 @@ namespace UE::MultiUserClient
 		bool CanClientAddProperty(const FSoftObjectPath& Object, const FGuid& ClientId, const FConcertPropertyChain& Chain) const;
 
 		/** Gets the client that has authority over the given property, if there is any. */
-		TOptional<FGuid> GetClientWithAuthorityOverProperty(const FSoftObjectPath& Object, const FConcertPropertyChain& Property) const;
+		TOptional<FGuid> GetClientWithAuthorityOverProperty(const FSoftObjectPath& Object, const FConcertPropertyChain& Property) const { return GetClientWithAuthorityOverProperty(Object, Property.GetPathToProperty()); }
+		/** Gets the client that has authority over the given property, if there is any. */
+		TOptional<FGuid> GetClientWithAuthorityOverProperty(const FSoftObjectPath& Object, const TArray<FName>& PropertyChain) const;
 		
 		/**
 		 * Removes entries from Request that would generate conflicts.
diff --git a/Engine/Plugins/Developer/Concert/ConcertSync/ConcertSyncCore/Source/ConcertSyncCore/Private/Replication/Data/ConcertPropertySelection.cpp b/Engine/Plugins/Developer/Concert/ConcertSync/ConcertSyncCore/Source/ConcertSyncCore/Private/Replication/Data/ConcertPropertySelection.cpp
--- a/Engine/Plugins/Developer/Concert/ConcertSync/ConcertSyncCore/Source/ConcertSyncCore/Private/Replication/Data/ConcertPropertySelection.cpp
+++ b/Engine/Plugins/Developer/Concert/ConcertSync/ConcertSyncCore/Source/ConcertSyncCore/Private/Replication/Data/ConcertPropertySelection.cpp
@@ -229,7 +229,7 @@ bool FConcertPropertySelection::EnumeratePropertyOverlaps(
     return !Intersection.IsEmpty();
 }
 
-uint32 GetTypeHash(const FConcertPropertyChain& Chain)
+uint32 UE::ConcertSyncCore::ComputeHashForPropertyChainContent(const TArray<FName>& PropertyChain)
 {
-    return GetTypeHash(Chain.GetPathToProperty());
-}
\ No newline at end of file
+    return GetTypeHash(PropertyChain);
+}
diff --git a/Engine/Plugins/Developer/Concert/ConcertSync/ConcertSyncCore/Source/ConcertSyncCore/Public/Replication/Data/ConcertPropertySelection.h b/Engine/Plugins/Developer/Concert/ConcertSync/ConcertSyncCore/Source/ConcertSyncCore/Public/Replication/Data/ConcertPropertySelection.h
--- a/Engine/Plugins/Developer/Concert/ConcertSync/ConcertSyncCore/Source/ConcertSyncCore/Public/Replication/Data/ConcertPropertySelection.h
+++ b/Engine/Plugins/Developer/Concert/ConcertSync/ConcertSyncCore/Source/ConcertSyncCore/Public/Replication/Data/ConcertPropertySelection.h
@@ -62,32 +62,14 @@ struct CONCERTSYNCCORE_API FConcertPropertyChain
 	};
 	FString ToString(EToStringMethod Method = EToStringMethod::Path) const;
 	
-	friend bool operator==(const FConcertPropertyChain& Left, const FConcertPropertyChain& Right)
-	{
-		return Left.PathToProperty == Right.PathToProperty;
-	}
-	friend bool operator!=(const FConcertPropertyChain& Left, const FConcertPropertyChain& Right)
-	{
-		return !(Left == Right);
-	}
+	friend bool operator==(const FConcertPropertyChain& Left, const FConcertPropertyChain& Right) { return Left.PathToProperty == Right.PathToProperty; }
+	friend bool operator!=(const FConcertPropertyChain& Left, const FConcertPropertyChain& Right) { return !(Left == Right); }
 
-	friend bool operator==(const FConcertPropertyChain& Left, const TArray<FName>& Path)
-	{
-		return Left.PathToProperty == Path;
-	}
-	friend bool operator!=(const FConcertPropertyChain& Left, const TArray<FName>& Path)
-	{
-		return Left.PathToProperty != Path;
-	}
+	friend bool operator==(const FConcertPropertyChain& Left, const TArray<FName>& Path) { return Left.PathToProperty == Path; }
+	friend bool operator!=(const FConcertPropertyChain& Left, const TArray<FName>& Path) { return Left.PathToProperty != Path; }
 	
-	friend bool operator==(const TArray<FName>& Path, const FConcertPropertyChain& Left)
-	{
-		return Left.PathToProperty == Path;
-	}
-	friend bool operator!=(const TArray<FName>& Path, const FConcertPropertyChain& Left)
-	{
-		return Left.PathToProperty != Path;
-	}
+	friend bool operator==(const TArray<FName>& Path, const FConcertPropertyChain& Left) { return Left.PathToProperty == Path; }
+	friend bool operator!=(const TArray<FName>& Path, const FConcertPropertyChain& Left) { return Left.PathToProperty != Path; }
 	
 private:
 	
@@ -186,4 +168,17 @@ struct CONCERTSYNCCORE_API FConcertPropertySelection
 	}
 };
 
-CONCERTSYNCCORE_API uint32 GetTypeHash(const FConcertPropertyChain& Chain);
\ No newline at end of file
+namespace UE::ConcertSyncCore
+{
+	/**
+	 * Implementation of hashing FConcertPropertyChain. Allows you to use TSet::ContainsByHash without constructing a FConcertPropertyChain, which is expensive because it searches the property tree.
+	 * You can rely on the fact that this function is either updated or deprecated when the hasing algorithm for FConcertPropertyChain is changed.
+	 */
+	CONCERTSYNCCORE_API uint32 ComputeHashForPropertyChainContent(const TArray<FName>& PropertyChain);
+}
+inline uint32 GetTypeHash(const FConcertPropertyChain& Chain)
+{
+	// If you need to changing the hashing function - update ComputeHashForPropertyChainContent since some code relies on the hasing logic.
+	return UE::ConcertSyncCore::ComputeHashForPropertyChainContent(Chain.GetPathToProperty());
+}
+
