Move several files into Replication/Misc files. This reduces the number of folders and makes it easier to understand / organize.

diff --git a/Engine/Plugins/Developer/Concert/ConcertApp/MultiUserClient/Source/MultiUserClient/Private/Replication/Authority/AuthorityChangeTracker.cpp b/Engine/Plugins/Developer/Concert/ConcertApp/MultiUserClient/Source/MultiUserClient/Private/Replication/Authority/AuthorityChangeTracker.cpp
--- a/Engine/Plugins/Developer/Concert/ConcertApp/MultiUserClient/Source/MultiUserClient/Private/Replication/Authority/AuthorityChangeTracker.cpp
+++ b/Engine/Plugins/Developer/Concert/ConcertApp/MultiUserClient/Source/MultiUserClient/Private/Replication/Authority/AuthorityChangeTracker.cpp
@@ -4,7 +4,7 @@
 
 #include "EAuthorityMutability.h"
 #include "IClientAuthoritySynchronizer.h"
-#include "Replication/Util/GlobalAuthorityCache.h"
+#include "Replication/Misc/GlobalAuthorityCache.h"
 
 namespace UE::MultiUserClient
 {
diff --git a/Engine/Plugins/Developer/Concert/ConcertApp/MultiUserClient/Source/MultiUserClient/Private/Replication/Authority/AuthoritySynchronizer_RemoteClient.cpp b/Engine/Plugins/Developer/Concert/ConcertApp/MultiUserClient/Source/MultiUserClient/Private/Replication/Authority/AuthoritySynchronizer_RemoteClient.cpp
--- a/Engine/Plugins/Developer/Concert/ConcertApp/MultiUserClient/Source/MultiUserClient/Private/Replication/Authority/AuthoritySynchronizer_RemoteClient.cpp
+++ b/Engine/Plugins/Developer/Concert/ConcertApp/MultiUserClient/Source/MultiUserClient/Private/Replication/Authority/AuthoritySynchronizer_RemoteClient.cpp
@@ -2,7 +2,7 @@
 
 #include "AuthoritySynchronizer_RemoteClient.h"
 
-#include "Replication/Util/Query/StreamAndAuthorityQueryService.h"
+#include "Replication/Misc/Query/StreamAndAuthorityQueryService.h"
 
 namespace UE::MultiUserClient
 {
diff --git a/Engine/Plugins/Developer/Concert/ConcertApp/MultiUserClient/Source/MultiUserClient/Private/Replication/Client/ReplicationClient.h b/Engine/Plugins/Developer/Concert/ConcertApp/MultiUserClient/Source/MultiUserClient/Private/Replication/Client/ReplicationClient.h
--- a/Engine/Plugins/Developer/Concert/ConcertApp/MultiUserClient/Source/MultiUserClient/Private/Replication/Client/ReplicationClient.h
+++ b/Engine/Plugins/Developer/Concert/ConcertApp/MultiUserClient/Source/MultiUserClient/Private/Replication/Client/ReplicationClient.h
@@ -5,7 +5,7 @@
 #include "Replication/Authority/AuthorityChangeTracker.h"
 #include "Replication/Authority/IClientAuthoritySynchronizer.h"
 #include "Replication/Editor/UnrealEditor/ModifyObjectInLevelHandler.h"
-#include "Replication/Frequency/FrequencyChangeTracker.h"
+#include "Replication/Misc/Frequency/FrequencyChangeTracker.h"
 #include "Replication/Stream/IClientStreamSynchronizer.h"
 #include "Replication/Stream/StreamChangeTracker.h"
 #include "Replication/Submission/AutoSubmissionPolicy.h"
diff --git a/Engine/Plugins/Developer/Concert/ConcertApp/MultiUserClient/Source/MultiUserClient/Private/Replication/Client/ReplicationClientManager.cpp b/Engine/Plugins/Developer/Concert/ConcertApp/MultiUserClient/Source/MultiUserClient/Private/Replication/Client/ReplicationClientManager.cpp
--- a/Engine/Plugins/Developer/Concert/ConcertApp/MultiUserClient/Source/MultiUserClient/Private/Replication/Client/ReplicationClientManager.cpp
+++ b/Engine/Plugins/Developer/Concert/ConcertApp/MultiUserClient/Source/MultiUserClient/Private/Replication/Client/ReplicationClientManager.cpp
@@ -7,7 +7,7 @@
 #include "RemoteReplicationClient.h"
 #include "Assets/MultiUserReplicationSessionPreset.h"
 #include "Replication/Stream/StreamSynchronizer_LocalClient.h"
-#include "Replication/Util/Query/StreamAndAuthorityQueryService.h"
+#include "Replication/Misc/Query/StreamAndAuthorityQueryService.h"
 
 #include "UObject/Package.h"
 #include "UObject/UObjectGlobals.h"
diff --git a/Engine/Plugins/Developer/Concert/ConcertApp/MultiUserClient/Source/MultiUserClient/Private/Replication/Client/ReplicationClientManager.h b/Engine/Plugins/Developer/Concert/ConcertApp/MultiUserClient/Source/MultiUserClient/Private/Replication/Client/ReplicationClientManager.h
--- a/Engine/Plugins/Developer/Concert/ConcertApp/MultiUserClient/Source/MultiUserClient/Private/Replication/Client/ReplicationClientManager.h
+++ b/Engine/Plugins/Developer/Concert/ConcertApp/MultiUserClient/Source/MultiUserClient/Private/Replication/Client/ReplicationClientManager.h
@@ -5,10 +5,9 @@
 #include "Assets/MultiUserReplicationSessionPreset.h"
 #include "ConcertMessageData.h"
 #include "LocalReplicationClient.h"
-#include "ReplicationClient.h"
+#include "Replication/Misc/GlobalAuthorityCache.h"
 #include "Replication/Submission/MultiEdit/ReassignObjectPropertiesLogic.h"
-#include "Replication/Submission/Notification/SubmissionNotifier.h"
-#include "Replication/Util/GlobalAuthorityCache.h"
+#include "ReplicationClient.h"
 
 #include "UObject/GCObject.h"
 #include "Templates/UnrealTemplate.h"
diff --git a/Engine/Plugins/Developer/Concert/ConcertApp/MultiUserClient/Source/MultiUserClient/Private/Replication/Frequency/FrequencyChangeTracker.cpp b/Engine/Plugins/Developer/Concert/ConcertApp/MultiUserClient/Source/MultiUserClient/Private/Replication/Frequency/FrequencyChangeTracker.cpp
deleted file mode 100644
diff --git a/Engine/Plugins/Developer/Concert/ConcertApp/MultiUserClient/Source/MultiUserClient/Private/Replication/Frequency/FrequencyChangeTracker.h b/Engine/Plugins/Developer/Concert/ConcertApp/MultiUserClient/Source/MultiUserClient/Private/Replication/Frequency/FrequencyChangeTracker.h
deleted file mode 100644
diff --git a/Engine/Plugins/Developer/Concert/ConcertApp/MultiUserClient/Source/MultiUserClient/Private/Replication/Misc/Frequency/FrequencyChangeTracker.cpp b/Engine/Plugins/Developer/Concert/ConcertApp/MultiUserClient/Source/MultiUserClient/Private/Replication/Misc/Frequency/FrequencyChangeTracker.cpp
new file mode 100644
--- /dev/null
+++ b/Engine/Plugins/Developer/Concert/ConcertApp/MultiUserClient/Source/MultiUserClient/Private/Replication/Misc/Frequency/FrequencyChangeTracker.cpp
@@ -0,0 +1,69 @@
+// Copyright Epic Games, Inc. All Rights Reserved.
+
+#include "FrequencyChangeTracker.h"
+
+#include "Replication/Stream/IClientStreamSynchronizer.h"
+#include "Settings/MultiUserReplicationSettings.h"
+
+namespace UE::MultiUserClient
+{
+	FFrequencyChangeTracker::FFrequencyChangeTracker(
+		IClientStreamSynchronizer& InStreamSynchronizer
+		)
+		: StreamSynchronizer(InStreamSynchronizer)
+	{
+		// Refresh locally recorded changes when the server state changes, e.g. after a request has been served.
+		StreamSynchronizer.OnServerStateChanged().AddRaw(this, &FFrequencyChangeTracker::RefreshChanges);
+	}
+
+	FFrequencyChangeTracker::~FFrequencyChangeTracker()
+	{
+		StreamSynchronizer.OnServerStateChanged().RemoveAll(this);
+	}
+
+	void FFrequencyChangeTracker::AddOverride(FSoftObjectPath Object, FConcertObjectReplicationSettings NewSettings)
+	{
+		RecordedChanges.OverridesToAdd.Emplace(MoveTemp(Object), MoveTemp(NewSettings));
+	}
+	
+	FFrequencyChangelist FFrequencyChangeTracker::BuildForSubmission(const FStreamChangelist& ObjectChanges)
+	{
+		FFrequencyChangelist Result = RecordedChanges;
+
+		// For now we only support adding overrides.
+		const FGuid StreamId = StreamSynchronizer.GetStreamId();
+		const FConcertObjectReplicationMap& ReplicationMap = StreamSynchronizer.GetServerState();
+		for (auto IteratorForAdds = Result.OverridesToAdd.CreateIterator(); IteratorForAdds; ++IteratorForAdds)
+		{
+			const TPair<FSoftObjectPath, FConcertObjectReplicationSettings>& Pair = *IteratorForAdds;
+			if (!ReplicationMap.HasProperties(Pair.Key) && !ObjectChanges.ObjectsToPut.Contains({ StreamId, Pair.Key }))
+			{
+				IteratorForAdds.RemoveCurrent();
+			}
+		}
+
+		// Set frequency defaults if either this is the first time we're submitting to the server or the user has changes the project settings since.
+		const FConcertObjectReplicationSettings& DefaultsToSet = UMultiUserReplicationSettings::Get()->FrequencyRules.DefaultObjectFrequencySettings;
+		if (DefaultsToSet != StreamSynchronizer.GetFrequencySettings().Defaults)
+		{
+			Result.NewDefaults = DefaultsToSet;
+		}
+		
+		return Result;
+	}
+
+	void FFrequencyChangeTracker::RefreshChanges()
+	{
+		const FConcertObjectReplicationMap& ReplicationMap = StreamSynchronizer.GetServerState();
+		for (auto IteratorForAdds = RecordedChanges.OverridesToAdd.CreateIterator(); IteratorForAdds; ++IteratorForAdds)
+		{
+			const TPair<FSoftObjectPath, FConcertObjectReplicationSettings>& Pair = *IteratorForAdds;
+			if (!ReplicationMap.HasProperties(Pair.Key))
+			{
+				IteratorForAdds.RemoveCurrent();
+			}
+		}
+
+		// In the future we must also cleans RecordedChanges.OverridesToRemove if we add a function like RemoveOverride
+	}
+}
diff --git a/Engine/Plugins/Developer/Concert/ConcertApp/MultiUserClient/Source/MultiUserClient/Private/Replication/Misc/Frequency/FrequencyChangeTracker.h b/Engine/Plugins/Developer/Concert/ConcertApp/MultiUserClient/Source/MultiUserClient/Private/Replication/Misc/Frequency/FrequencyChangeTracker.h
new file mode 100644
--- /dev/null
+++ b/Engine/Plugins/Developer/Concert/ConcertApp/MultiUserClient/Source/MultiUserClient/Private/Replication/Misc/Frequency/FrequencyChangeTracker.h
@@ -0,0 +1,55 @@
+// Copyright Epic Games, Inc. All Rights Reserved.
+
+#pragma once
+
+#include "Replication/Misc/Util/StreamRequestUtils.h"
+
+#include "Templates/UnrealTemplate.h"
+#include "UObject/UnrealType.h"
+
+struct FConcertObjectReplicationSettings;
+struct FSoftObjectPath;
+
+namespace UE::MultiUserClient
+{
+	class IClientStreamSynchronizer;
+	class FStreamChangeTracker;
+	
+	/** Used to queue up local made changes to object frequency settings. */
+	class FFrequencyChangeTracker : public FNoncopyable
+    {
+    public:
+
+		FFrequencyChangeTracker(
+			IClientStreamSynchronizer& InStreamSynchronizer UE_LIFETIMEBOUND
+			);
+		~FFrequencyChangeTracker();
+
+		/** Adds a frequency override for the given object. */
+		void AddOverride(FSoftObjectPath Object, FConcertObjectReplicationSettings NewSettings);
+		
+		// In the future we could add RemoveOverride if needed
+
+		/** Compares the recorded changes to about to be submitted ObjectChanges and strips any invalid changes. */
+		FFrequencyChangelist BuildForSubmission(const FStreamChangelist& ObjectChanges);
+
+		/** Compares the local changes against the server state and optionally removes changes. */
+		void RefreshChanges();
+
+		DECLARE_MULTICAST_DELEGATE(FOnFrequencySettingsChanged);
+		/** Called when frequency settings have changed been recorded that need to be submitted to the server. */
+		FOnFrequencySettingsChanged& OnFrequencySettingsChanged() { return OnFrequencySettingsChangedDelegate; } 
+
+	private:
+
+		/** Tells us when the server state changes so RecordedChanges can be cleansed of applied changes. */
+		IClientStreamSynchronizer& StreamSynchronizer;
+
+		/** The changes recorded thus far. */
+		FFrequencyChangelist RecordedChanges;
+
+		/** Called when frequency settings have changed been recorded that need to be submitted to the server. */
+		FOnFrequencySettingsChanged OnFrequencySettingsChangedDelegate;
+    };
+}
+
diff --git a/Engine/Plugins/Developer/Concert/ConcertApp/MultiUserClient/Source/MultiUserClient/Private/Replication/Misc/Frequency/FrequencyUtils.cpp b/Engine/Plugins/Developer/Concert/ConcertApp/MultiUserClient/Source/MultiUserClient/Private/Replication/Misc/Frequency/FrequencyUtils.cpp
new file mode 100644
--- /dev/null
+++ b/Engine/Plugins/Developer/Concert/ConcertApp/MultiUserClient/Source/MultiUserClient/Private/Replication/Misc/Frequency/FrequencyUtils.cpp
@@ -0,0 +1,48 @@
+// Copyright Epic Games, Inc. All Rights Reserved.
+
+#include "FrequencyUtils.h"
+
+namespace UE::MultiUserClient::FrequencyUtils
+{
+	TSharedPtr<IParallelSubmissionOperation> SetFrequencySettingForClients(
+		const FSoftObjectPath& ContextObject,
+		const FInlineClientArray& Clients,
+		FReplicationClientManager& InClientManager,
+		TFunctionRef<void(FConcertObjectReplicationSettings&)> ApplySettingChange
+		)
+	{
+		TMap<FGuid, FConcertReplication_ChangeStream_Request> ParallelOperations;
+		
+		for (const FGuid& ClientId : Clients)
+		{
+			FReplicationClient* Client = InClientManager.FindClient(ClientId);
+			if (Client && CanChangeFrequencySettings(ContextObject, *Client))
+			{
+				const FGuid StreamId = Client->GetStreamSynchronizer().GetStreamId();
+				const FConcertStreamFrequencySettings& AllSettings = Client->GetStreamSynchronizer().GetFrequencySettings();
+				FConcertObjectReplicationSettings PerObjectOverride = AllSettings.GetSettingsFor(ContextObject);
+
+				const bool bNeedsToApplyDefault = !AllSettings.ObjectOverrides.Contains(ContextObject);
+				// TODO UE-203350: Specify default replication settings from Multi User project settings
+				PerObjectOverride = bNeedsToApplyDefault ? FConcertObjectReplicationSettings{ EConcertObjectReplicationMode::SpecifiedRate, 30 } : PerObjectOverride;
+				ApplySettingChange(PerObjectOverride);
+				if (!ensure(PerObjectOverride.ReplicationRate > 0))
+				{
+					continue;
+				}
+				
+				FConcertReplication_ChangeStream_Frequency FrequencyChange;
+				FrequencyChange.OverridesToAdd.Add(ContextObject, PerObjectOverride);
+				FConcertReplication_ChangeStream_Request StreamChangeRequest;
+				StreamChangeRequest.FrequencyChanges.Add(StreamId, FrequencyChange);
+				
+				ParallelOperations.Add(ClientId, StreamChangeRequest);
+			}
+		}
+
+		return ParallelOperations.IsEmpty()
+			? nullptr
+			: ExecuteParallelStreamChanges(InClientManager, ParallelOperations);
+	}
+	
+}
\ No newline at end of file
diff --git a/Engine/Plugins/Developer/Concert/ConcertApp/MultiUserClient/Source/MultiUserClient/Private/Replication/Misc/Frequency/FrequencyUtils.h b/Engine/Plugins/Developer/Concert/ConcertApp/MultiUserClient/Source/MultiUserClient/Private/Replication/Misc/Frequency/FrequencyUtils.h
new file mode 100644
--- /dev/null
+++ b/Engine/Plugins/Developer/Concert/ConcertApp/MultiUserClient/Source/MultiUserClient/Private/Replication/Misc/Frequency/FrequencyUtils.h
@@ -0,0 +1,130 @@
+// Copyright Epic Games, Inc. All Rights Reserved.
+
+#pragma once
+
+#include "Replication/Client/ReplicationClientManager.h"
+#include "Replication/Data/ReplicationFrequencySettings.h"
+
+#include "Algo/AllOf.h"
+#include "Async/Future.h"
+#include "Misc/Optional.h"
+#include "Templates/Function.h"
+#include "Templates/SharedPointer.h"
+
+#include <type_traits>
+
+namespace UE::MultiUserClient
+{
+	class IParallelSubmissionOperation;
+}
+
+namespace UE::MultiUserClient::FrequencyUtils
+{
+	// Low effort optimization to minimize heap allocations
+	using FInlineClientArray = TArray<FGuid, TInlineAllocator<8>>;
+	
+	/** Gets the shared value from FConcertObjectReplicationSettings, if all clients share the same value. Transform selects the member. */
+	template<typename TSubSetting, typename TSelectSubSetting>
+	TOptional<TSubSetting> FindSharedFrequencySetting(
+		const FSoftObjectPath& ContextObject,
+		const FInlineClientArray& Clients,
+		const FReplicationClientManager& InClientManager,
+		TSelectSubSetting&& Transform
+		)
+	requires std::is_invocable_r_v<TSubSetting, TSelectSubSetting, const FConcertObjectReplicationSettings&>;
+
+	enum class EChangeResult : uint8 { Success, Failure };
+	/**
+	 * Goes through all clients and sets a FConcertObjectReplicationSettings for ContextObject. The setting is applied by invoking ApplySettingChange.
+	 * @return Operations object that informs you when the change operations have completed.
+	 */
+	TSharedPtr<IParallelSubmissionOperation> SetFrequencySettingForClients(
+		const FSoftObjectPath& ContextObject,
+		const FInlineClientArray& Clients,
+		FReplicationClientManager& InClientManager,
+		TFunctionRef<void(FConcertObjectReplicationSettings&)> ApplySettingChange
+		);
+
+	/** Gets the replication rate of all clients, if they share the same value. */
+	inline TOptional<uint8> FindSharedFrequencyRate(const FSoftObjectPath& ContextObject, const FInlineClientArray& Clients, const FReplicationClientManager& InClientManager)
+	{
+		bool bHasRightMode = false;
+		const TOptional<uint8> SharedMode = FindSharedFrequencySetting<uint8>(ContextObject, Clients, InClientManager,
+			[&bHasRightMode](const FConcertObjectReplicationSettings& Setting)
+			{
+				bHasRightMode |= Setting.ReplicationMode == EConcertObjectReplicationMode::SpecifiedRate;
+				return Setting.ReplicationRate;
+			});
+		return bHasRightMode ? SharedMode : TOptional<uint8>{};
+	}
+
+	/** @return The replication mode that all clients are using, if they share the same value. */
+	inline TOptional<EConcertObjectReplicationMode> FindSharedReplicationMode(const FSoftObjectPath& ContextObject, const FInlineClientArray& Clients, const FReplicationClientManager& InClientManager)
+	{
+		return FindSharedFrequencySetting<EConcertObjectReplicationMode>(ContextObject, Clients, InClientManager,
+			[](const FConcertObjectReplicationSettings& Setting)
+			{
+				return Setting.ReplicationMode;
+			});
+	}
+	
+	/** @return Whether all clients have the given replication mode */
+	inline bool AllClientsHaveMode(const EConcertObjectReplicationMode SearchedMode, const FSoftObjectPath& ContextObject, const FInlineClientArray& Clients, const FReplicationClientManager& InClientManager) 
+	{
+		return SearchedMode == FindSharedReplicationMode(ContextObject, Clients, InClientManager);
+	}
+
+	/** @return Whether ContextObject's frequency settings can be changed for Client. */
+	inline bool CanChangeFrequencySettings(const FSoftObjectPath& ContextObject, const FReplicationClient& Client)
+	{
+		const bool bHasProperties = Client.GetStreamSynchronizer().GetServerState().HasProperties(ContextObject);
+		const bool bAllowsEditing = Client.AllowsEditing();
+		const bool bNotProcessingOtherRequests = Client.GetSubmissionWorkflow().CanSubmit();
+		return bHasProperties && bAllowsEditing && bNotProcessingOtherRequests;
+	}
+
+	/** @return Whether ContextObject's frequency settings can be changed for all given Clients. */
+	inline bool CanChangeFrequencySettings(const FSoftObjectPath& ContextObject, const FInlineClientArray& Clients, const FReplicationClientManager& InClientManager)
+	{
+		return Algo::AllOf(Clients, [ContextObject, &InClientManager](const FGuid& ClientId)
+		{
+			const FReplicationClient* Client = InClientManager.FindClient(ClientId);
+			return Client && CanChangeFrequencySettings(ContextObject, *Client);
+		});
+	}
+}
+
+namespace UE::MultiUserClient::FrequencyUtils
+{
+	template<typename TSubSetting, typename TSelectSubSetting>
+	TOptional<TSubSetting> FindSharedFrequencySetting(
+		const FSoftObjectPath& ContextObject,
+		const FInlineClientArray& Clients,
+		const FReplicationClientManager& InClientManager,
+		TSelectSubSetting&& Transform
+		)
+	requires std::is_invocable_r_v<TSubSetting, TSelectSubSetting, const FConcertObjectReplicationSettings&>
+	{
+		TOptional<TSubSetting> SharedSetting;
+		for (const FGuid& ClientId : Clients)
+		{
+			const FReplicationClient* Client = InClientManager.FindClient(ClientId);
+			if (!Client)
+			{
+				continue;
+			}
+
+			const FConcertObjectReplicationSettings& ObjectSettings = Client->GetStreamSynchronizer().GetFrequencySettings().GetSettingsFor(ContextObject);
+			const TSubSetting SubSetting = Transform(ObjectSettings);
+			if (SharedSetting && SharedSetting != SubSetting)
+			{
+				return {};
+			}
+			else
+			{
+				SharedSetting = SubSetting;
+			}
+		}
+		return SharedSetting;
+	}
+}
diff --git a/Engine/Plugins/Developer/Concert/ConcertApp/MultiUserClient/Source/MultiUserClient/Private/Replication/Misc/GlobalAuthorityCache.cpp b/Engine/Plugins/Developer/Concert/ConcertApp/MultiUserClient/Source/MultiUserClient/Private/Replication/Misc/GlobalAuthorityCache.cpp
new file mode 100644
--- /dev/null
+++ b/Engine/Plugins/Developer/Concert/ConcertApp/MultiUserClient/Source/MultiUserClient/Private/Replication/Misc/GlobalAuthorityCache.cpp
@@ -0,0 +1,270 @@
+// Copyright Epic Games, Inc. All Rights Reserved.
+
+#include "GlobalAuthorityCache.h"
+
+#include "Replication/AuthorityConflictSharedUtils.h"
+#include "Replication/Authority/EAuthorityMutability.h"
+#include "Replication/Client/RemoteReplicationClient.h"
+#include "Replication/Client/ReplicationClientManager.h"
+
+namespace UE::MultiUserClient
+{
+	FGlobalAuthorityCache::FGlobalAuthorityCache(FReplicationClientManager& InClientManager)
+		: ClientManager(InClientManager)
+	{}
+
+	void FGlobalAuthorityCache::RegisterEvents()
+	{
+		RegisterForClientEvents(ClientManager.GetLocalClient());
+
+		ClientManager.OnPostRemoteClientAdded().AddRaw(this, &FGlobalAuthorityCache::OnPostRemoteClientAdded);
+		ClientManager.OnPreRemoteClientRemoved().AddRaw(this, &FGlobalAuthorityCache::OnPreRemoteClientRemoved);
+	}
+
+	void FGlobalAuthorityCache::ForEachClientWithObjectInStream(const FSoftObjectPath& Object, TFunctionRef<EBreakBehavior(const FGuid& ClientId)> Callback) const
+	{
+		const TSet<FGuid>* Clients = RegisteredObjectsToClients.Find(Object);
+		if (!Clients)
+		{
+			return;
+		}
+
+		for (const FGuid& ClientId : *Clients)
+		{
+			if (Callback(ClientId) == EBreakBehavior::Break)
+			{
+				break;
+			}
+		}
+	}
+
+	bool FGlobalAuthorityCache::IsObjectOrChildReferenced(const FSoftObjectPath& Object) const
+	{
+		const FString ObjectPathString = Object.ToString();
+		for (const TPair<FSoftObjectPath, TSet<FGuid>>& Pair : RegisteredObjectsToClients)
+		{
+			if (Pair.Key.ToString().Contains(ObjectPathString))
+			{
+				return true;
+			}
+		}
+		return false;
+	}
+
+	void FGlobalAuthorityCache::ForEachClientWithAuthorityOverObject(const FSoftObjectPath& Object, TFunctionRef<EBreakBehavior(const FGuid& ClientId)> Callback) const
+	{
+		const TSet<FGuid>* Clients = OwnedObjectsToClients.Find(Object);
+		if (!Clients)
+		{
+			return;
+		}
+
+		for (const FGuid& ClientId : *Clients)
+		{
+			if (Callback(ClientId) == EBreakBehavior::Break)
+			{
+				break;
+			}
+		}
+	}
+
+	TArray<FGuid> FGlobalAuthorityCache::GetClientsWithAuthorityOverObject(const FSoftObjectPath& Object) const
+	{
+		TArray<FGuid> Result;
+		ForEachClientWithAuthorityOverObject(Object, [&Result](const FGuid& ClientId)
+		{
+			Result.Add(ClientId);
+			return EBreakBehavior::Continue;
+		});
+		return Result;
+	}
+
+	bool FGlobalAuthorityCache::HasAuthorityOverObject(const FSoftObjectPath& Object, const FGuid& ClientId) const
+	{
+		bool bHasAuthority = false;
+		ForEachClientWithAuthorityOverObject(Object, [&ClientId, &bHasAuthority](const FGuid& OtherClientId)
+		{
+			bHasAuthority = ClientId == OtherClientId;
+			return bHasAuthority ? EBreakBehavior::Break : EBreakBehavior::Continue;
+		});
+		return bHasAuthority;
+	}
+
+	EAuthorityMutability FGlobalAuthorityCache::CanClientTakeAuthorityAfterSubmission(const FSoftObjectPath& Object, const FGuid& ClientId, FProcessPropertyConflict ProcessConflict) const
+	{
+		const FReplicationClient* Client = ClientManager.FindClient(ClientId);
+		if (!ensure(Client))
+		{
+			return EAuthorityMutability::NotApplicable;
+		}
+
+		// Important: Get server state with local changes applied to it! CanClientTakeAuthority answers: "Can the client take authority after submitting?" 
+		const FConcertPropertySelection* PropertySelection = Client->GetStreamDiffer().GetPropertiesAfterSubmit(Object);
+		if (!PropertySelection)
+		{
+			// Nothing to take authority over
+			return EAuthorityMutability::NotApplicable;
+		}
+		
+		using namespace ConcertSyncCore::Replication::AuthorityConflictUtils;
+		const EAuthorityConflict Conflict = EnumerateAuthorityConflicts(
+			ClientId,
+			Object,
+			PropertySelection->ReplicatedProperties,
+			*this,
+			[&ProcessConflict](const FGuid& ClientId, const FGuid&, const FConcertPropertyChain& ConflictingProperty)
+			{
+				return ProcessConflict(ClientId, ConflictingProperty);
+			});
+		return Conflict == EAuthorityConflict::Allowed ? EAuthorityMutability::Allowed : EAuthorityMutability::Conflict;
+	}
+
+	bool FGlobalAuthorityCache::CanClientAddProperty(const FSoftObjectPath& Object, const FGuid& ClientId, const FConcertPropertyChain& Chain) const
+	{
+		const FReplicationClient* Client = ClientManager.FindClient(ClientId);
+		if (!ensure(Client))
+		{
+			return false;
+		}
+		
+		using namespace ConcertSyncCore::Replication::AuthorityConflictUtils;
+		const EAuthorityConflict Conflict = EnumerateAuthorityConflicts(
+			ClientId,
+			Object,
+			{ Chain },
+			*this
+			);
+		return Conflict == EAuthorityConflict::Allowed;
+	}
+
+	TOptional<FGuid> FGlobalAuthorityCache::GetClientWithAuthorityOverProperty(const FSoftObjectPath& Object, const TArray<FName>& PropertyChain) const
+	{
+		TOptional<FGuid> Result;
+		const uint32 Hash = ConcertSyncCore::ComputeHashForPropertyChainContent(PropertyChain);
+		
+		ForEachClientWithAuthorityOverObject(Object, [this, &Object, &PropertyChain, &Result, Hash](const FGuid& ClientId)
+		{
+			const FReplicationClient* Client = ClientManager.FindClient(ClientId);
+			if (!ensureMsgf(Client, TEXT("OnPreRemoteClientRemoved should have updated OwnedObjectsToClients")))
+			{
+				return EBreakBehavior::Continue;
+			}
+			
+			const FConcertReplicatedObjectInfo* ObjectInfo = Client->GetStreamSynchronizer().GetServerState().ReplicatedObjects.Find(Object);
+			const bool bHasObjectRegistered = ensureMsgf(ObjectInfo, TEXT("OnStreamChanged should have updated OwnedObjectsToClients"))
+				&& ObjectInfo->PropertySelection.ReplicatedProperties.ContainsByHash(Hash, PropertyChain);
+			const bool bHasAuthority = Client->GetAuthoritySynchronizer().HasAuthorityOver(Object);
+			if (bHasObjectRegistered && bHasAuthority)
+			{
+				Result = ClientId;
+				return EBreakBehavior::Break;
+			}
+			return EBreakBehavior::Continue;
+		});
+		return Result;
+	}
+
+	void FGlobalAuthorityCache::CleanseConflictsFromAuthorityRequest(FConcertReplication_ChangeAuthority_Request& Request, const FGuid& SendingClient) const
+	{
+		ConcertSyncCore::Replication::AuthorityConflictUtils::CleanseConflictsFromAuthorityRequest(Request, SendingClient, *this);
+	}
+
+	void FGlobalAuthorityCache::CleanseConflictsFromStreamRequest(FConcertReplication_ChangeStream_Request& Request, const FGuid& SendingClient) const
+	{
+		ConcertSyncCore::Replication::AuthorityConflictUtils::CleanseConflictsFromStreamRequest(Request, SendingClient, *this);
+	}
+
+	void FGlobalAuthorityCache::RegisterForClientEvents(const FReplicationClient& Client)
+	{
+		const FGuid& ClientEndpointId = Client.GetEndpointId();
+		Client.GetAuthoritySynchronizer().OnServerStateChanged().AddRaw(this, &FGlobalAuthorityCache::OnPostAuthorityChanged, ClientEndpointId);
+		Client.GetStreamSynchronizer().OnServerStateChanged().AddRaw(this, &FGlobalAuthorityCache::OnStreamChanged, ClientEndpointId);
+	}
+
+	void FGlobalAuthorityCache::UnregisterFromClientEvents(const FReplicationClient& Client) const
+	{
+		Client.GetAuthoritySynchronizer().OnServerStateChanged().RemoveAll(this);
+		Client.GetStreamSynchronizer().OnServerStateChanged().RemoveAll(this);
+	}
+
+	void FGlobalAuthorityCache::AddClient(const FGuid& ClientId)
+	{
+		const FReplicationClient* Client = ClientManager.FindClient(ClientId);
+		if (!ensure(Client))
+		{
+			return;
+		}
+		
+		const FConcertObjectReplicationMap& ClientObjectMap = Client->GetStreamSynchronizer().GetServerState();
+		const IClientAuthoritySynchronizer& ClientAuthority = Client->GetAuthoritySynchronizer();
+		for (const TPair<FSoftObjectPath, FConcertReplicatedObjectInfo>& StreamContents : ClientObjectMap.ReplicatedObjects)
+		{
+			const FSoftObjectPath& Object = StreamContents.Key;
+			RegisteredObjectsToClients.FindOrAdd(Object).Add(ClientId);
+			if (ClientAuthority.HasAuthorityOver(Object))
+			{
+				OwnedObjectsToClients.FindOrAdd(Object).Add(ClientId);
+			}
+		}
+	}
+
+	void FGlobalAuthorityCache::RemoveClient(const FGuid& ClientId)
+	{
+		for (auto It = OwnedObjectsToClients.CreateIterator(); It; ++It)
+		{
+			It->Value.Remove(ClientId);
+			if (It->Value.IsEmpty())
+			{
+				It.RemoveCurrent();
+			}
+		}
+
+		for (auto It = RegisteredObjectsToClients.CreateIterator(); It; ++It)
+		{
+			It->Value.Remove(ClientId);
+			if (It->Value.IsEmpty())
+			{
+				It.RemoveCurrent();
+			}
+		}
+	}
+	
+	void FGlobalAuthorityCache::ForEachStream(const FGuid& ClientEndpointId, TFunctionRef<EBreakBehavior(const FGuid& StreamId, const FConcertObjectReplicationMap& ReplicationMap)> Callback) const
+	{
+		const FReplicationClient* Client = ClientManager.FindClient(ClientEndpointId);
+		if (!ensure(Client))
+		{
+			return;
+		}
+
+		const IClientStreamSynchronizer& StreamSynchronizer = Client->GetStreamSynchronizer();
+		Callback(StreamSynchronizer.GetStreamId(), StreamSynchronizer.GetServerState());
+	}
+
+	void FGlobalAuthorityCache::ForEachSendingClient(TFunctionRef<EBreakBehavior(const FGuid& ClientEndpointId)> Callback) const
+	{
+		auto ProcessClient = [&Callback](const FReplicationClient& Client)
+		{
+			return Client.GetAuthoritySynchronizer().HasAnyAuthority()
+				? Callback(Client.GetEndpointId())
+				:EBreakBehavior::Continue;
+		};
+
+		if (ProcessClient(ClientManager.GetLocalClient()) == EBreakBehavior::Break)
+		{
+			return;
+		}
+		for (const TNonNullPtr<FRemoteReplicationClient>& RemoteClient : ClientManager.GetRemoteClients())
+		{
+			if (ProcessClient(*RemoteClient) == EBreakBehavior::Break)
+			{
+				break;
+			}
+		}
+	}
+
+	bool FGlobalAuthorityCache::HasAuthority(const FGuid& ClientId, const FGuid& StreamId, const FSoftObjectPath& ObjectPath) const
+	{
+		return GetClientsWithAuthorityOverObject(ObjectPath).Contains(ClientId);
+	}
+}
diff --git a/Engine/Plugins/Developer/Concert/ConcertApp/MultiUserClient/Source/MultiUserClient/Private/Replication/Misc/GlobalAuthorityCache.h b/Engine/Plugins/Developer/Concert/ConcertApp/MultiUserClient/Source/MultiUserClient/Private/Replication/Misc/GlobalAuthorityCache.h
new file mode 100644
--- /dev/null
+++ b/Engine/Plugins/Developer/Concert/ConcertApp/MultiUserClient/Source/MultiUserClient/Private/Replication/Misc/GlobalAuthorityCache.h
@@ -0,0 +1,131 @@
+// Copyright Epic Games, Inc. All Rights Reserved.
+
+#pragma once
+
+#include "Replication/Client/RemoteReplicationClient.h"
+
+#include "Replication/AuthorityConflictSharedUtils.h"
+
+#include "Delegates/Delegate.h"
+#include "Templates/Function.h"
+
+struct FConcertPropertyChain;
+
+namespace UE::MultiUserClient
+{
+	class FRemoteReplicationClient;
+	class FReplicationClient;
+	class FReplicationClientManager;
+}
+
+namespace UE::MultiUserClient
+{
+	/**
+	 * Allows efficient look-up of which objects and properties are owned by which clients.
+	 * This class efficiently answers the question: "Which clients own this object?"
+	 */
+	class FGlobalAuthorityCache
+		: public FNoncopyable
+		, ConcertSyncCore::Replication::AuthorityConflictUtils::IReplicationGroundTruth
+	{
+	public:
+
+		using FProcessPropertyConflict = TFunctionRef<EBreakBehavior(const FGuid& ConflictingClientId, const FConcertPropertyChain& Property)>;
+		
+		FGlobalAuthorityCache(FReplicationClientManager& InClientManager);
+		/** Called when the local client has been created and it is safe to register client events with FReplicationClientManager. */
+		void RegisterEvents();
+		
+		/** Iterates every client that Object in its stream. */
+		void ForEachClientWithObjectInStream(const FSoftObjectPath& Object, TFunctionRef<EBreakBehavior(const FGuid& ClientId)> Callback) const;
+
+		/** @return Whether the object is referenced by at least one client stream. */
+		bool IsObjectOrChildReferenced(const FSoftObjectPath& Object) const;
+
+		/** Iterates every client that has authority over Object. */
+		void ForEachClientWithAuthorityOverObject(const FSoftObjectPath& Object, TFunctionRef<EBreakBehavior(const FGuid& ClientId)> Callback) const;
+		/** Util that uses ForEachClientWithAuthorityOverObject to make an array. */
+		TArray<FGuid> GetClientsWithAuthorityOverObject(const FSoftObjectPath& Object) const;
+		/** @return Returns whether the client has (partial) authority over the object. */
+		bool HasAuthorityOverObject(const FSoftObjectPath& Object, const FGuid& ClientId) const;
+		
+		/** @return Whether the given client can take authority over the object without causing any conflicts. This also considers the changes made to the stream after submission. */
+		EAuthorityMutability CanClientTakeAuthorityAfterSubmission(const FSoftObjectPath& Object, const FGuid& ClientId, FProcessPropertyConflict ProcessConflict = [](auto&, auto&){ return EBreakBehavior::Break; }) const;
+		/** @return Whether the given client add the given property to the object without causing any conflicts. */
+		bool CanClientAddProperty(const FSoftObjectPath& Object, const FGuid& ClientId, const FConcertPropertyChain& Chain) const;
+
+		/** Gets the client that has authority over the given property, if there is any. */
+		TOptional<FGuid> GetClientWithAuthorityOverProperty(const FSoftObjectPath& Object, const FConcertPropertyChain& Property) const { return GetClientWithAuthorityOverProperty(Object, Property.GetPathToProperty()); }
+		/** Gets the client that has authority over the given property, if there is any. */
+		TOptional<FGuid> GetClientWithAuthorityOverProperty(const FSoftObjectPath& Object, const TArray<FName>& PropertyChain) const;
+		
+		/**
+		 * Removes entries from Request that would generate conflicts.
+		 * 
+		 * @param Request The request to clense
+		 * @param SendingClient The client that will send the request
+		 */
+		void CleanseConflictsFromAuthorityRequest(FConcertReplication_ChangeAuthority_Request& Request, const FGuid& SendingClient) const;
+		/**
+		 * Removes entries from Request that would generate conflicts.
+		 * 
+		 * @param Request The request to clense
+		 * @param SendingClient The client that will send the request
+		 */
+		void CleanseConflictsFromStreamRequest(FConcertReplication_ChangeStream_Request& Request, const FGuid& SendingClient) const;
+		
+		DECLARE_MULTICAST_DELEGATE_OneParam(FOnCacheChanged, const FGuid& ClientId);
+		/** Called when the cache changes for a specific client. */
+		FOnCacheChanged& OnCacheChanged() { return OnCacheChangedDelegate; }
+		
+	private:
+
+		/** Used to obtain the clients and their states */
+		FReplicationClientManager& ClientManager;
+		
+		/** Maps objects that are owned to the clients that own them */
+		TMap<FSoftObjectPath, TSet<FGuid>> OwnedObjectsToClients;
+		/** Maps objects that are owned to the clients that own them */
+		TMap<FSoftObjectPath, TSet<FGuid>> RegisteredObjectsToClients;
+		
+		/** Called when the cache changes for a specific client. */
+		FOnCacheChanged OnCacheChangedDelegate;
+
+		/** Registers for authority and stream changes */
+		void RegisterForClientEvents(const FReplicationClient& Client);
+		void UnregisterFromClientEvents(const FReplicationClient& Client) const;
+
+		/** Adds the client to OwnedObjectsToClients */
+		void AddClient(const FGuid& ClientId);
+		void RemoveClient(const FGuid& ClientId);
+
+		// Respond to remote client registration
+		void OnPostRemoteClientAdded(FRemoteReplicationClient& RemoteClient)
+		{
+			RebuildClient(RemoteClient.GetEndpointId());
+			RegisterForClientEvents(RemoteClient);
+		}
+		void OnPreRemoteClientRemoved(FRemoteReplicationClient& RemoteClient)
+		{
+			RemoveClient(RemoteClient.GetEndpointId());
+			UnregisterFromClientEvents(RemoteClient);
+		}
+
+		// Rebuild client when their authority changes
+		void OnPostAuthorityChanged(const FGuid ClientId) { RebuildClient(ClientId); }
+		void OnStreamChanged(const FGuid ClientId){ RebuildClient(ClientId); }
+		void RebuildClient(const FGuid ClientId)
+		{
+			RemoveClient(ClientId);
+			AddClient(ClientId);
+			OnCacheChangedDelegate.Broadcast(ClientId);
+		}
+
+		//~ Begin IReplicationGroundTruth Interface
+		virtual void ForEachStream(const FGuid& ClientEndpointId, TFunctionRef<EBreakBehavior(const FGuid& StreamId, const FConcertObjectReplicationMap& ReplicationMap)> Callback) const override;
+		virtual void ForEachSendingClient(TFunctionRef<EBreakBehavior(const FGuid& ClientEndpointId)> Callback) const override;
+		virtual bool HasAuthority(const FGuid& ClientId, const FGuid& StreamId, const FSoftObjectPath& ObjectPath) const override;
+		//~ End IReplicationGroundTruth Interface
+	};
+}
+
diff --git a/Engine/Plugins/Developer/Concert/ConcertApp/MultiUserClient/Source/MultiUserClient/Private/Replication/Misc/Notification/MutingNotifier.cpp b/Engine/Plugins/Developer/Concert/ConcertApp/MultiUserClient/Source/MultiUserClient/Private/Replication/Misc/Notification/MutingNotifier.cpp
new file mode 100644
--- /dev/null
+++ b/Engine/Plugins/Developer/Concert/ConcertApp/MultiUserClient/Source/MultiUserClient/Private/Replication/Misc/Notification/MutingNotifier.cpp
@@ -0,0 +1,41 @@
+﻿// Copyright Epic Games, Inc. All Rights Reserved.
+
+#include "MutingNotifier.h"
+
+#include "Replication/Muting/MuteStateManager.h"
+
+#include "Framework/Application/SlateApplication.h"
+#include "Framework/Notifications/NotificationManager.h"
+#include "Widgets/Notifications/SNotificationList.h"
+
+#define LOCTEXT_NAMESPACE "FMutingNotifier"
+
+namespace UE::MultiUserClient
+{
+	FMutingNotifier::FMutingNotifier(FMuteStateManager& InMuteManager)
+		: MuteManager(InMuteManager)
+	{
+		MuteManager.OnMuteRequestFailure().AddRaw(this, &FMutingNotifier::OnMuteRequestFailed);
+	}
+	
+	FMutingNotifier::~FMutingNotifier()
+	{
+		MuteManager.OnMuteRequestFailure().RemoveAll(this);
+	}
+
+	void FMutingNotifier::OnMuteRequestFailed(const FConcertReplication_ChangeMuteState_Request& Request, const FConcertReplication_ChangeMuteState_Response& Response)
+	{
+		if (FSlateApplication::IsInitialized())
+		{
+			FNotificationInfo Info(LOCTEXT("MuteRequestFailed.Title", "Pause / resume rejected by server."));
+			Info.SubText = LOCTEXT("MuteRequestFailed.Subtext", "Server-client state was likely de-synched.\nTry again.");
+			Info.bFireAndForget = true;
+			Info.ExpireDuration = 4.f;
+			Info.bUseSuccessFailIcons = true;
+			FSlateNotificationManager::Get().AddNotification(Info)
+				->SetCompletionState(SNotificationItem::CS_Fail);
+		}
+	}
+}
+
+#undef LOCTEXT_NAMESPACE
\ No newline at end of file
diff --git a/Engine/Plugins/Developer/Concert/ConcertApp/MultiUserClient/Source/MultiUserClient/Private/Replication/Misc/Notification/MutingNotifier.h b/Engine/Plugins/Developer/Concert/ConcertApp/MultiUserClient/Source/MultiUserClient/Private/Replication/Misc/Notification/MutingNotifier.h
new file mode 100644
--- /dev/null
+++ b/Engine/Plugins/Developer/Concert/ConcertApp/MultiUserClient/Source/MultiUserClient/Private/Replication/Misc/Notification/MutingNotifier.h
@@ -0,0 +1,31 @@
+﻿// Copyright Epic Games, Inc. All Rights Reserved.
+
+#pragma once
+
+#include "HAL/Platform.h"
+#include "Replication/Messages/Muting.h"
+#include "Templates/UnrealTemplate.h"
+
+struct FConcertReplication_ChangeMuteState_Request;
+struct FConcertReplication_ChangeMuteState_Response;
+
+namespace UE::MultiUserClient
+{
+	class FMuteStateManager;
+
+	/** Informs the user of things that went wrong with muting requests. */
+	class FMutingNotifier : public FNoncopyable
+	{
+	public:
+		FMutingNotifier(FMuteStateManager& InMuteManager UE_LIFETIMEBOUND);
+		~FMutingNotifier();
+
+	private:
+
+		/** Informs us when a request goes wrong. */
+		FMuteStateManager& MuteManager;
+		
+		void OnMuteRequestFailed(const FConcertReplication_ChangeMuteState_Request& Request, const FConcertReplication_ChangeMuteState_Response& Response);
+	};
+}
+
diff --git a/Engine/Plugins/Developer/Concert/ConcertApp/MultiUserClient/Source/MultiUserClient/Private/Replication/Misc/Notification/ReplicationUserNotifier.cpp b/Engine/Plugins/Developer/Concert/ConcertApp/MultiUserClient/Source/MultiUserClient/Private/Replication/Misc/Notification/ReplicationUserNotifier.cpp
new file mode 100644
--- /dev/null
+++ b/Engine/Plugins/Developer/Concert/ConcertApp/MultiUserClient/Source/MultiUserClient/Private/Replication/Misc/Notification/ReplicationUserNotifier.cpp
@@ -0,0 +1,11 @@
+﻿// Copyright Epic Games, Inc. All Rights Reserved.
+
+#include "ReplicationUserNotifier.h"
+
+namespace UE::MultiUserClient
+{
+	FReplicationUserNotifier::FReplicationUserNotifier(FReplicationClientManager& InReplicationClientManager, FMuteStateManager& InMuteManager)
+		: SubmissionNotifier(InReplicationClientManager)
+		, MutingNotifier(InMuteManager)
+	{}
+}
diff --git a/Engine/Plugins/Developer/Concert/ConcertApp/MultiUserClient/Source/MultiUserClient/Private/Replication/Misc/Notification/ReplicationUserNotifier.h b/Engine/Plugins/Developer/Concert/ConcertApp/MultiUserClient/Source/MultiUserClient/Private/Replication/Misc/Notification/ReplicationUserNotifier.h
new file mode 100644
--- /dev/null
+++ b/Engine/Plugins/Developer/Concert/ConcertApp/MultiUserClient/Source/MultiUserClient/Private/Replication/Misc/Notification/ReplicationUserNotifier.h
@@ -0,0 +1,31 @@
+﻿// Copyright Epic Games, Inc. All Rights Reserved.
+
+#pragma once
+
+#include "MutingNotifier.h"
+#include "Replication/Submission/Notification/SubmissionNotifier.h"
+
+#include "HAL/Platform.h"
+#include "Templates/UnrealTemplate.h"
+
+namespace UE::MultiUserClient
+{
+	class FMuteStateManager;
+
+	/** This informs the user of things that went wrong during replication. */
+	class FReplicationUserNotifier : public FNoncopyable
+	{
+	public:
+
+		FReplicationUserNotifier(FReplicationClientManager& InReplicationClientManager UE_LIFETIMEBOUND, FMuteStateManager& InMuteManager UE_LIFETIMEBOUND);
+
+	private:
+		
+		/** Notifies the user when stream or authority submission to the server fails. */
+		const FSubmissionNotifier SubmissionNotifier;
+		
+		/** Informs the user of things that went wrong with muting requests. */
+		const FMutingNotifier MutingNotifier;
+	};
+}
+
diff --git a/Engine/Plugins/Developer/Concert/ConcertApp/MultiUserClient/Source/MultiUserClient/Private/Replication/Misc/PreventReplicatedPropertyTransaction.cpp b/Engine/Plugins/Developer/Concert/ConcertApp/MultiUserClient/Source/MultiUserClient/Private/Replication/Misc/PreventReplicatedPropertyTransaction.cpp
--- a/Engine/Plugins/Developer/Concert/ConcertApp/MultiUserClient/Source/MultiUserClient/Private/Replication/Misc/PreventReplicatedPropertyTransaction.cpp
+++ b/Engine/Plugins/Developer/Concert/ConcertApp/MultiUserClient/Source/MultiUserClient/Private/Replication/Misc/PreventReplicatedPropertyTransaction.cpp
@@ -4,7 +4,7 @@
 
 #include "IConcertSyncClient.h"
 #include "Replication/Client/ReplicationClientManager.h"
-#include "Replication/Util/GlobalAuthorityCache.h"
+#include "Replication/Misc/GlobalAuthorityCache.h"
 
 #include "Algo/AnyOf.h"
 #include "Replication/Muting/MuteStateManager.h"
diff --git a/Engine/Plugins/Developer/Concert/ConcertApp/MultiUserClient/Source/MultiUserClient/Private/Replication/Misc/Query/MuteStateQueryService.cpp b/Engine/Plugins/Developer/Concert/ConcertApp/MultiUserClient/Source/MultiUserClient/Private/Replication/Misc/Query/MuteStateQueryService.cpp
new file mode 100644
--- /dev/null
+++ b/Engine/Plugins/Developer/Concert/ConcertApp/MultiUserClient/Source/MultiUserClient/Private/Replication/Misc/Query/MuteStateQueryService.cpp
@@ -0,0 +1,35 @@
+﻿// Copyright Epic Games, Inc. All Rights Reserved.
+
+#include "MuteStateQueryService.h"
+
+#include "IConcertSyncClient.h"
+#include "Replication/IConcertClientReplicationManager.h"
+
+namespace UE::MultiUserClient
+{
+	FMuteStateQueryService::FMuteStateQueryService(TWeakPtr<FToken> InToken, const IConcertSyncClient& InOwningClient)
+		: Token(MoveTemp(InToken))
+		, OwningClient(InOwningClient)
+	{}
+
+	void FMuteStateQueryService::SendQueryEvent()
+	{
+		// Avoid network calls when nobody is subscribed.
+		IConcertClientReplicationManager* ReplicationManager = OwningClient.GetReplicationManager();
+		if (!OnMuteStateQueriedDelegate.IsBound() || !ensure(ReplicationManager))
+		{
+			return;
+		}
+
+		// For now, we'll query for ALL objects - in the future we could restrict it to the local application's GWorld
+		ReplicationManager->QueryMuteState()
+			.Next([this, WeakToken = Token](FConcertReplication_QueryMuteState_Response&& Response)
+			{
+				// If WeakToken is not valid, our this pointer has been destroyed (there's currently no API for cancelling started requests).
+				if (WeakToken.IsValid())
+				{
+					OnMuteStateQueriedDelegate.Broadcast(Response);
+				}
+			});
+	}
+}
diff --git a/Engine/Plugins/Developer/Concert/ConcertApp/MultiUserClient/Source/MultiUserClient/Private/Replication/Misc/Query/MuteStateQueryService.h b/Engine/Plugins/Developer/Concert/ConcertApp/MultiUserClient/Source/MultiUserClient/Private/Replication/Misc/Query/MuteStateQueryService.h
new file mode 100644
--- /dev/null
+++ b/Engine/Plugins/Developer/Concert/ConcertApp/MultiUserClient/Source/MultiUserClient/Private/Replication/Misc/Query/MuteStateQueryService.h
@@ -0,0 +1,56 @@
+﻿// Copyright Epic Games, Inc. All Rights Reserved.
+
+#pragma once
+
+#include "Replication/IToken.h"
+
+#include "Delegates/Delegate.h"
+#include "HAL/Platform.h"
+#include "Templates/SharedPointer.h"
+#include "Templates/UnrealTemplate.h"
+
+class IConcertSyncClient;
+struct FConcertReplication_QueryMuteState_Response;
+
+namespace UE::MultiUserClient
+{
+	DECLARE_MULTICAST_DELEGATE_OneParam(FOnMuteStateQueried, const FConcertReplication_QueryMuteState_Response& NewMuteState);
+	
+	/** Sends regular FConcertReplication_QueryMuteState_Request to endpoints and publishes the results. */
+	class FMuteStateQueryService : public FNoncopyable
+	{
+		friend class FRegularQueryService;
+	public:
+		
+		/**
+		 * @param InToken Used to check whether this instance was destroyed after a response is received.
+		 * @param InOwningClient The client through which queries are made. The caller ensures it outlasts the constructed instance.
+		 */
+		FMuteStateQueryService(TWeakPtr<FToken> InToken, const IConcertSyncClient& InOwningClient UE_LIFETIMEBOUND);
+
+		/** @return Sends a request to the server immediately */
+		void RequestInstantRefresh() { SendQueryEvent(); }
+
+		FOnMuteStateQueried& OnMuteStateQueried() { return OnMuteStateQueriedDelegate; }
+		
+	private:
+
+		/** Used to check whether we were destroyed after a response is received. */
+		const TWeakPtr<FToken> Token;
+		
+		/**
+		 * Used to send queries.
+		 * 
+		 * This FRegularQueryService's owner is supposed to make sure this FRegularQueryService is destroyed
+		 * when the client shuts down.
+		 */
+		const IConcertSyncClient& OwningClient;
+
+		/** Broadcasts when the mute state has been queried. */
+		FOnMuteStateQueried OnMuteStateQueriedDelegate;
+		
+		/** Queries the server for the client's current state. */
+		void SendQueryEvent();
+	};
+}
+
diff --git a/Engine/Plugins/Developer/Concert/ConcertApp/MultiUserClient/Source/MultiUserClient/Private/Replication/Misc/Query/RegularQueryService.cpp b/Engine/Plugins/Developer/Concert/ConcertApp/MultiUserClient/Source/MultiUserClient/Private/Replication/Misc/Query/RegularQueryService.cpp
new file mode 100644
--- /dev/null
+++ b/Engine/Plugins/Developer/Concert/ConcertApp/MultiUserClient/Source/MultiUserClient/Private/Replication/Misc/Query/RegularQueryService.cpp
@@ -0,0 +1,31 @@
+﻿// Copyright Epic Games, Inc. All Rights Reserved.
+
+#include "RegularQueryService.h"
+
+namespace UE::MultiUserClient
+{
+	FRegularQueryService::FRegularQueryService(const IConcertSyncClient& InOwningClient, float InInterval)
+		: OwningClient(InOwningClient)
+		, TickerDelegateHandle(FTSTicker::GetCoreTicker().AddTicker(
+			TEXT("Multi-User Replication Query"),
+			InInterval,
+			[this](float)
+			{
+				Tick();
+				return true;
+			}))
+		, StreamAndAuthorityQueryService(Token, OwningClient)
+		, MuteStateQueryService(Token, OwningClient)
+	{}
+
+	FRegularQueryService::~FRegularQueryService()
+	{
+		FTSTicker::GetCoreTicker().RemoveTicker(TickerDelegateHandle);
+	}
+
+	void FRegularQueryService::Tick()
+	{
+		StreamAndAuthorityQueryService.SendQueryEvent();
+		MuteStateQueryService.SendQueryEvent();
+	}
+}
diff --git a/Engine/Plugins/Developer/Concert/ConcertApp/MultiUserClient/Source/MultiUserClient/Private/Replication/Misc/Query/RegularQueryService.h b/Engine/Plugins/Developer/Concert/ConcertApp/MultiUserClient/Source/MultiUserClient/Private/Replication/Misc/Query/RegularQueryService.h
new file mode 100644
--- /dev/null
+++ b/Engine/Plugins/Developer/Concert/ConcertApp/MultiUserClient/Source/MultiUserClient/Private/Replication/Misc/Query/RegularQueryService.h
@@ -0,0 +1,62 @@
+﻿// Copyright Epic Games, Inc. All Rights Reserved.
+
+#pragma once
+
+#include "MuteStateQueryService.h"
+#include "StreamAndAuthorityQueryService.h"
+#include "Replication/IToken.h"
+
+#include "Containers/Ticker.h"
+#include "HAL/Platform.h"
+#include "Templates/UnrealTemplate.h"
+
+class IConcertSyncClient;
+
+namespace UE::MultiUserClient
+{
+	/**
+	 * This service manages multiple sub-services that query the server state, e.g. stream, authority, etc.
+	 * Its responsibility is to tick the sub-services and manage access to them. 
+	 */
+	class FRegularQueryService : public FNoncopyable
+	{
+	public:
+		
+		/**
+		 * @param InOwningClient The client through which queries are made. The caller ensures it outlasts the constructed instance.
+		 * @param InInterval The interval at which queries are to be made.
+		 */
+		FRegularQueryService(const IConcertSyncClient& InOwningClient UE_LIFETIMEBOUND, float InInterval = 1.f);
+		~FRegularQueryService();
+
+		/** @return The service for querying stream and authority state. */
+		FStreamAndAuthorityQueryService& GetStreamAndAuthorityQueryService() { return StreamAndAuthorityQueryService; }
+		/** @return The service for querying global mute state. */
+		FMuteStateQueryService& GetMuteStateQueryService() { return MuteStateQueryService; }
+
+	private:
+		
+		/** Passed to sub-services to check whether they were destroyed after a response is received. */
+		const TSharedRef<FToken> Token = FToken::Make();
+		
+		/**
+		 * Used to send queries.
+		 * 
+		 * This FRegularQueryService's owner is supposed to make sure this FRegularQueryService is destroyed
+		 * when the client shuts down.
+		 */
+		const IConcertSyncClient& OwningClient;
+
+		/** Used to unregister the ticker. */
+		const FTSTicker::FDelegateHandle TickerDelegateHandle;
+
+		/** Queries stream and authority changes. */
+		FStreamAndAuthorityQueryService StreamAndAuthorityQueryService;
+		/** Queries global mute state at regular intervals. */
+		FMuteStateQueryService MuteStateQueryService;
+
+		/** Ticks the sub-services when it is time to send a new request. */
+		void Tick();
+	};
+}
+
diff --git a/Engine/Plugins/Developer/Concert/ConcertApp/MultiUserClient/Source/MultiUserClient/Private/Replication/Misc/Query/StreamAndAuthorityQueryService.cpp b/Engine/Plugins/Developer/Concert/ConcertApp/MultiUserClient/Source/MultiUserClient/Private/Replication/Misc/Query/StreamAndAuthorityQueryService.cpp
new file mode 100644
--- /dev/null
+++ b/Engine/Plugins/Developer/Concert/ConcertApp/MultiUserClient/Source/MultiUserClient/Private/Replication/Misc/Query/StreamAndAuthorityQueryService.cpp
@@ -0,0 +1,186 @@
+// Copyright Epic Games, Inc. All Rights Reserved.
+
+#include "StreamAndAuthorityQueryService.h"
+
+#include "IConcertSyncClient.h"
+#include "Replication/IConcertClientReplicationManager.h"
+
+#include "Templates/UnrealTemplate.h"
+
+namespace UE::MultiUserClient
+{
+	FStreamAndAuthorityQueryService::FStreamAndAuthorityQueryService(TWeakPtr<FToken> InToken, const IConcertSyncClient& InOwningClient)
+		: Token(MoveTemp(InToken))
+		, OwningClient(InOwningClient)
+	{}
+
+	FDelegateHandle FStreamAndAuthorityQueryService::RegisterStreamQuery(const FGuid& EndpointId, FStreamQueryDelegate Delegate)
+	{
+		FStreamQueryInfo& Info = StreamQueryInfos.FindOrAdd(EndpointId);
+		const FDelegateHandle DelegateHandle = Info.Delegate.AddLambda([Delegate = MoveTemp(Delegate)](const TArray<FConcertBaseStreamInfo>& Descriptions)
+		{
+			Delegate.Execute(Descriptions);
+		});
+		Info.Handles.Add(DelegateHandle);
+		return DelegateHandle;
+	}
+
+	FDelegateHandle FStreamAndAuthorityQueryService::RegisterAuthorityQuery(const FGuid& EndpointId, FAuthorityQueryDelegate Delegate)
+	{
+		FAuthorityQueryInfo& Info = AuthorityQueryInfos.FindOrAdd(EndpointId);
+		const FDelegateHandle DelegateHandle = Info.Delegate.AddLambda([Delegate = MoveTemp(Delegate)](const TArray<FConcertAuthorityClientInfo>& Infos)
+		{
+			Delegate.Execute(Infos);
+		});
+		Info.Handles.Add(DelegateHandle);
+		return DelegateHandle;
+	}
+
+	void FStreamAndAuthorityQueryService::UnregisterStreamQuery(const FDelegateHandle& Handle)
+	{
+		for (auto StreamIt = StreamQueryInfos.CreateIterator(); StreamIt; ++StreamIt)
+		{
+			FStreamQueryInfo& Info = StreamIt->Value;
+			if (Info.Handles.Contains(Handle))
+			{
+				Info.Handles.Remove(Handle);
+				Info.Delegate.Remove(Handle);
+				
+				// If bIsHandlingQueryResponse, then HandleQueryResponse is executing delegates and a delegate is unregistering.
+				// Removing now would cause a crash since Broadcast will read some bookkeeping memory after Broadcast.
+				// Memory is cleaned up later by HandleQueryResponse.
+				if (!bIsHandlingQueryResponse && Info.Handles.IsEmpty())
+				{
+					StreamIt.RemoveCurrent();
+				}
+				
+				return;
+			}
+		}
+	}
+
+	void FStreamAndAuthorityQueryService::UnregisterAuthorityQuery(const FDelegateHandle& Handle)
+	{
+		for (auto AuthorityIt = AuthorityQueryInfos.CreateIterator(); AuthorityIt; ++AuthorityIt)
+		{
+			FAuthorityQueryInfo& Info = AuthorityIt->Value;
+			if (Info.Handles.Contains(Handle))
+			{
+				Info.Handles.Remove(Handle);
+				Info.Delegate.Remove(Handle);
+				
+				// If bIsHandlingQueryResponse, then HandleQueryResponse is executing delegates and a delegate is unregistering.
+				// Removing now would cause a crash since Broadcast will read some bookkeeping memory after Broadcast.
+				// Memory is cleaned up later by HandleQueryResponse.
+				if (!bIsHandlingQueryResponse && Info.Handles.IsEmpty())
+				{
+					AuthorityIt.RemoveCurrent();
+				}
+				
+				return;
+			}
+		}
+	}
+
+	void FStreamAndAuthorityQueryService::SendQueryEvent()
+	{
+		IConcertClientReplicationManager* ReplicationManager = OwningClient.GetReplicationManager();
+		if (!ensure(ReplicationManager))
+		{
+			return;
+		}
+		
+		FConcertReplication_QueryReplicationInfo_Request Request;
+		BuildStreamRequest(Request);
+		BuildAuthorityRequest(Request);
+
+		if (!Request.ClientEndpointIds.IsEmpty())
+		{
+			ReplicationManager->QueryClientInfo({ Request })
+				.Next([this, WeakToken = Token](FConcertReplication_QueryReplicationInfo_Response&& Response)
+				{
+					const bool bCanProcessRequest = WeakToken.Pin().IsValid();
+					if (bCanProcessRequest)
+					{
+						HandleQueryResponse(Response);
+					}
+				});
+		}
+	}
+
+	void FStreamAndAuthorityQueryService::BuildStreamRequest(FConcertReplication_QueryReplicationInfo_Request& Request) const
+	{
+		if (StreamQueryInfos.IsEmpty())
+		{
+			Request.QueryFlags |= EConcertQueryClientStreamFlags::SkipStreamInfo;
+		}
+		else
+		{
+			for (const TPair<FGuid, FStreamQueryInfo>& Query : StreamQueryInfos)
+			{
+				Request.ClientEndpointIds.Add(Query.Key);
+			}
+		}
+	}
+
+	void FStreamAndAuthorityQueryService::BuildAuthorityRequest(FConcertReplication_QueryReplicationInfo_Request& Request) const
+	{
+		if (AuthorityQueryInfos.IsEmpty())
+		{
+			Request.QueryFlags |= EConcertQueryClientStreamFlags::SkipAuthority;
+		}
+		else
+		{
+			for (const TPair<FGuid, FAuthorityQueryInfo>& Query : AuthorityQueryInfos)
+			{
+				Request.ClientEndpointIds.Add(Query.Key);
+			}
+		}
+	}
+
+	void FStreamAndAuthorityQueryService::HandleQueryResponse(
+		const FConcertReplication_QueryReplicationInfo_Response& Response
+		)
+	{
+		TGuardValue<bool> Guard(bIsHandlingQueryResponse, true);
+		
+		for (const TPair<FGuid, FConcertQueriedClientInfo>& StreamQueryPair : Response.ClientInfo)
+		{
+			const FGuid StreamId = StreamQueryPair.Key;
+			
+			// Note that the delegates may have been unsubscribed since request was sent, which is why we need to Find() them
+			if (const FStreamQueryInfo* QueryInfo = StreamQueryInfos.Find(StreamId))
+			{
+				QueryInfo->Delegate.Broadcast(StreamQueryPair.Value.Streams);
+			}
+
+			// Note that the delegates may have been unsubscribed since request was sent, which is why we need to Find() them
+			if (const FAuthorityQueryInfo* QueryInfo = AuthorityQueryInfos.Find(StreamId))
+			{
+				QueryInfo->Delegate.Broadcast(StreamQueryPair.Value.Authority);
+			}
+		}
+
+		// Delegates may have removed themselves up above.
+		CompactDelegates();
+	}
+
+	void FStreamAndAuthorityQueryService::CompactDelegates()
+	{
+		for (auto StreamIt = StreamQueryInfos.CreateIterator(); StreamIt; ++StreamIt)
+		{
+			if (StreamIt->Value.Handles.IsEmpty())
+			{
+				StreamIt.RemoveCurrent();
+			}
+		}
+
+		for (auto Authority = AuthorityQueryInfos.CreateIterator(); Authority; ++Authority)
+		{
+			if (Authority->Value.Handles.IsEmpty())
+			{
+				Authority.RemoveCurrent();
+			}
+		}
+	}
+}
diff --git a/Engine/Plugins/Developer/Concert/ConcertApp/MultiUserClient/Source/MultiUserClient/Private/Replication/Misc/Query/StreamAndAuthorityQueryService.h b/Engine/Plugins/Developer/Concert/ConcertApp/MultiUserClient/Source/MultiUserClient/Private/Replication/Misc/Query/StreamAndAuthorityQueryService.h
new file mode 100644
--- /dev/null
+++ b/Engine/Plugins/Developer/Concert/ConcertApp/MultiUserClient/Source/MultiUserClient/Private/Replication/Misc/Query/StreamAndAuthorityQueryService.h
@@ -0,0 +1,82 @@
+// Copyright Epic Games, Inc. All Rights Reserved.
+
+#pragma once
+
+#include "Replication/Messages/ClientQuery.h"
+#include "Replication/IToken.h"
+
+#include "Delegates/Delegate.h"
+#include "HAL/Platform.h"
+#include "Templates/SharedPointer.h"
+#include "Templates/UnrealTemplate.h"
+
+class IConcertSyncClient;
+
+namespace UE::MultiUserClient
+{
+	DECLARE_DELEGATE_OneParam(FStreamQueryDelegate, const TArray<FConcertBaseStreamInfo>&);
+	DECLARE_DELEGATE_OneParam(FAuthorityQueryDelegate, const TArray<FConcertAuthorityClientInfo>&);
+	
+	/** Sends regular FConcertReplication_QueryReplicationInfo_Request to endpoints and publishes the results. */
+	class FStreamAndAuthorityQueryService : public FNoncopyable
+	{
+		friend class FRegularQueryService;
+	public:
+		
+		/**
+		 * @param InToken Used to check whether this instance was destroyed after a response is received.
+		 * @param InOwningClient The client through which queries are made. The caller ensures it outlasts the constructed instance.
+		 */
+		FStreamAndAuthorityQueryService(TWeakPtr<FToken> InToken, const IConcertSyncClient& InOwningClient UE_LIFETIMEBOUND);
+
+		/** Registers a delegate to invoke for querying an endpoint about its registered streams. */
+		FDelegateHandle RegisterStreamQuery(const FGuid& EndpointId, FStreamQueryDelegate Delegate);
+		/** Registers a delegate to invoke for querying an endpoint about its authority. */
+		FDelegateHandle RegisterAuthorityQuery(const FGuid& EndpointId, FAuthorityQueryDelegate Delegate);
+
+		void UnregisterStreamQuery(const FDelegateHandle& Handle);
+		void UnregisterAuthorityQuery(const FDelegateHandle& Handle);
+		
+	private:
+
+		/** Used to check whether we were destroyed after a response is received. */
+		const TWeakPtr<FToken> Token;
+		
+		/**
+		 * Used to send queries.
+		 * 
+		 * This FRegularQueryService's owner is supposed to make sure this FRegularQueryService is destroyed
+		 * when the client shuts down.
+		 */
+		const IConcertSyncClient& OwningClient;
+
+		// We use a multicast delegate here because it handles unsubscribing while being Broadcast
+		DECLARE_MULTICAST_DELEGATE_OneParam(FMulticastStreamQueryDelegate, const TArray<FConcertBaseStreamInfo>&);
+		DECLARE_MULTICAST_DELEGATE_OneParam(FMulticastAuthorityQueryDelegate, const TArray<FConcertAuthorityClientInfo>&);
+		struct FStreamQueryInfo
+		{
+			TSet<FDelegateHandle> Handles;
+			FMulticastStreamQueryDelegate Delegate;
+		};
+		struct FAuthorityQueryInfo
+		{
+			TSet<FDelegateHandle> Handles;
+			FMulticastAuthorityQueryDelegate Delegate;
+		};
+		TMap<FGuid, FStreamQueryInfo> StreamQueryInfos;
+		TMap<FGuid, FAuthorityQueryInfo> AuthorityQueryInfos;
+
+		bool bIsHandlingQueryResponse = false;
+		
+		/** Queries the server for the client's current state. */
+		void SendQueryEvent();
+		
+		void BuildStreamRequest(FConcertReplication_QueryReplicationInfo_Request& Request) const;
+		void BuildAuthorityRequest(FConcertReplication_QueryReplicationInfo_Request& Request) const;
+		
+		/** Handles the server's response. */
+		void HandleQueryResponse(const FConcertReplication_QueryReplicationInfo_Response& Response);
+		void CompactDelegates();
+	};
+}
+
diff --git a/Engine/Plugins/Developer/Concert/ConcertApp/MultiUserClient/Source/MultiUserClient/Private/Replication/Misc/Util/StreamRequestUtils.cpp b/Engine/Plugins/Developer/Concert/ConcertApp/MultiUserClient/Source/MultiUserClient/Private/Replication/Misc/Util/StreamRequestUtils.cpp
new file mode 100644
--- /dev/null
+++ b/Engine/Plugins/Developer/Concert/ConcertApp/MultiUserClient/Source/MultiUserClient/Private/Replication/Misc/Util/StreamRequestUtils.cpp
@@ -0,0 +1,75 @@
+// Copyright Epic Games, Inc. All Rights Reserved.
+
+#include "StreamRequestUtils.h"
+
+#include "ConcertLogGlobal.h"
+#include "JsonObjectConverter.h"
+#include "Replication/Stream/StreamChangeTracker.h"
+#include "Widgets/ActiveSession/Replication/Client/ClientUtils.h"
+
+#include "Containers/UnrealString.h"
+
+namespace UE::MultiUserClient::StreamRequestUtils
+{
+	FConcertReplication_ChangeStream_Request BuildChangeRequest_CreateNewStream(
+		const FGuid& StreamId,
+		const FStreamChangelist& ObjectChanges,
+		FFrequencyChangelist FrequencyChanges
+		)
+	{
+		const TMap<FConcertObjectInStreamID, FConcertReplication_ChangeStream_PutObject>& ObjectsToPut = ObjectChanges.ObjectsToPut;
+		
+		FConcertReplication_ChangeStream_Request Request;
+		Request.StreamsToAdd.Emplace();
+		FConcertReplicationStream& NewStream = Request.StreamsToAdd[0];
+		NewStream.BaseDescription.Identifier = StreamId;
+			
+		// If creating a new stream, the objects must be supplied in the description instead of in PutObjects!
+		FConcertObjectReplicationMap& ReplicationMap = NewStream.BaseDescription.ReplicationMap;
+		ReplicationMap.ReplicatedObjects.Reserve(ObjectsToPut.Num());
+		for (const TPair<FConcertObjectInStreamID, FConcertReplication_ChangeStream_PutObject>& PutObjectPair : ObjectsToPut)
+		{
+			const TOptional<FConcertReplicatedObjectInfo> NewObjectInfo = PutObjectPair.Value.MakeObjectInfoIfValid();
+			// The editing UI allows adding objects without properties (to make UX easier) - do not submit those to the server.
+			if (!NewObjectInfo)
+			{
+				continue;
+			}
+			checkf(PutObjectPair.Key.StreamId == StreamId, TEXT("BuildChangeRequest should always use LocalClientStreamId ID!"));
+				
+			ReplicationMap.ReplicatedObjects.Add(PutObjectPair.Key.Object, *NewObjectInfo);
+		}
+		
+		FConcertStreamFrequencySettings& FrequencySettings = NewStream.BaseDescription.FrequencySettings;
+		if (FrequencyChanges.NewDefaults)
+		{
+			FrequencySettings.Defaults = MoveTemp(*FrequencyChanges.NewDefaults);
+		}
+		for (TPair<FSoftObjectPath, FConcertObjectReplicationSettings>& Overrides : FrequencyChanges.OverridesToAdd)
+		{
+			FrequencySettings.ObjectOverrides.Emplace(MoveTemp(Overrides.Key), MoveTemp(Overrides.Value));
+		}
+		
+		return Request;
+	}
+		
+	FConcertReplication_ChangeStream_Request BuildChangeRequest_UpdateExistingStream(
+		const FGuid& StreamId,
+		FStreamChangelist ObjectChanges,
+		FFrequencyChangelist FrequencyChanges
+		)
+	{
+		FConcertReplication_ChangeStream_Request Result { MoveTemp(ObjectChanges.ObjectsToRemove), MoveTemp(ObjectChanges.ObjectsToPut) };
+		
+		FConcertReplication_ChangeStream_Frequency& FrequencyChangeRequest = Result.FrequencyChanges.Add(StreamId);
+		if (FrequencyChanges.NewDefaults)
+		{
+			FrequencyChangeRequest.NewDefaults = MoveTemp(*FrequencyChanges.NewDefaults); 
+			FrequencyChangeRequest.Flags = EConcertReplicationChangeFrequencyFlags::SetDefaults; 
+		}
+		FrequencyChangeRequest.OverridesToAdd = MoveTemp(FrequencyChanges.OverridesToAdd);
+		FrequencyChangeRequest.OverridesToRemove = MoveTemp(FrequencyChanges.OverridesToRemove);
+
+		return Result;
+	}
+}
diff --git a/Engine/Plugins/Developer/Concert/ConcertApp/MultiUserClient/Source/MultiUserClient/Private/Replication/Misc/Util/StreamRequestUtils.h b/Engine/Plugins/Developer/Concert/ConcertApp/MultiUserClient/Source/MultiUserClient/Private/Replication/Misc/Util/StreamRequestUtils.h
new file mode 100644
--- /dev/null
+++ b/Engine/Plugins/Developer/Concert/ConcertApp/MultiUserClient/Source/MultiUserClient/Private/Replication/Misc/Util/StreamRequestUtils.h
@@ -0,0 +1,56 @@
+// Copyright Epic Games, Inc. All Rights Reserved.
+
+#pragma once
+
+#include "Replication/IConcertClientReplicationManager.h"
+
+class IConcertClientSession;
+
+struct FConcertReplication_ChangeStream_Request;
+
+namespace UE::MultiUserClient
+{
+	/** Describes changes that MU clients make to the objects in a stream. */
+	struct FStreamChangelist
+	{
+		TSet<FConcertObjectInStreamID> ObjectsToRemove;
+		TMap<FConcertObjectInStreamID, FConcertReplication_ChangeStream_PutObject> ObjectsToPut;
+	};
+	
+	/** Describes changes that MU clients make to the frequency settings in a stream. */
+	struct FFrequencyChangelist
+	{
+		TSet<FSoftObjectPath> OverridesToRemove;
+		TMap<FSoftObjectPath, FConcertObjectReplicationSettings> OverridesToAdd;
+		TOptional<FConcertObjectReplicationSettings> NewDefaults;
+
+		bool IsEmpty() const { return OverridesToRemove.IsEmpty() && OverridesToAdd.IsEmpty() && !NewDefaults.IsSet(); }
+	};
+}
+
+namespace UE::MultiUserClient::StreamRequestUtils
+{
+	/**
+	 * Builds a request for creating a new stream.
+	 * @param StreamId The stream that should be modified
+	 * @param ObjectChanges The object changes to be made
+	 * @param FrequencyChanges The frequency changes to be made
+	 */
+	FConcertReplication_ChangeStream_Request BuildChangeRequest_CreateNewStream(
+		const FGuid& StreamId,
+		const FStreamChangelist& ObjectChanges,
+		FFrequencyChangelist FrequencyChanges = FFrequencyChangelist()
+		);
+	
+	/**
+	 * Builds a request for updating a preexisting stream.
+	 * @param StreamId The stream that should be modified
+	 * @param ObjectChanges The object changes to be made
+	 * @param FrequencyChanges The frequency changes to be made
+	 */
+	FConcertReplication_ChangeStream_Request BuildChangeRequest_UpdateExistingStream(
+		const FGuid& StreamId,
+		FStreamChangelist ObjectChanges,
+		FFrequencyChangelist FrequencyChanges = FFrequencyChangelist()
+		);
+}
diff --git a/Engine/Plugins/Developer/Concert/ConcertApp/MultiUserClient/Source/MultiUserClient/Private/Replication/Misc/Util/SynchronizedRequestUtils.cpp b/Engine/Plugins/Developer/Concert/ConcertApp/MultiUserClient/Source/MultiUserClient/Private/Replication/Misc/Util/SynchronizedRequestUtils.cpp
new file mode 100644
--- /dev/null
+++ b/Engine/Plugins/Developer/Concert/ConcertApp/MultiUserClient/Source/MultiUserClient/Private/Replication/Misc/Util/SynchronizedRequestUtils.cpp
@@ -0,0 +1,290 @@
+// Copyright Epic Games, Inc. All Rights Reserved.
+
+#include "SynchronizedRequestUtils.h"
+
+#include "ConcertLogGlobal.h"
+#include "Replication/Client/ReplicationClientManager.h"
+#include "Replication/Submission/Data/AuthoritySubmission.h"
+#include "Replication/Submission/Data/StreamSubmission.h"
+#include "Replication/Submission/ISubmissionOperation.h"
+
+#include "Algo/AllOf.h"
+#include "Containers/UnrealString.h"
+#include <atomic>
+
+namespace UE::MultiUserClient
+{
+	namespace ParallelSubmission::Private
+	{
+		static void RemoveInvalidClientEntries(const FReplicationClientManager& ClientManager, TMap<FGuid, FSubmissionParams> ParallelOperations)
+		{
+			for (auto It = ParallelOperations.CreateIterator(); It; ++It)
+			{
+				const TPair<FGuid, FSubmissionParams>& Operation =*It;
+				if (!ensure(ClientManager.FindClient(Operation.Key)))
+				{
+					It.RemoveCurrent();
+				}
+			}
+		}
+
+		class FSyncOperation;
+
+		/** Adapter for FSubmissionQueue. Handles submitting for one client and notifying FSyncOperation if completions. */
+		class FDeferredSubmitter
+			: public FNoncopyable
+			, public IDeferredSubmitter
+		{
+			FSyncOperation& Owner;
+			FReplicationClientManager& ClientManager;
+			FSubmissionQueue& SubmissionQueue;
+			const FGuid ClientId;
+			const FSubmissionParams SubmissionParams;
+			
+		public:
+
+			FDeferredSubmitter(
+				FSyncOperation& InOwner,
+				FReplicationClientManager& InClientManager,
+				FSubmissionQueue& InSubmissionQueue,
+				const FGuid& InClientId,
+				FSubmissionParams&& InSubmissionParams
+				)
+				: Owner(InOwner)
+				, ClientManager(InClientManager)
+				, SubmissionQueue(InSubmissionQueue)
+				, ClientId(InClientId)
+				, SubmissionParams(MoveTemp(InSubmissionParams))
+			{}
+
+			virtual ~FDeferredSubmitter() override
+			{
+				checkf(IsInGameThread(), TEXT("You must destroy IParallelSubmissionOperation on the game thread"));
+				if (FReplicationClient* Client = ClientManager.FindClient(ClientId))
+				{
+					Client->GetSubmissionQueue().Dequeue_GameThread(*this);
+				}
+			}
+
+			void EnqueueOrSubmit()
+			{
+				SubmissionQueue.SubmitNowOrEnqueue_GameThread(*this);
+			}
+			
+			const FGuid& GetClientId() const { return ClientId; }
+
+			virtual void PerformSubmission_GameThread(ISubmissionWorkflow& Workflow) override;
+		};
+
+		/** Starts all submissions and synchronizes their concurrent completion. */
+		class FSyncOperation : public IParallelSubmissionOperation, public TSharedFromThis<FSyncOperation>
+		{
+		public:
+			
+			FSyncOperation(FReplicationClientManager& ClientManager, TMap<FGuid, FSubmissionParams>&& Operations)
+			{
+				const int32 NumOperations = Operations.Num();
+				IntermediateResults.StreamResponses.Reserve(NumOperations);
+				IntermediateResults.AuthorityResponses.Reserve(NumOperations);
+				SubOperations.Reserve(NumOperations);
+
+				// Prepare the intermediate slots in advance to handle concurrent writing
+				for (TPair<FGuid, FSubmissionParams>& Pair : Operations)
+				{
+					// Allocate slots
+					const FGuid& ClientId = Pair.Key;
+					IntermediateResults.StreamResponses.Add(ClientId);
+					IntermediateResults.AuthorityResponses.Add(ClientId);
+
+					FReplicationClient* Client = ClientManager.FindClient(ClientId);
+					checkf(Client, TEXT("Was supposed to have been validated ExecuteParallelStreamChanges."))
+					SubOperations.Emplace(*this, ClientManager, ClientId, MoveTemp(Pair.Value), Client->GetSubmissionQueue());
+				}
+			}
+
+			virtual ~FSyncOperation() override
+			{
+				if (!bPromiseWasSet)
+				{
+					bPromiseWasSet = true;
+					constexpr bool bWasCancelled = true;
+					Promise.EmplaceValue(FParallelExecutionResult{ bWasCancelled });
+				}
+			}
+
+			void StartSubmissions()
+			{
+				// FDeferredSubmitter needs access to SharedThis(), so it must be called outside of constructor.
+				for (FPendingOperation& PendingOperation : SubOperations)
+				{
+					FDeferredSubmitter& Submitter = PendingOperation.Submitter;
+					Submitter.EnqueueOrSubmit();
+				}
+			}
+
+			//~ Begin IParallelSubmissionOperation Interface
+			virtual TFuture<FParallelExecutionResult> OnCompletedFuture_AnyThread() override { return Promise.GetFuture(); }
+			//~ End IParallelSubmissionOperation Interface
+
+			void OnCompleteStream(const FGuid& ClientId, FSubmitStreamChangesResponse&& Response)
+			{
+				IntermediateResults.StreamResponses[ClientId] = MoveTemp(Response);
+				FindOperationByClient(ClientId)->bHasCompletedStreamChanges.store(true, std::memory_order_release);
+
+				FinishIfDone();
+			}
+
+			void OnCompleteAuthority(const FGuid& ClientId, FSubmitAuthorityChangesResponse&& Response)
+			{
+				IntermediateResults.AuthorityResponses[ClientId] = MoveTemp(Response);
+				FindOperationByClient(ClientId)->bHasCompletedAuthority.store(true, std::memory_order_release);
+				
+				const bool bIsFailure = Response.ErrorCode == EAuthoritySubmissionResponseErrorCode::Timeout || (Response.Response && !Response.Response->RejectedObjects.IsEmpty());
+				UE_CLOG(bIsFailure, LogConcert, Warning, TEXT("Remote authority change to client %s failed"), *ClientId.ToString());
+				FinishIfDone();
+			}
+
+			void OnSubmissionFailure(const FGuid& ClientId)
+			{
+				FPendingOperation* PendingOperation = FindOperationByClient(ClientId);
+				PendingOperation->bHasCompletedStreamChanges.store(true, std::memory_order_release);
+				PendingOperation->bHasCompletedAuthority.store(true, std::memory_order_release);
+				FinishIfDone();
+			}
+
+		private:
+			
+			TPromise<FParallelExecutionResult> Promise;
+			/** Must be atomic because FinishIfDone may be executed in parallel by the game thread (completing an operation) and the UDP (timing out an operation). */
+			std::atomic<bool> bPromiseWasSet = false;
+
+			struct FPendingOperation
+			{
+				/** Handles the submission. Unregisters from the client queue (if we're destroyed early). */
+				FDeferredSubmitter Submitter;
+			
+				/**
+				 * Whether the stream change has completed, i.e. OnCompleteStream has been called.
+				 * 
+				 * Must be atomic because FinishIfDone() may read while OnCompleteStream() / OnCompleteAuthority() are writing.
+				 * The (unlikely) scenario is that the game thread is finishing an operation while the UDP thread is timing out a task.
+				 */
+				std::atomic<bool> bHasCompletedStreamChanges = false;
+				/**
+				 * Whether the authority change has completed, i.e. OnCompleteAuthority has been called.
+				 * 
+				 * Must be atomic because FinishIfDone() may read while OnCompleteStream() / OnCompleteAuthority() are writing.
+				 * The (unlikely) scenario is that the game thread is finishing an operation while the UDP thread is timing out a task.
+				 */
+				std::atomic<bool> bHasCompletedAuthority = false;
+				
+				FPendingOperation(
+					FSyncOperation& InOwner,
+					FReplicationClientManager& InClientManager,
+					const FGuid& InClientId,
+					FSubmissionParams&& InSubmissionParams,
+					FSubmissionQueue& InSubmissionQueue
+					)
+					: Submitter(InOwner, InClientManager, InSubmissionQueue, InClientId, MoveTemp(InSubmissionParams))
+				{}
+			};
+			/**
+			 * Each latent tasks has its own, pre-allocated slot which effectively handles potentially concurrent writes.
+			 *
+			 * Initially every task value is initialized to 2: one for each operation (stream, authority).
+			 * Every time a sub-operation finishes, the value decrements.
+			 * Once all values have reached 0, that means Operations can be Emplace()'ed with IntermediateResults.
+			 */
+			TArray<FPendingOperation> SubOperations;
+
+			/** Holds intermediate results as they come in. */
+			FParallelExecutionResult IntermediateResults;
+
+			void FinishIfDone()
+			{
+				const bool bAllDone = Algo::AllOf(SubOperations, [](const FPendingOperation& OperationInfo)
+				{
+					// This is called right after OnCompleteStream or OnCompleteAuthority.
+					// This may run concurrently: e.g. the game thread is completing an operation and the UDP is timing out an operation.
+					// In that case, both threads wrote to these bools with memory_order_release (in OnCompleteStream and OnCompleteAuthority),
+					// meaning the writes happen-before these reads.
+					return OperationInfo.bHasCompletedStreamChanges.load(std::memory_order_acquire)
+						&& OperationInfo.bHasCompletedAuthority.load(std::memory_order_acquire);
+				});
+
+				// In the above scenario, both game thread and UDP thread may have bAllDone. Using exchange, the first thread to write will Emplace the promise.
+				if (bAllDone && !bPromiseWasSet.exchange(true))
+				{
+					Promise.EmplaceValue(MoveTemp(IntermediateResults));
+				}
+			}
+
+			FPendingOperation* FindOperationByClient(const FGuid& ClientId)
+			{
+				return SubOperations.FindByPredicate([&ClientId](const FPendingOperation& Operation)
+				{
+					return Operation.Submitter.GetClientId() == ClientId;
+				});
+			}
+		};
+
+		void FDeferredSubmitter::PerformSubmission_GameThread(ISubmissionWorkflow& Workflow)
+		{
+			const TSharedPtr<ISubmissionOperation> SubmissionOperation = Workflow.SubmitChanges(SubmissionParams);
+			if (ensure(SubmissionOperation))
+			{
+				SubmissionOperation->OnCompleteStreamChangesFuture_AnyThread()
+					.Next([this, WeakOwner = Owner.AsWeak()](FSubmitStreamChangesResponse&& Response)
+					{
+						if (const TSharedPtr<FSyncOperation> SyncOperation = WeakOwner.Pin())
+						{
+							SyncOperation->OnCompleteStream(ClientId, MoveTemp(Response));
+						}
+					});
+				SubmissionOperation->OnCompleteAuthorityChangeFuture_AnyThread()
+					.Next([this, WeakOwner = Owner.AsWeak()](FSubmitAuthorityChangesResponse&& Response)
+					{
+						if (const TSharedPtr<FSyncOperation> SyncOperation = WeakOwner.Pin())
+						{
+							SyncOperation->OnCompleteAuthority(ClientId, MoveTemp(Response));
+						}
+					});
+			}
+			else
+			{
+				Owner.OnSubmissionFailure(ClientId);
+			}
+		}
+	}
+	
+	TSharedPtr<IParallelSubmissionOperation> ExecuteParallelStreamChanges(FReplicationClientManager& ClientManager, TMap<FGuid, FSubmissionParams> ParallelOperations)
+	{
+		ParallelSubmission::Private::RemoveInvalidClientEntries(ClientManager, ParallelOperations);
+		if (!ensure(!ParallelOperations.IsEmpty()))
+		{
+			return nullptr;
+		}
+
+		using namespace ParallelSubmission::Private;
+		TSharedRef<FSyncOperation> Operation = MakeShared<FSyncOperation>(ClientManager, MoveTemp(ParallelOperations));
+		Operation->StartSubmissions();
+		return Operation;
+	}
+
+	TSharedPtr<IParallelSubmissionOperation> ExecuteParallelStreamChanges(FReplicationClientManager& ClientManager, TMap<FGuid, FConcertReplication_ChangeStream_Request> ParallelOperations)
+	{
+		if (!ensure(!ParallelOperations.IsEmpty()))
+		{
+			return nullptr;
+		}
+		
+		TMap<FGuid, FSubmissionParams> Transformed;
+		Transformed.Reserve(ParallelOperations.Num()); // check()s if Num() == 0
+		for (TPair<FGuid, FConcertReplication_ChangeStream_Request>& RequestPair : ParallelOperations)
+		{
+			Transformed.Emplace(RequestPair.Key, FSubmissionParams{ MoveTemp(RequestPair.Value) });
+		}
+		
+		return ExecuteParallelStreamChanges(ClientManager, MoveTemp(Transformed));
+	}
+}
diff --git a/Engine/Plugins/Developer/Concert/ConcertApp/MultiUserClient/Source/MultiUserClient/Private/Replication/Misc/Util/SynchronizedRequestUtils.h b/Engine/Plugins/Developer/Concert/ConcertApp/MultiUserClient/Source/MultiUserClient/Private/Replication/Misc/Util/SynchronizedRequestUtils.h
new file mode 100644
--- /dev/null
+++ b/Engine/Plugins/Developer/Concert/ConcertApp/MultiUserClient/Source/MultiUserClient/Private/Replication/Misc/Util/SynchronizedRequestUtils.h
@@ -0,0 +1,58 @@
+// Copyright Epic Games, Inc. All Rights Reserved.
+
+#pragma once
+
+#include "Replication/Submission/Data/AuthoritySubmission.h"
+#include "Replication/Submission/Data/StreamSubmission.h"
+
+#include "Async/Future.h"
+#include "Templates/UnrealTemplate.h"
+
+struct FConcertReplication_ChangeStream_Request;
+
+namespace UE::MultiUserClient
+{
+	class FReplicationClientManager;
+	struct FSubmissionParams;
+
+	struct FParallelExecutionResult
+	{
+		/** True if the TSharedPtr<IParallelSubmissionOperation> was destroyed before completion of submissions. */
+		bool bWasCancelled = false;
+
+		/** Maps client IDs to responses from the submissions. Empty if bWasCancelled. */
+		TMap<FGuid, FSubmitStreamChangesResponse> StreamResponses;
+		/** Maps client IDs to responses from the submissions. Empty if bWasCancelled. */
+		TMap<FGuid, FSubmitAuthorityChangesResponse> AuthorityResponses;
+	};
+
+	/** Result of ExecuteParallelStreamChanges. Handles correctly unregistering when the object is destroyed. */
+	class IParallelSubmissionOperation
+	{
+	public:
+
+		/**
+		 * @return Future completed when all parallel submissions have been completed.
+		 * @note This can complete on any thread. Usually this finishes on the game thread but timeouts usually occur on the messaging (UDP) thread.
+		 */
+		virtual TFuture<FParallelExecutionResult> OnCompletedFuture_AnyThread() = 0;
+
+		virtual ~IParallelSubmissionOperation() = default;
+	};
+
+	// TODO UE-201136: This code can be removed once a dedicated server operation for changing client streams has been implemented.
+	
+	/**
+	 * Executes the stream changes in parallel on the remote clients.
+	 *
+	 * @param ClientManager Used to get the clients to send to
+	 * @param ParallelOperations The operations to complete
+	 * 
+	 * @return Object you can get future through. This object also handles correct unregistering from other systems when you destroy it.
+	 * @note The returned object MUST be destroyed on the game thread.
+	 */
+	TSharedPtr<IParallelSubmissionOperation> ExecuteParallelStreamChanges(FReplicationClientManager& ClientManager, TMap<FGuid, FSubmissionParams> ParallelOperations);
+	/** Util for transforming into FSubmissionParams */
+	TSharedPtr<IParallelSubmissionOperation> ExecuteParallelStreamChanges(FReplicationClientManager& ClientManager, TMap<FGuid, FConcertReplication_ChangeStream_Request> ParallelOperations);
+}
+
diff --git a/Engine/Plugins/Developer/Concert/ConcertApp/MultiUserClient/Source/MultiUserClient/Private/Replication/MultiUserReplicationManager.h b/Engine/Plugins/Developer/Concert/ConcertApp/MultiUserClient/Source/MultiUserClient/Private/Replication/MultiUserReplicationManager.h
--- a/Engine/Plugins/Developer/Concert/ConcertApp/MultiUserClient/Source/MultiUserClient/Private/Replication/MultiUserReplicationManager.h
+++ b/Engine/Plugins/Developer/Concert/ConcertApp/MultiUserClient/Source/MultiUserClient/Private/Replication/MultiUserReplicationManager.h
@@ -5,11 +5,11 @@
 #include "Client/ReplicationClientManager.h"
 #include "IConcertSession.h"
 #include "Misc/ChangeLevelHandler.h"
+#include "Misc/Notification/ReplicationUserNotifier.h"
+#include "Misc/Query/RegularQueryService.h"
 #include "Muting/MuteStateManager.h"
-#include "Notification/ReplicationUserNotifier.h"
 #include "Replication/IMultiUserReplication.h"
 #include "Replication/Stream/Discovery/ReplicationDiscoveryContainer.h"
-#include "Util/Query/RegularQueryService.h"
 
 #include "Misc/Optional.h"
 #include "Misc/PreventReplicatedPropertyTransaction.h"
diff --git a/Engine/Plugins/Developer/Concert/ConcertApp/MultiUserClient/Source/MultiUserClient/Private/Replication/Notification/MutingNotifier.cpp b/Engine/Plugins/Developer/Concert/ConcertApp/MultiUserClient/Source/MultiUserClient/Private/Replication/Notification/MutingNotifier.cpp
deleted file mode 100644
diff --git a/Engine/Plugins/Developer/Concert/ConcertApp/MultiUserClient/Source/MultiUserClient/Private/Replication/Notification/MutingNotifier.h b/Engine/Plugins/Developer/Concert/ConcertApp/MultiUserClient/Source/MultiUserClient/Private/Replication/Notification/MutingNotifier.h
deleted file mode 100644
diff --git a/Engine/Plugins/Developer/Concert/ConcertApp/MultiUserClient/Source/MultiUserClient/Private/Replication/Notification/ReplicationUserNotifier.cpp b/Engine/Plugins/Developer/Concert/ConcertApp/MultiUserClient/Source/MultiUserClient/Private/Replication/Notification/ReplicationUserNotifier.cpp
deleted file mode 100644
diff --git a/Engine/Plugins/Developer/Concert/ConcertApp/MultiUserClient/Source/MultiUserClient/Private/Replication/Notification/ReplicationUserNotifier.h b/Engine/Plugins/Developer/Concert/ConcertApp/MultiUserClient/Source/MultiUserClient/Private/Replication/Notification/ReplicationUserNotifier.h
deleted file mode 100644
diff --git a/Engine/Plugins/Developer/Concert/ConcertApp/MultiUserClient/Source/MultiUserClient/Private/Replication/Stream/StreamChangeTracker.h b/Engine/Plugins/Developer/Concert/ConcertApp/MultiUserClient/Source/MultiUserClient/Private/Replication/Stream/StreamChangeTracker.h
--- a/Engine/Plugins/Developer/Concert/ConcertApp/MultiUserClient/Source/MultiUserClient/Private/Replication/Stream/StreamChangeTracker.h
+++ b/Engine/Plugins/Developer/Concert/ConcertApp/MultiUserClient/Source/MultiUserClient/Private/Replication/Stream/StreamChangeTracker.h
@@ -5,7 +5,7 @@
 #include "IClientStreamSynchronizer.h"
 #include "Replication/Data/ObjectReplicationMap.h"
 #include "Replication/Messages/ChangeStream.h"
-#include "Replication/Util/StreamRequestUtils.h"
+#include "Replication/Misc/Util/StreamRequestUtils.h"
 
 #include "Async/Future.h"
 #include "Containers/Set.h"
diff --git a/Engine/Plugins/Developer/Concert/ConcertApp/MultiUserClient/Source/MultiUserClient/Private/Replication/Stream/StreamSynchronizer_RemoteClient.cpp b/Engine/Plugins/Developer/Concert/ConcertApp/MultiUserClient/Source/MultiUserClient/Private/Replication/Stream/StreamSynchronizer_RemoteClient.cpp
--- a/Engine/Plugins/Developer/Concert/ConcertApp/MultiUserClient/Source/MultiUserClient/Private/Replication/Stream/StreamSynchronizer_RemoteClient.cpp
+++ b/Engine/Plugins/Developer/Concert/ConcertApp/MultiUserClient/Source/MultiUserClient/Private/Replication/Stream/StreamSynchronizer_RemoteClient.cpp
@@ -3,7 +3,7 @@
 #include "StreamSynchronizer_RemoteClient.h"
 
 #include "Assets/MultiUserReplicationClientPreset.h"
-#include "Replication/Util/Query/StreamAndAuthorityQueryService.h"
+#include "Replication/Misc/Query/StreamAndAuthorityQueryService.h"
 
 namespace UE::MultiUserClient
 {
diff --git a/Engine/Plugins/Developer/Concert/ConcertApp/MultiUserClient/Source/MultiUserClient/Private/Replication/Submission/AutoSubmissionPolicy.cpp b/Engine/Plugins/Developer/Concert/ConcertApp/MultiUserClient/Source/MultiUserClient/Private/Replication/Submission/AutoSubmissionPolicy.cpp
--- a/Engine/Plugins/Developer/Concert/ConcertApp/MultiUserClient/Source/MultiUserClient/Private/Replication/Submission/AutoSubmissionPolicy.cpp
+++ b/Engine/Plugins/Developer/Concert/ConcertApp/MultiUserClient/Source/MultiUserClient/Private/Replication/Submission/AutoSubmissionPolicy.cpp
@@ -6,7 +6,7 @@
 #include "ISubmissionWorkflow.h"
 #include "Replication/Authority/AuthorityChangeTracker.h"
 #include "Replication/Editor/Model/IEditableReplicationStreamModel.h"
-#include "Replication/Frequency/FrequencyChangeTracker.h"
+#include "Replication/Misc/Frequency/FrequencyChangeTracker.h"
 #include "Replication/Submission/Queue/SubmissionQueue.h"
 
 namespace UE::MultiUserClient
diff --git a/Engine/Plugins/Developer/Concert/ConcertApp/MultiUserClient/Source/MultiUserClient/Private/Replication/Submission/ChangeRequestBuilder.cpp b/Engine/Plugins/Developer/Concert/ConcertApp/MultiUserClient/Source/MultiUserClient/Private/Replication/Submission/ChangeRequestBuilder.cpp
--- a/Engine/Plugins/Developer/Concert/ConcertApp/MultiUserClient/Source/MultiUserClient/Private/Replication/Submission/ChangeRequestBuilder.cpp
+++ b/Engine/Plugins/Developer/Concert/ConcertApp/MultiUserClient/Source/MultiUserClient/Private/Replication/Submission/ChangeRequestBuilder.cpp
@@ -4,8 +4,8 @@
 
 #include "Replication/Authority/AuthorityChangeTracker.h"
 #include "Replication/Stream/StreamChangeTracker.h"
-#include "Replication/Util/GlobalAuthorityCache.h"
-#include "Replication/Util/StreamRequestUtils.h"
+#include "Replication/Misc/GlobalAuthorityCache.h"
+#include "Replication/Misc/Util/StreamRequestUtils.h"
 
 namespace UE::MultiUserClient
 {
diff --git a/Engine/Plugins/Developer/Concert/ConcertApp/MultiUserClient/Source/MultiUserClient/Private/Replication/Submission/External/ClientChangeConversionUtils.cpp b/Engine/Plugins/Developer/Concert/ConcertApp/MultiUserClient/Source/MultiUserClient/Private/Replication/Submission/External/ClientChangeConversionUtils.cpp
--- a/Engine/Plugins/Developer/Concert/ConcertApp/MultiUserClient/Source/MultiUserClient/Private/Replication/Submission/External/ClientChangeConversionUtils.cpp
+++ b/Engine/Plugins/Developer/Concert/ConcertApp/MultiUserClient/Source/MultiUserClient/Private/Replication/Submission/External/ClientChangeConversionUtils.cpp
@@ -5,7 +5,7 @@
 #include "Replication/ChangeOperationTypes.h"
 #include "Replication/Submission/Data/AuthoritySubmission.h"
 #include "Replication/Submission/Data/StreamSubmission.h"
-#include "Replication/Util/StreamRequestUtils.h"
+#include "Replication/Misc/Util/StreamRequestUtils.h"
 
 #include "Algo/RemoveIf.h"
 #include "Misc/Guid.h"
diff --git a/Engine/Plugins/Developer/Concert/ConcertApp/MultiUserClient/Source/MultiUserClient/Private/Replication/Submission/MultiEdit/ReassignObjectPropertiesLogic.cpp b/Engine/Plugins/Developer/Concert/ConcertApp/MultiUserClient/Source/MultiUserClient/Private/Replication/Submission/MultiEdit/ReassignObjectPropertiesLogic.cpp
--- a/Engine/Plugins/Developer/Concert/ConcertApp/MultiUserClient/Source/MultiUserClient/Private/Replication/Submission/MultiEdit/ReassignObjectPropertiesLogic.cpp
+++ b/Engine/Plugins/Developer/Concert/ConcertApp/MultiUserClient/Source/MultiUserClient/Private/Replication/Submission/MultiEdit/ReassignObjectPropertiesLogic.cpp
@@ -5,7 +5,7 @@
 #include "ConcertLogGlobal.h"
 #include "Replication/Client/ReplicationClientManager.h"
 #include "Replication/Submission/ISubmissionOperation.h"
-#include "Replication/Util/SynchronizedRequestUtils.h"
+#include "Replication/Misc/Util/SynchronizedRequestUtils.h"
 
 #include "Algo/AllOf.h"
 #include "Algo/AnyOf.h"
diff --git a/Engine/Plugins/Developer/Concert/ConcertApp/MultiUserClient/Source/MultiUserClient/Private/Replication/Submission/MultiEdit/ReassignObjectPropertiesLogic.h b/Engine/Plugins/Developer/Concert/ConcertApp/MultiUserClient/Source/MultiUserClient/Private/Replication/Submission/MultiEdit/ReassignObjectPropertiesLogic.h
--- a/Engine/Plugins/Developer/Concert/ConcertApp/MultiUserClient/Source/MultiUserClient/Private/Replication/Submission/MultiEdit/ReassignObjectPropertiesLogic.h
+++ b/Engine/Plugins/Developer/Concert/ConcertApp/MultiUserClient/Source/MultiUserClient/Private/Replication/Submission/MultiEdit/ReassignObjectPropertiesLogic.h
@@ -3,9 +3,8 @@
 #pragma once
 
 #include "Replication/IToken.h"
+#include "Replication/Misc/Util/SynchronizedRequestUtils.h"
 #include "Replication/Submission/Queue/DeferredSubmitter.h"
-#include "Replication/Util/StreamRequestUtils.h"
-#include "Replication/Util/SynchronizedRequestUtils.h"
 
 #include "Misc/OptionalFwd.h"
 #include "Templates/Function.h"
diff --git a/Engine/Plugins/Developer/Concert/ConcertApp/MultiUserClient/Source/MultiUserClient/Private/Replication/Submission/Remote/SubmissionWorkflow_RemoteClient.cpp b/Engine/Plugins/Developer/Concert/ConcertApp/MultiUserClient/Source/MultiUserClient/Private/Replication/Submission/Remote/SubmissionWorkflow_RemoteClient.cpp
--- a/Engine/Plugins/Developer/Concert/ConcertApp/MultiUserClient/Source/MultiUserClient/Private/Replication/Submission/Remote/SubmissionWorkflow_RemoteClient.cpp
+++ b/Engine/Plugins/Developer/Concert/ConcertApp/MultiUserClient/Source/MultiUserClient/Private/Replication/Submission/Remote/SubmissionWorkflow_RemoteClient.cpp
@@ -6,7 +6,7 @@
 #include "IConcertSession.h"
 #include "RemoteSubmissionMessages.h"
 #include "Replication/Submission/SingleClientSubmissionOperation.h"
-#include "Replication/Util/StreamRequestUtils.h"
+#include "Replication/Misc/Util/StreamRequestUtils.h"
 #include "Widgets/ActiveSession/Replication/Client/ClientUtils.h"
 
 namespace UE::MultiUserClient
diff --git a/Engine/Plugins/Developer/Concert/ConcertApp/MultiUserClient/Source/MultiUserClient/Private/Replication/Submission/SubmissionWorkflow_LocalClient.cpp b/Engine/Plugins/Developer/Concert/ConcertApp/MultiUserClient/Source/MultiUserClient/Private/Replication/Submission/SubmissionWorkflow_LocalClient.cpp
--- a/Engine/Plugins/Developer/Concert/ConcertApp/MultiUserClient/Source/MultiUserClient/Private/Replication/Submission/SubmissionWorkflow_LocalClient.cpp
+++ b/Engine/Plugins/Developer/Concert/ConcertApp/MultiUserClient/Source/MultiUserClient/Private/Replication/Submission/SubmissionWorkflow_LocalClient.cpp
@@ -4,7 +4,7 @@
 
 #include "IConcertSyncClient.h"
 #include "Replication/IConcertClientReplicationManager.h"
-#include "Replication/Util/StreamRequestUtils.h"
+#include "Replication/Misc/Util/StreamRequestUtils.h"
 #include "Widgets/ActiveSession/Replication/Client/ClientUtils.h"
 
 namespace UE::MultiUserClient
diff --git a/Engine/Plugins/Developer/Concert/ConcertApp/MultiUserClient/Source/MultiUserClient/Private/Replication/Util/FrequencyUtils.cpp b/Engine/Plugins/Developer/Concert/ConcertApp/MultiUserClient/Source/MultiUserClient/Private/Replication/Util/FrequencyUtils.cpp
deleted file mode 100644
diff --git a/Engine/Plugins/Developer/Concert/ConcertApp/MultiUserClient/Source/MultiUserClient/Private/Replication/Util/FrequencyUtils.h b/Engine/Plugins/Developer/Concert/ConcertApp/MultiUserClient/Source/MultiUserClient/Private/Replication/Util/FrequencyUtils.h
deleted file mode 100644
diff --git a/Engine/Plugins/Developer/Concert/ConcertApp/MultiUserClient/Source/MultiUserClient/Private/Replication/Util/GlobalAuthorityCache.cpp b/Engine/Plugins/Developer/Concert/ConcertApp/MultiUserClient/Source/MultiUserClient/Private/Replication/Util/GlobalAuthorityCache.cpp
deleted file mode 100644
diff --git a/Engine/Plugins/Developer/Concert/ConcertApp/MultiUserClient/Source/MultiUserClient/Private/Replication/Util/GlobalAuthorityCache.h b/Engine/Plugins/Developer/Concert/ConcertApp/MultiUserClient/Source/MultiUserClient/Private/Replication/Util/GlobalAuthorityCache.h
deleted file mode 100644
diff --git a/Engine/Plugins/Developer/Concert/ConcertApp/MultiUserClient/Source/MultiUserClient/Private/Replication/Util/Query/MuteStateQueryService.cpp b/Engine/Plugins/Developer/Concert/ConcertApp/MultiUserClient/Source/MultiUserClient/Private/Replication/Util/Query/MuteStateQueryService.cpp
deleted file mode 100644
diff --git a/Engine/Plugins/Developer/Concert/ConcertApp/MultiUserClient/Source/MultiUserClient/Private/Replication/Util/Query/MuteStateQueryService.h b/Engine/Plugins/Developer/Concert/ConcertApp/MultiUserClient/Source/MultiUserClient/Private/Replication/Util/Query/MuteStateQueryService.h
deleted file mode 100644
diff --git a/Engine/Plugins/Developer/Concert/ConcertApp/MultiUserClient/Source/MultiUserClient/Private/Replication/Util/Query/RegularQueryService.cpp b/Engine/Plugins/Developer/Concert/ConcertApp/MultiUserClient/Source/MultiUserClient/Private/Replication/Util/Query/RegularQueryService.cpp
deleted file mode 100644
diff --git a/Engine/Plugins/Developer/Concert/ConcertApp/MultiUserClient/Source/MultiUserClient/Private/Replication/Util/Query/RegularQueryService.h b/Engine/Plugins/Developer/Concert/ConcertApp/MultiUserClient/Source/MultiUserClient/Private/Replication/Util/Query/RegularQueryService.h
deleted file mode 100644
diff --git a/Engine/Plugins/Developer/Concert/ConcertApp/MultiUserClient/Source/MultiUserClient/Private/Replication/Util/Query/StreamAndAuthorityQueryService.cpp b/Engine/Plugins/Developer/Concert/ConcertApp/MultiUserClient/Source/MultiUserClient/Private/Replication/Util/Query/StreamAndAuthorityQueryService.cpp
deleted file mode 100644
diff --git a/Engine/Plugins/Developer/Concert/ConcertApp/MultiUserClient/Source/MultiUserClient/Private/Replication/Util/Query/StreamAndAuthorityQueryService.h b/Engine/Plugins/Developer/Concert/ConcertApp/MultiUserClient/Source/MultiUserClient/Private/Replication/Util/Query/StreamAndAuthorityQueryService.h
deleted file mode 100644
diff --git a/Engine/Plugins/Developer/Concert/ConcertApp/MultiUserClient/Source/MultiUserClient/Private/Replication/Util/StreamRequestUtils.cpp b/Engine/Plugins/Developer/Concert/ConcertApp/MultiUserClient/Source/MultiUserClient/Private/Replication/Util/StreamRequestUtils.cpp
deleted file mode 100644
diff --git a/Engine/Plugins/Developer/Concert/ConcertApp/MultiUserClient/Source/MultiUserClient/Private/Replication/Util/StreamRequestUtils.h b/Engine/Plugins/Developer/Concert/ConcertApp/MultiUserClient/Source/MultiUserClient/Private/Replication/Util/StreamRequestUtils.h
deleted file mode 100644
diff --git a/Engine/Plugins/Developer/Concert/ConcertApp/MultiUserClient/Source/MultiUserClient/Private/Replication/Util/SynchronizedRequestUtils.cpp b/Engine/Plugins/Developer/Concert/ConcertApp/MultiUserClient/Source/MultiUserClient/Private/Replication/Util/SynchronizedRequestUtils.cpp
deleted file mode 100644
diff --git a/Engine/Plugins/Developer/Concert/ConcertApp/MultiUserClient/Source/MultiUserClient/Private/Replication/Util/SynchronizedRequestUtils.h b/Engine/Plugins/Developer/Concert/ConcertApp/MultiUserClient/Source/MultiUserClient/Private/Replication/Util/SynchronizedRequestUtils.h
deleted file mode 100644
diff --git a/Engine/Plugins/Developer/Concert/ConcertApp/MultiUserClient/Source/MultiUserClient/Private/Widgets/ActiveSession/Replication/Client/Context/ContextMenuUtils_Frequency.cpp b/Engine/Plugins/Developer/Concert/ConcertApp/MultiUserClient/Source/MultiUserClient/Private/Widgets/ActiveSession/Replication/Client/Context/ContextMenuUtils_Frequency.cpp
--- a/Engine/Plugins/Developer/Concert/ConcertApp/MultiUserClient/Source/MultiUserClient/Private/Widgets/ActiveSession/Replication/Client/Context/ContextMenuUtils_Frequency.cpp
+++ b/Engine/Plugins/Developer/Concert/ConcertApp/MultiUserClient/Source/MultiUserClient/Private/Widgets/ActiveSession/Replication/Client/Context/ContextMenuUtils_Frequency.cpp
@@ -3,7 +3,8 @@
 #include "ContextMenuUtils.h"
 
 #include "Replication/Client/ReplicationClientManager.h"
-#include "Replication/Util/FrequencyUtils.h"
+#include "Replication/Misc/Frequency/FrequencyUtils.h"
+#include "Replication/Misc/Util/SynchronizedRequestUtils.h"
 
 #include "Algo/AnyOf.h"
 #include "Async/Async.h"
diff --git a/Engine/Plugins/Developer/Concert/ConcertApp/MultiUserClient/Source/MultiUserClient/Private/Widgets/ActiveSession/Replication/Client/SReplicationStatus.cpp b/Engine/Plugins/Developer/Concert/ConcertApp/MultiUserClient/Source/MultiUserClient/Private/Widgets/ActiveSession/Replication/Client/SReplicationStatus.cpp
--- a/Engine/Plugins/Developer/Concert/ConcertApp/MultiUserClient/Source/MultiUserClient/Private/Widgets/ActiveSession/Replication/Client/SReplicationStatus.cpp
+++ b/Engine/Plugins/Developer/Concert/ConcertApp/MultiUserClient/Source/MultiUserClient/Private/Widgets/ActiveSession/Replication/Client/SReplicationStatus.cpp
@@ -2,7 +2,7 @@
 
 #include "SReplicationStatus.h"
 
-#include "Replication/Util/GlobalAuthorityCache.h"
+#include "Replication/Misc/GlobalAuthorityCache.h"
 
 #include "Algo/AnyOf.h"
 #include "Styling/AppStyle.h"
