Remove common properties from Replication edit button. Feature is post-boned for now.

diff --git a/Engine/Plugins/Developer/Concert/ConcertApp/MultiUserClient/Source/MultiUserClient/Private/Widgets/ActiveSession/Replication/Client/PropertySelection/RootPropertySourceModel.cpp b/Engine/Plugins/Developer/Concert/ConcertApp/MultiUserClient/Source/MultiUserClient/Private/Widgets/ActiveSession/Replication/Client/PropertySelection/RootPropertySourceModel.cpp
--- a/Engine/Plugins/Developer/Concert/ConcertApp/MultiUserClient/Source/MultiUserClient/Private/Widgets/ActiveSession/Replication/Client/PropertySelection/RootPropertySourceModel.cpp
+++ b/Engine/Plugins/Developer/Concert/ConcertApp/MultiUserClient/Source/MultiUserClient/Private/Widgets/ActiveSession/Replication/Client/PropertySelection/RootPropertySourceModel.cpp
@@ -19,7 +19,6 @@ namespace UE::MultiUserClient
 
 	void FRootPropertySourceModel::RefreshSelectableProperties(TConstArrayView<ConcertSharedSlate::FObjectGroup> DisplayedObjectGroups)
 	{
-		PerObjectGroup_CommonPropertiesSources.Empty(DisplayedObjectGroups.Num());
 		PerObjectGroup_AllPropertiesSources.Empty(DisplayedObjectGroups.Num());
 		
 		for (const ConcertSharedSlate::FObjectGroup& Group : DisplayedObjectGroups)
diff --git a/Engine/Plugins/Developer/Concert/ConcertApp/MultiUserClient/Source/MultiUserClient/Private/Widgets/ActiveSession/Replication/Client/PropertySelection/RootPropertySourceModel.h b/Engine/Plugins/Developer/Concert/ConcertApp/MultiUserClient/Source/MultiUserClient/Private/Widgets/ActiveSession/Replication/Client/PropertySelection/RootPropertySourceModel.h
--- a/Engine/Plugins/Developer/Concert/ConcertApp/MultiUserClient/Source/MultiUserClient/Private/Widgets/ActiveSession/Replication/Client/PropertySelection/RootPropertySourceModel.h
+++ b/Engine/Plugins/Developer/Concert/ConcertApp/MultiUserClient/Source/MultiUserClient/Private/Widgets/ActiveSession/Replication/Client/PropertySelection/RootPropertySourceModel.h
@@ -26,8 +26,6 @@ namespace UE::MultiUserClient
 		/** Refreshes the properties that the user can select given the objects currently displayed in the bottom view. */
 		void RefreshSelectableProperties(TConstArrayView<ConcertSharedSlate::FObjectGroup> DisplayedObjectGroups);
 
-		/** Contains one source for each object group. Each source displays common properties, configured in the project settings. */
-		const TArray<TSharedRef<IPropertyItemSource>>& GetPerObjectGroup_CommonPropertiesSources() const { return PerObjectGroup_CommonPropertiesSources; }
 		/** Contains one source for each object group. Each source displays all root properties in that class. */
 		const TArray<TSharedRef<IPropertyItemSource>>& GetPerObjectGroup_AllPropertiesSources() const { return PerObjectGroup_AllPropertiesSources; }
 
@@ -36,9 +34,6 @@ namespace UE::MultiUserClient
 		/** Determines the display string of an object. */
 		const FGetObjectDisplayString GetObjectDisplayStringDelegate;
 
-		// TODO UE-216094: Populate this array.
-		/** Contains one source for each object group. Each source displays common properties, configured in the project settings. */
-		TArray<TSharedRef<IPropertyItemSource>> PerObjectGroup_CommonPropertiesSources;
 		/** Contains one source for each object group. Each source displays all root properties in that class. */
 		TArray<TSharedRef<IPropertyItemSource>> PerObjectGroup_AllPropertiesSources;
 	};
diff --git a/Engine/Plugins/Developer/Concert/ConcertApp/MultiUserClient/Source/MultiUserClient/Private/Widgets/ActiveSession/Replication/Client/PropertySelection/SPropertySelectionComboButton.cpp b/Engine/Plugins/Developer/Concert/ConcertApp/MultiUserClient/Source/MultiUserClient/Private/Widgets/ActiveSession/Replication/Client/PropertySelection/SPropertySelectionComboButton.cpp
--- a/Engine/Plugins/Developer/Concert/ConcertApp/MultiUserClient/Source/MultiUserClient/Private/Widgets/ActiveSession/Replication/Client/PropertySelection/SPropertySelectionComboButton.cpp
+++ b/Engine/Plugins/Developer/Concert/ConcertApp/MultiUserClient/Source/MultiUserClient/Private/Widgets/ActiveSession/Replication/Client/PropertySelection/SPropertySelectionComboButton.cpp
@@ -79,23 +79,6 @@ namespace UE::MultiUserClient
 			FText::GetEmpty()
 		);
 		
-		MenuBuilder.BeginSection(NAME_None, LOCTEXT("Section.CommonProperties", "Common Properties"));
-		MenuBuilder.AddMenuEntry(
-			LOCTEXT("Configure.Label", "Configure Common Properties"),
-			LOCTEXT("Configure.ToolTip", "Opens the project settings in which you can set up common properties"),
-			FSlateIcon(),
-			FUIAction(
-				FExecuteAction::CreateLambda([]()
-				{
-					// TODO UE-216094: Open project settings
-				}))
-			);
-		for (const TSharedRef<IPropertyItemSource>& ItemSource : PropertySourceModel->GetPerObjectGroup_CommonPropertiesSources())
-		{
-			FModelBuilder::AddOptionToMenu(ItemSource, PickerArgs, MenuBuilder);
-		}
-		MenuBuilder.EndSection();
-		
 		MenuBuilder.BeginSection(NAME_None, LOCTEXT("Section.AllProperties", "All Properties"));
 		for (const TSharedRef<IPropertyItemSource>& ItemSource : PropertySourceModel->GetPerObjectGroup_AllPropertiesSources())
 		{
