Fix: "Assigned clients" colum does not show clients for subobjects, only for actors.

diff --git a/Engine/Plugins/Developer/Concert/ConcertApp/MultiUserClient/Source/MultiUserClient/Private/Widgets/ActiveSession/Replication/Client/Multi/Columns/MultiStreamColumns.h b/Engine/Plugins/Developer/Concert/ConcertApp/MultiUserClient/Source/MultiUserClient/Private/Widgets/ActiveSession/Replication/Client/Multi/Columns/MultiStreamColumns.h
--- a/Engine/Plugins/Developer/Concert/ConcertApp/MultiUserClient/Source/MultiUserClient/Private/Widgets/ActiveSession/Replication/Client/Multi/Columns/MultiStreamColumns.h
+++ b/Engine/Plugins/Developer/Concert/ConcertApp/MultiUserClient/Source/MultiUserClient/Private/Widgets/ActiveSession/Replication/Client/Multi/Columns/MultiStreamColumns.h
@@ -56,16 +56,16 @@ namespace UE::MultiUserClient::MultiStreamColumns
 	 *
 	 * @param ConcertClient Used to look up client names
 	 * @param MultiStreamModelAttribute Used for requesting resorting the column
-	 * @param ObjectHierarchyModelAttribute Used to get child objects
-	 * @param ReassignmentLogic Performs the act of reassigning
+	 * @param ObjectHierarchy Used to display clients when a child object is replicated
+	 * @param ReassignmentLogic Used to get the assigned clients and detect when they change.
 	 * @param ClientManager Used to access all clients for display in the combo box drop-down
 	 * @param ColumnsSortPriority The order relative to the other columns
 	 */
 	ConcertSharedSlate::FObjectColumnEntry AssignedClientsColumn(
 		TSharedRef<IConcertClient> ConcertClient,
 		TAttribute<TSharedPtr<ConcertSharedSlate::IMultiReplicationStreamEditor>> MultiStreamModelAttribute,
-		TAttribute<ConcertSharedSlate::IObjectHierarchyModel*> ObjectHierarchyModelAttribute,
-		FReassignObjectPropertiesLogic& ReassignmentLogic,
+		const ConcertSharedSlate::IObjectHierarchyModel& ObjectHierarchy UE_LIFETIMEBOUND,
+		FReassignObjectPropertiesLogic& ReassignmentLogic UE_LIFETIMEBOUND,
 		const FReplicationClientManager& ClientManager UE_LIFETIMEBOUND,
 		int32 ColumnsSortPriority = static_cast<int32>(EColumnSortOrder::ReassignOwnership)
 		);
diff --git a/Engine/Plugins/Developer/Concert/ConcertApp/MultiUserClient/Source/MultiUserClient/Private/Widgets/ActiveSession/Replication/Client/Multi/Columns/MultiStreamColumns_AssignedClients.cpp b/Engine/Plugins/Developer/Concert/ConcertApp/MultiUserClient/Source/MultiUserClient/Private/Widgets/ActiveSession/Replication/Client/Multi/Columns/MultiStreamColumns_AssignedClients.cpp
--- a/Engine/Plugins/Developer/Concert/ConcertApp/MultiUserClient/Source/MultiUserClient/Private/Widgets/ActiveSession/Replication/Client/Multi/Columns/MultiStreamColumns_AssignedClients.cpp
+++ b/Engine/Plugins/Developer/Concert/ConcertApp/MultiUserClient/Source/MultiUserClient/Private/Widgets/ActiveSession/Replication/Client/Multi/Columns/MultiStreamColumns_AssignedClients.cpp
@@ -4,6 +4,7 @@
 
 #include "IConcertClient.h"
 #include "MultiUserReplicationStyle.h"
+#include "Replication/Editor/Model/Object/IObjectHierarchyModel.h"
 #include "Replication/Editor/View/IMultiReplicationStreamEditor.h"
 #include "Replication/Submission/MultiEdit/ReassignObjectPropertiesLogic.h"
 #include "Widgets/ActiveSession/Replication/Client/ClientUtils.h"
@@ -20,17 +21,34 @@ namespace UE::MultiUserClient::MultiStreamColumns
 
 	namespace Private
 	{
-		static TArray<FGuid> GetDisplayedClients(const FReassignObjectPropertiesLogic& ReassignmentLogic, const FSoftObjectPath& ManagedObject)
+		static TArray<FGuid> GetDisplayedClients(
+			const ConcertSharedSlate::IObjectHierarchyModel& ObjectHierarchy,
+			const FReassignObjectPropertiesLogic& ReassignmentLogic,
+			const FSoftObjectPath& ManagedObject
+			)
 		{
 			TArray<FGuid> ClientsWithOwnership;
-			ReassignmentLogic.EnumerateClientOwnershipState(ManagedObject, [&ClientsWithOwnership](const FGuid& ClientId, FReassignObjectPropertiesLogic::EOwnershipState Ownership)
+			const auto ProcessObject = [&ReassignmentLogic, &ClientsWithOwnership](const FSoftObjectPath& ObjectPath)
 			{
-				if (Ownership == FReassignObjectPropertiesLogic::EOwnershipState::HasObjectRegistered)
+				ReassignmentLogic.EnumerateClientOwnershipState(ObjectPath, [&ClientsWithOwnership](const FGuid& ClientId, FReassignObjectPropertiesLogic::EOwnershipState Ownership)
 				{
-					ClientsWithOwnership.Add(ClientId);
-				}
-				return EBreakBehavior::Continue;
-			});
+					if (Ownership == FReassignObjectPropertiesLogic::EOwnershipState::HasObjectRegistered)
+					{
+						ClientsWithOwnership.Add(ClientId);
+					}
+					return EBreakBehavior::Continue;
+				});
+			};
+			
+			ProcessObject(ManagedObject);
+			ObjectHierarchy.ForEachChildRecursive(
+				TSoftObjectPtr<>{ ManagedObject },
+				[&ProcessObject](const TSoftObjectPtr<>&, const TSoftObjectPtr<>& ChildObject, ConcertSharedSlate::EChildRelationship)
+				{
+					ProcessObject(ChildObject.GetUniqueID());
+					return EBreakBehavior::Continue;
+				});
+			
 			return ClientsWithOwnership;
 		}
 		
@@ -38,22 +56,26 @@ namespace UE::MultiUserClient::MultiStreamColumns
 		{
 			TSharedPtr<ConcertClientSharedSlate::SHorizontalClientList> ClientList;
 			FSoftObjectPath ManagedObject;
+			const ConcertSharedSlate::IObjectHierarchyModel* ObjectHierarchy = nullptr;
 			FReassignObjectPropertiesLogic* ReassignmentLogic = nullptr;
+			
 		public:
 		
 			SLATE_BEGIN_ARGS(SAssignedClientsWidget){}
+				SLATE_ARGUMENT(FSoftObjectPath, ManagedObject)
+				SLATE_ATTRIBUTE(FText, HighlightText)
 			SLATE_END_ARGS()
 
 			void Construct(
 				const FArguments& InArgs,
 				const TSharedRef<IConcertClient>& InConcertClient,
-				FReassignObjectPropertiesLogic& InReassignmentLogic,
-				FSoftObjectPath InManagedObject,
-				TSharedPtr<FText> InHighlightText
+				const ConcertSharedSlate::IObjectHierarchyModel& InObjectHierarchy,
+				FReassignObjectPropertiesLogic& InReassignmentLogic
 				)
 			{
+				ObjectHierarchy = &InObjectHierarchy;
 				ReassignmentLogic = &InReassignmentLogic;
-				ManagedObject = MoveTemp(InManagedObject);
+				ManagedObject = InArgs._ManagedObject;
 			
 				ChildSlot
 				[
@@ -61,7 +83,7 @@ namespace UE::MultiUserClient::MultiStreamColumns
 					.IsLocalClient(ConcertClientSharedSlate::MakeIsLocalClientGetter(InConcertClient))
 					.GetClientInfo(ConcertClientSharedSlate::MakeClientInfoGetter(InConcertClient))
 					.Font(FAppStyle::GetFontStyle("PropertyWindow.NormalFont"))
-					.HighlightText_Lambda([InHighlightText](){ return InHighlightText ? *InHighlightText : FText::GetEmpty(); })
+					.HighlightText(InArgs._HighlightText)
 					.ListToolTipText(LOCTEXT("Clients.ToolTip", "These clients will replicate their assigned properties when replication is active.\nYou can pause & resume replication at the beginnig of this row."))
 					.EmptyListSlot()
 					[
@@ -84,7 +106,7 @@ namespace UE::MultiUserClient::MultiStreamColumns
 			void UpdateComboButtonContent() const
 			{
 				ClientList->RefreshList(
-					GetDisplayedClients(*ReassignmentLogic, ManagedObject)
+					GetDisplayedClients(*ObjectHierarchy, *ReassignmentLogic, ManagedObject)
 					);
 			}
 		};
@@ -93,7 +115,7 @@ namespace UE::MultiUserClient::MultiStreamColumns
 	ConcertSharedSlate::FObjectColumnEntry AssignedClientsColumn(
 		TSharedRef<IConcertClient> ConcertClient,
 		TAttribute<TSharedPtr<ConcertSharedSlate::IMultiReplicationStreamEditor>> MultiStreamModelAttribute,
-		TAttribute<ConcertSharedSlate::IObjectHierarchyModel*> ObjectHierarchyModelAttribute,
+		const ConcertSharedSlate::IObjectHierarchyModel& ObjectHierarchy,
 		FReassignObjectPropertiesLogic& ReassignmentLogic,
 		const FReplicationClientManager& ClientManager,
 		const int32 ColumnsSortPriority
@@ -106,13 +128,13 @@ namespace UE::MultiUserClient::MultiStreamColumns
 			FObjectColumn_ReassignOwnership(
 				TSharedRef<IConcertClient> ConcertClient,
 				TAttribute<TSharedPtr<ConcertSharedSlate::IMultiReplicationStreamEditor>> MultiStreamModelAttribute,
-				TAttribute<ConcertSharedSlate::IObjectHierarchyModel*> ObjectHierarchyModelAttribute,
-				FReassignObjectPropertiesLogic& ReassignmentLogic,
-				const FReplicationClientManager& ClientManager
+				const ConcertSharedSlate::IObjectHierarchyModel& ObjectHierarchy UE_LIFETIMEBOUND,
+				FReassignObjectPropertiesLogic& ReassignmentLogic UE_LIFETIMEBOUND,
+				const FReplicationClientManager& ClientManager UE_LIFETIMEBOUND
 				)
 				: ConcertClient(MoveTemp(ConcertClient))
 				, MultiStreamModelAttribute(MoveTemp(MultiStreamModelAttribute))
-				, ObjectHierarchyModelAttribute(MoveTemp(ObjectHierarchyModelAttribute))
+				, ObjectHierarchy(ObjectHierarchy)
 				, ReassignmentLogic(ReassignmentLogic)
 				, ClientManager(ClientManager)
 			{}
@@ -127,11 +149,13 @@ namespace UE::MultiUserClient::MultiStreamColumns
 			
 			virtual TSharedRef<SWidget> GenerateColumnWidget(const FBuildArgs& InArgs) override
 			{
-				return SNew(Private::SAssignedClientsWidget, ConcertClient, ReassignmentLogic, InArgs.RowItem.RowData.GetObjectPath(), InArgs.HighlightText);
+				return SNew(Private::SAssignedClientsWidget, ConcertClient, ObjectHierarchy, ReassignmentLogic)
+					.ManagedObject(InArgs.RowItem.RowData.GetObjectPath())
+					.HighlightText_Lambda([HighlightText = InArgs.HighlightText](){ return HighlightText ? *HighlightText : FText::GetEmpty(); });
 			}
 			virtual void PopulateSearchString(const ConcertSharedSlate::FObjectTreeRowContext& InItem, TArray<FString>& InOutSearchStrings) const override
 			{
-				for (const FGuid& ClientId : Private::GetDisplayedClients(ReassignmentLogic, InItem.RowData.GetObjectPath()))
+				for (const FGuid& ClientId : Private::GetDisplayedClients(ObjectHierarchy, ReassignmentLogic, InItem.RowData.GetObjectPath()))
 				{
 					InOutSearchStrings.Add(ClientUtils::GetClientDisplayName(*ConcertClient->GetCurrentSession(), ClientId));
 				}
@@ -155,14 +179,14 @@ namespace UE::MultiUserClient::MultiStreamColumns
 			
 			const TSharedRef<IConcertClient> ConcertClient;
 			const TAttribute<TSharedPtr<ConcertSharedSlate::IMultiReplicationStreamEditor>> MultiStreamModelAttribute;
-			const TAttribute<ConcertSharedSlate::IObjectHierarchyModel*> ObjectHierarchyModelAttribute;
+			const ConcertSharedSlate::IObjectHierarchyModel& ObjectHierarchy;
 			FReassignObjectPropertiesLogic& ReassignmentLogic;
 			const FReplicationClientManager& ClientManager;
 			
 			TOptional<FString> GetDisplayString(const FSoftObjectPath& ManagedObject) const
 			{
 				using SWidgetType = ConcertClientSharedSlate::SHorizontalClientList;
-				const TArray<FGuid> Clients = Private::GetDisplayedClients(ReassignmentLogic, ManagedObject);
+				const TArray<FGuid> Clients = Private::GetDisplayedClients(ObjectHierarchy, ReassignmentLogic, ManagedObject);
 				const ConcertSharedSlate::FIsLocalClient IsLocalClientDelegate = ConcertClientSharedSlate::MakeIsLocalClientGetter(ConcertClient);
 				return SWidgetType::GetDisplayString(
 					Clients,
@@ -175,9 +199,9 @@ namespace UE::MultiUserClient::MultiStreamColumns
 
 		return {
 			ConcertSharedSlate::TReplicationColumnDelegates<ConcertSharedSlate::FObjectTreeRowContext>::FCreateColumn::CreateLambda(
-				[ConcertClient = MoveTemp(ConcertClient), MultiStreamModelAttribute = MoveTemp(MultiStreamModelAttribute), ObjectHierarchyModelAttribute = MoveTemp(ObjectHierarchyModelAttribute), &ReassignmentLogic, &ClientManager]()
+				[ConcertClient = MoveTemp(ConcertClient), MultiStreamModelAttribute = MoveTemp(MultiStreamModelAttribute), &ObjectHierarchy, &ReassignmentLogic, &ClientManager]()
 				{
-					return MakeShared<FObjectColumn_ReassignOwnership>(ConcertClient, MultiStreamModelAttribute, ObjectHierarchyModelAttribute, ReassignmentLogic, ClientManager);
+					return MakeShared<FObjectColumn_ReassignOwnership>(ConcertClient, MultiStreamModelAttribute, ObjectHierarchy, ReassignmentLogic, ClientManager);
 				}),
 			AssignedClientsColumnId,
 			{ ColumnsSortPriority }
diff --git a/Engine/Plugins/Developer/Concert/ConcertApp/MultiUserClient/Source/MultiUserClient/Private/Widgets/ActiveSession/Replication/Client/Multi/SMultiClientView.cpp b/Engine/Plugins/Developer/Concert/ConcertApp/MultiUserClient/Source/MultiUserClient/Private/Widgets/ActiveSession/Replication/Client/Multi/SMultiClientView.cpp
--- a/Engine/Plugins/Developer/Concert/ConcertApp/MultiUserClient/Source/MultiUserClient/Private/Widgets/ActiveSession/Replication/Client/Multi/SMultiClientView.cpp
+++ b/Engine/Plugins/Developer/Concert/ConcertApp/MultiUserClient/Source/MultiUserClient/Private/Widgets/ActiveSession/Replication/Client/Multi/SMultiClientView.cpp
@@ -95,11 +95,6 @@ namespace UE::MultiUserClient
 		   {
 			   return StreamEditor;
 		   });
-		const TAttribute<IObjectHierarchyModel*> ObjectHierarchyAttribute =
-		   TAttribute<IObjectHierarchyModel*>::CreateLambda([this]()
-		   {
-			   return ObjectHierarchy.Get();
-		   });
 		FGetAutoAssignTarget GetAutoAssignTargetDelegate = FGetAutoAssignTarget::CreateLambda([this](TConstArrayView<UObject*>)
 		{
 			const TSharedRef<IEditableReplicationStreamModel>& LocalStream = ClientManager->GetLocalClient().GetClientEditModel();
@@ -147,7 +142,7 @@ namespace UE::MultiUserClient
 			.ObjectColumns =
 			{
 				MultiStreamColumns::MuteToggleColumn(MuteManager.GetChangeTracker()),
-				MultiStreamColumns::AssignedClientsColumn(InConcertClient, MultiStreamEditorAttribute, ObjectHierarchyAttribute, ClientManager->GetReassignmentLogic(), *ClientManager)
+				MultiStreamColumns::AssignedClientsColumn(InConcertClient, MultiStreamEditorAttribute, *ObjectHierarchy, ClientManager->GetReassignmentLogic(), *ClientManager)
 			},
 			.ShouldDisplayObjectDelegate = FShouldDisplayObject::CreateSP(this, &SMultiClientView::ShouldDisplayObject)
 		};
