Fix client name appearing multiple in "Author" column when multiple subobjects were being replicated.

diff --git a/Engine/Plugins/Developer/Concert/ConcertApp/MultiUserClient/Source/MultiUserClient/Private/Widgets/ActiveSession/Replication/Client/Multi/Columns/MultiStreamColumns_AssignedClients.cpp b/Engine/Plugins/Developer/Concert/ConcertApp/MultiUserClient/Source/MultiUserClient/Private/Widgets/ActiveSession/Replication/Client/Multi/Columns/MultiStreamColumns_AssignedClients.cpp
--- a/Engine/Plugins/Developer/Concert/ConcertApp/MultiUserClient/Source/MultiUserClient/Private/Widgets/ActiveSession/Replication/Client/Multi/Columns/MultiStreamColumns_AssignedClients.cpp
+++ b/Engine/Plugins/Developer/Concert/ConcertApp/MultiUserClient/Source/MultiUserClient/Private/Widgets/ActiveSession/Replication/Client/Multi/Columns/MultiStreamColumns_AssignedClients.cpp
@@ -34,7 +34,7 @@ namespace UE::MultiUserClient::MultiStreamColumns
 				{
 					if (Ownership == FReassignObjectPropertiesLogic::EOwnershipState::HasObjectRegistered)
 					{
-						ClientsWithOwnership.Add(ClientId);
+						ClientsWithOwnership.AddUnique(ClientId);
 					}
 					return EBreakBehavior::Continue;
 				});
