Fix: Properties are now only auto added when using a Add button.
Before properties would confusingly be auto-added again when a user e.g. added a VCam, manually removed all properties, and assigned a property to the VCam actor.

diff --git a/Engine/Plugins/Developer/Concert/ConcertApp/MultiUserClient/Source/MultiUserClient/Private/Replication/Client/ReplicationClient.cpp b/Engine/Plugins/Developer/Concert/ConcertApp/MultiUserClient/Source/MultiUserClient/Private/Replication/Client/ReplicationClient.cpp
--- a/Engine/Plugins/Developer/Concert/ConcertApp/MultiUserClient/Source/MultiUserClient/Private/Replication/Client/ReplicationClient.cpp
+++ b/Engine/Plugins/Developer/Concert/ConcertApp/MultiUserClient/Source/MultiUserClient/Private/Replication/Client/ReplicationClient.cpp
@@ -35,14 +35,12 @@ namespace UE::MultiUserClient
 			FGetStreamContent::CreateLambda([this](){ return &StreamSynchronizer->GetServerState(); }),
 			SubmissionQueue
 			)
+		, StreamExtender(MakeShared<FMultiUserStreamExtender>(EndpointId, InDiscoveryContainer))
 		, LocalClientEditModel(
 			// Transact ClientContentStorage
 			ConcertClientSharedSlate::CreateTransactionalStreamModel(
 				// Read & write the stream data in ClientContentStorage
-				CreateBaseStreamModel(
-					ClientContentStorage->Stream->MakeReplicationMapGetterAttribute(),
-					MakeShared<FMultiUserStreamExtender>(EndpointId, InDiscoveryContainer)
-					),
+				CreateBaseStreamModel(ClientContentStorage->Stream->MakeReplicationMapGetterAttribute(), StreamExtender),
 				*ClientContentStorage->Stream
 				)
 			)
diff --git a/Engine/Plugins/Developer/Concert/ConcertApp/MultiUserClient/Source/MultiUserClient/Private/Replication/Client/ReplicationClient.h b/Engine/Plugins/Developer/Concert/ConcertApp/MultiUserClient/Source/MultiUserClient/Private/Replication/Client/ReplicationClient.h
--- a/Engine/Plugins/Developer/Concert/ConcertApp/MultiUserClient/Source/MultiUserClient/Private/Replication/Client/ReplicationClient.h
+++ b/Engine/Plugins/Developer/Concert/ConcertApp/MultiUserClient/Source/MultiUserClient/Private/Replication/Client/ReplicationClient.h
@@ -28,6 +28,7 @@ namespace UE::ConcertSharedSlate
 
 namespace UE::MultiUserClient
 {
+	class FMultiUserStreamExtender;
 	class FReplicationDiscoveryContainer;
 	class ISubmissionWorkflow;
 	
@@ -60,6 +61,7 @@ namespace UE::MultiUserClient
 		~FReplicationClient();
 
 		UMultiUserReplicationClientPreset* GetClientContent() const { return ClientContentStorage; }
+		FMultiUserStreamExtender& GetStreamExtender() { return *StreamExtender; }
 		/**
 		 * This is used so the UI can construct the IReplicationStreamEditor.
 		 * @see CreateBaseStreamEditor and FCreateEditorParams.
@@ -125,7 +127,9 @@ namespace UE::MultiUserClient
 		FSubmissionQueue SubmissionQueue;
 		/** Enqueues external change requests into the SubmissionQueue. */
 		FExternalClientChangeRequestHandler ExternalRequestHandler;
-		
+
+		/** Injected into LocalClientEditModel during construction. Adds additional objects and properties when a new object is added to the stream. */
+		TSharedRef<FMultiUserStreamExtender> StreamExtender;
 		/**
 		 * Used to detect changes made to the client's config by the local editor.
 		 * Those changes can later be applied to the client.
diff --git a/Engine/Plugins/Developer/Concert/ConcertApp/MultiUserClient/Source/MultiUserClient/Private/Replication/Stream/Discovery/MultiUserStreamExtender.cpp b/Engine/Plugins/Developer/Concert/ConcertApp/MultiUserClient/Source/MultiUserClient/Private/Replication/Stream/Discovery/MultiUserStreamExtender.cpp
--- a/Engine/Plugins/Developer/Concert/ConcertApp/MultiUserClient/Source/MultiUserClient/Private/Replication/Stream/Discovery/MultiUserStreamExtender.cpp
+++ b/Engine/Plugins/Developer/Concert/ConcertApp/MultiUserClient/Source/MultiUserClient/Private/Replication/Stream/Discovery/MultiUserStreamExtender.cpp
@@ -23,8 +23,11 @@ namespace UE::MultiUserClient
 
 	void FMultiUserStreamExtender::ExtendStream(UObject& ExtendedObject, ConcertSharedSlate::IStreamExtensionContext& Context)
 	{
-		ExtendBySettings.ExtendStream(ExtendedObject, Context);
-		ExtendStreamWithRegisteredDiscoverers(ExtendedObject, Context);
+		if (bShouldExtend)
+		{
+			ExtendBySettings.ExtendStream(ExtendedObject, Context);
+			ExtendStreamWithRegisteredDiscoverers(ExtendedObject, Context);
+		}
 	}
 	
 	void FMultiUserStreamExtender::ExtendStreamWithRegisteredDiscoverers(UObject& ExtendedObject, ConcertSharedSlate::IStreamExtensionContext& Context) const
diff --git a/Engine/Plugins/Developer/Concert/ConcertApp/MultiUserClient/Source/MultiUserClient/Private/Replication/Stream/Discovery/MultiUserStreamExtender.h b/Engine/Plugins/Developer/Concert/ConcertApp/MultiUserClient/Source/MultiUserClient/Private/Replication/Stream/Discovery/MultiUserStreamExtender.h
--- a/Engine/Plugins/Developer/Concert/ConcertApp/MultiUserClient/Source/MultiUserClient/Private/Replication/Stream/Discovery/MultiUserStreamExtender.h
+++ b/Engine/Plugins/Developer/Concert/ConcertApp/MultiUserClient/Source/MultiUserClient/Private/Replication/Stream/Discovery/MultiUserStreamExtender.h
@@ -13,11 +13,13 @@ namespace UE::MultiUserClient
 	
 	/**
 	 * When the user adds an object, this object handles auto selecting properties and adding additional objects from context.
+	 * 
 	 * The following sources exist:
 	 * - Static Settings: user can specify properties & objects in the MU project settings
 	 * - Dynamic API: by using the IMultiUserReplication, external modules can register dynamic rules for auto-discovery using IReplicationDiscoverer.
 	 *
 	 * Every FReplicationClient owns an instance of this and executes it when an object is added for the client by the local editor.
+	 * Effectively, this is
 	 */
 	class FMultiUserStreamExtender : public ConcertSharedSlate::IStreamExtender
 	{
@@ -29,6 +31,9 @@ namespace UE::MultiUserClient
 		 */
 		FMultiUserStreamExtender(const FGuid& InClientId, FReplicationDiscoveryContainer& InRegisteredExtenders);
 
+		/** Sets whether the stream should be extended. */
+		void SetShouldExtend(bool bValue) { bShouldExtend = bValue; }
+
 		//~ Begin IStreamExtender Interface
 		virtual void ExtendStream(UObject& ExtendedObject, ConcertSharedSlate::IStreamExtensionContext& Context) override;
 		//~ End IStreamExtender Interface
@@ -38,9 +43,11 @@ namespace UE::MultiUserClient
 		/** The client ID to which this extender is adds objects. */
 		const FGuid ClientId;
 
+		/** Whether this extender is enabled. */
+		bool bShouldExtend = false;
+		
 		/** Handles properties from the MU settings. */
 		ConcertClientSharedSlate::FStreamExtenderBySettings ExtendBySettings;
-
 		/** Allows generic extenders to extend added objects. These can be added via the IMultiUserReplication interface. */
 		FReplicationDiscoveryContainer& RegisteredExtenders;
 		
diff --git a/Engine/Plugins/Developer/Concert/ConcertApp/MultiUserClient/Source/MultiUserClient/Private/Widgets/ActiveSession/Replication/Client/Multi/SMultiClientView.cpp b/Engine/Plugins/Developer/Concert/ConcertApp/MultiUserClient/Source/MultiUserClient/Private/Widgets/ActiveSession/Replication/Client/Multi/SMultiClientView.cpp
--- a/Engine/Plugins/Developer/Concert/ConcertApp/MultiUserClient/Source/MultiUserClient/Private/Widgets/ActiveSession/Replication/Client/Multi/SMultiClientView.cpp
+++ b/Engine/Plugins/Developer/Concert/ConcertApp/MultiUserClient/Source/MultiUserClient/Private/Widgets/ActiveSession/Replication/Client/Multi/SMultiClientView.cpp
@@ -15,6 +15,7 @@
 #include "Replication/Editor/View/IMultiReplicationStreamEditor.h"
 #include "Replication/Editor/View/IReplicationStreamEditor.h"
 #include "Replication/Editor/Model/ObjectSource/ActorSelectionSourceModel.h"
+#include "Replication/Stream/Discovery/MultiUserStreamExtender.h"
 #include "Widgets/ActiveSession/Replication/Client/Context/ContextMenuUtils.h"
 #include "Widgets/ActiveSession/Replication/Client/Multi/Columns/MultiStreamColumns.h"
 #include "Widgets/ActiveSession/Replication/Client/PropertySelection/SPropertySelectionComboButton.h"
@@ -131,7 +132,9 @@ namespace UE::MultiUserClient
 			.ConsolidatedObjectModel = ConcertClientSharedSlate::CreateTransactionalStreamModel(),
 			.ObjectSource = MakeShared<ConcertClientSharedSlate::FActorSelectionSourceModel>(),
 			.PropertySource = PropertySourceModel,
-			.GetAutoAssignToStreamDelegate = MoveTemp(GetAutoAssignTargetDelegate)
+			.GetAutoAssignToStreamDelegate = MoveTemp(GetAutoAssignTargetDelegate),
+			.OnPreAddSelectedObjectsDelegate = FSelectObjectsFromComboButton::CreateSP(this, &SMultiClientView::OnPreAddObjectsFromComboButton),
+			.OnPostAddSelectedObjectsDelegate = FSelectObjectsFromComboButton::CreateSP(this, &SMultiClientView::OnPostAddObjectsFromComboButton),
 		};
 		FCreateViewerParams ViewerParams
 		{
@@ -144,7 +147,7 @@ namespace UE::MultiUserClient
 				MultiStreamColumns::MuteToggleColumn(MuteManager.GetChangeTracker()),
 				MultiStreamColumns::AssignedClientsColumn(InConcertClient, MultiStreamEditorAttribute, *ObjectHierarchy, ClientManager->GetReassignmentLogic(), *ClientManager)
 			},
-			.ShouldDisplayObjectDelegate = FShouldDisplayObject::CreateSP(this, &SMultiClientView::ShouldDisplayObject)
+			.ShouldDisplayObjectDelegate = FShouldDisplayObject::CreateSP(this, &SMultiClientView::ShouldDisplayObject),
 		};
 		StreamEditor = CreateBaseMultiStreamEditor(MoveTemp(Params), MoveTemp(ViewerParams));
 		check(StreamEditor);
@@ -240,6 +243,19 @@ namespace UE::MultiUserClient
 	{
 		return HideObjectsNotInEditorWorld.ShouldShowObject(Object);
 	}
+
+	void SMultiClientView::OnPreAddObjectsFromComboButton(TArrayView<const ConcertSharedSlate::FSelectableObjectInfo>)
+	{
+		// When the user adds using the combo button, automatically add discover relevant objects and properties
+		ClientManager->GetLocalClient().GetStreamExtender()
+			.SetShouldExtend(true);
+	}
+
+	void SMultiClientView::OnPostAddObjectsFromComboButton(TArrayView<const ConcertSharedSlate::FSelectableObjectInfo>)
+	{
+		ClientManager->GetLocalClient().GetStreamExtender()
+			.SetShouldExtend(false);
+	}
 }
 
 #undef LOCTEXT_NAMESPACE
\ No newline at end of file
diff --git a/Engine/Plugins/Developer/Concert/ConcertApp/MultiUserClient/Source/MultiUserClient/Private/Widgets/ActiveSession/Replication/Client/Multi/SMultiClientView.h b/Engine/Plugins/Developer/Concert/ConcertApp/MultiUserClient/Source/MultiUserClient/Private/Widgets/ActiveSession/Replication/Client/Multi/SMultiClientView.h
--- a/Engine/Plugins/Developer/Concert/ConcertApp/MultiUserClient/Source/MultiUserClient/Private/Widgets/ActiveSession/Replication/Client/Multi/SMultiClientView.h
+++ b/Engine/Plugins/Developer/Concert/ConcertApp/MultiUserClient/Source/MultiUserClient/Private/Widgets/ActiveSession/Replication/Client/Multi/SMultiClientView.h
@@ -2,6 +2,7 @@
 
 #pragma once
 
+#include "Replication/Editor/Model/ObjectSource/IObjectSourceModel.h"
 #include "Replication/Editor/UnrealEditor/HideObjectsNotInWorldLogic.h"
 #include "Widgets/DeclarativeSyntaxSupport.h"
 #include "Widgets/SCompoundWidget.h"
@@ -89,5 +90,8 @@ namespace UE::MultiUserClient
 		void ExtendObjectContextMenu(FMenuBuilder& MenuBuilder, TConstArrayView<TSoftObjectPtr<>> ContextObjects) const;
 		/** Decides whether the object should be displayed: do not show it if it's not in the editor world. */
 		bool ShouldDisplayObject(const FSoftObjectPath& Object) const;
+		
+		void OnPreAddObjectsFromComboButton(TArrayView<const ConcertSharedSlate::FSelectableObjectInfo>);
+		void OnPostAddObjectsFromComboButton(TArrayView<const ConcertSharedSlate::FSelectableObjectInfo>);
 	};
 }
diff --git a/Engine/Plugins/Developer/Concert/ConcertUI/ConcertSharedSlate/Source/ConcertSharedSlate/Private/Replication/Editor/View/MultiEditor/SMultiReplicationStreamEditor.cpp b/Engine/Plugins/Developer/Concert/ConcertUI/ConcertSharedSlate/Source/ConcertSharedSlate/Private/Replication/Editor/View/MultiEditor/SMultiReplicationStreamEditor.cpp
--- a/Engine/Plugins/Developer/Concert/ConcertUI/ConcertSharedSlate/Source/ConcertSharedSlate/Private/Replication/Editor/View/MultiEditor/SMultiReplicationStreamEditor.cpp
+++ b/Engine/Plugins/Developer/Concert/ConcertUI/ConcertSharedSlate/Source/ConcertSharedSlate/Private/Replication/Editor/View/MultiEditor/SMultiReplicationStreamEditor.cpp
@@ -19,6 +19,8 @@ namespace UE::ConcertSharedSlate
 			.DataModel = ConsolidatedModel.ToSharedRef(),
 			.ObjectSource = MoveTemp(EditorParams.ObjectSource),
 			.PropertySource = MoveTemp(EditorParams.PropertySource),
+			.OnPreAddSelectedObjectsDelegate = MoveTemp(EditorParams.OnPreAddSelectedObjectsDelegate),
+			.OnPostAddSelectedObjectsDelegate = MoveTemp(EditorParams.OnPostAddSelectedObjectsDelegate)
 		};
 		EditorView = CreateBaseStreamEditor(MoveTemp(BaseEditorParams), MoveTemp(ViewerParams));
 
diff --git a/Engine/Plugins/Developer/Concert/ConcertUI/ConcertSharedSlate/Source/ConcertSharedSlate/Private/Replication/Editor/View/ObjectEditor/SBaseReplicationStreamEditor.cpp b/Engine/Plugins/Developer/Concert/ConcertUI/ConcertSharedSlate/Source/ConcertSharedSlate/Private/Replication/Editor/View/ObjectEditor/SBaseReplicationStreamEditor.cpp
--- a/Engine/Plugins/Developer/Concert/ConcertUI/ConcertSharedSlate/Source/ConcertSharedSlate/Private/Replication/Editor/View/ObjectEditor/SBaseReplicationStreamEditor.cpp
+++ b/Engine/Plugins/Developer/Concert/ConcertUI/ConcertSharedSlate/Source/ConcertSharedSlate/Private/Replication/Editor/View/ObjectEditor/SBaseReplicationStreamEditor.cpp
@@ -40,6 +40,8 @@ namespace UE::ConcertSharedSlate
 		IsEditingEnabledAttribute = InArgs._IsEditingEnabled;
 		EditingDisabledToolTipTextAttribute = InArgs._EditingDisabledToolTipText;
 		OnExtendObjectsContextMenuDelegate = InArgs._OnExtendObjectsContextMenu;
+		OnPreAddSelectedObjectsDelegate = InArgs._OnPreAddSelectedObjectsDelegate;
+		OnPostAddSelectedObjectsDelegate = InArgs._OnPostAddSelectedObjectsDelegate;
 		
 		ChildSlot
 		[
@@ -200,7 +202,9 @@ namespace UE::ConcertSharedSlate
 			);
 
 		TGuardValue<bool> GuardObjectSelection(bIsAddingFromSelection, true);
+		OnPreAddSelectedObjectsDelegate.ExecuteIfBound(ObjectsToAdd);
 		EditablePropertiesModel->AddObjects(Objects);
+		OnPostAddSelectedObjectsDelegate.ExecuteIfBound(ObjectsToAdd);
 	}
 
 	void SBaseReplicationStreamEditor::OnDeleteObjects(const TArray<TSharedPtr<FReplicatedObjectData>>& ObjectsToDelete) const
diff --git a/Engine/Plugins/Developer/Concert/ConcertUI/ConcertSharedSlate/Source/ConcertSharedSlate/Private/Replication/Editor/View/ObjectEditor/SBaseReplicationStreamEditor.h b/Engine/Plugins/Developer/Concert/ConcertUI/ConcertSharedSlate/Source/ConcertSharedSlate/Private/Replication/Editor/View/ObjectEditor/SBaseReplicationStreamEditor.h
--- a/Engine/Plugins/Developer/Concert/ConcertUI/ConcertSharedSlate/Source/ConcertSharedSlate/Private/Replication/Editor/View/ObjectEditor/SBaseReplicationStreamEditor.h
+++ b/Engine/Plugins/Developer/Concert/ConcertUI/ConcertSharedSlate/Source/ConcertSharedSlate/Private/Replication/Editor/View/ObjectEditor/SBaseReplicationStreamEditor.h
@@ -71,6 +71,11 @@ namespace UE::ConcertSharedSlate
 			/** Called to generate the context menu for objects. This extends the options already generated by this widget. */
 			SLATE_EVENT(FExtendObjectMenu, OnExtendObjectsContextMenu)
 			
+			/** Called just before the passed objects are added to the stream model. Called in response to the user selecting these objects from the combo button. */
+			SLATE_EVENT(FSelectObjectsFromComboButton, OnPreAddSelectedObjectsDelegate)
+			/** Called just after the passed objects have been added to the stream model. Called in response to the user selecting these objects from the combo button. */
+			SLATE_EVENT(FSelectObjectsFromComboButton, OnPostAddSelectedObjectsDelegate)
+			
 			/** Optional. Whether a given object should be displayed. If this returns false on an object, none of its children will be shown either. */
 			SLATE_EVENT(FShouldDisplayObject, ShouldDisplayObject)
 
@@ -121,6 +126,10 @@ namespace UE::ConcertSharedSlate
 		
 		/** Generates additional entries */
 		FExtendObjectMenu OnExtendObjectsContextMenuDelegate;
+		/** Called just before the passed objects are added to the stream model. Called in response to the user selecting these objects from the combo button. */
+		FSelectObjectsFromComboButton OnPreAddSelectedObjectsDelegate;
+		/** Called just after the passed objects have been added to the stream model. Called in response to the user selecting these objects from the combo button. */
+		FSelectObjectsFromComboButton OnPostAddSelectedObjectsDelegate;
 
 		/**
 		 * Set when the combo button is adding an object.
diff --git a/Engine/Plugins/Developer/Concert/ConcertUI/ConcertSharedSlate/Source/ConcertSharedSlate/Private/Replication/ReplicationWidgetFactories.cpp b/Engine/Plugins/Developer/Concert/ConcertUI/ConcertSharedSlate/Source/ConcertSharedSlate/Private/Replication/ReplicationWidgetFactories.cpp
--- a/Engine/Plugins/Developer/Concert/ConcertUI/ConcertSharedSlate/Source/ConcertSharedSlate/Private/Replication/ReplicationWidgetFactories.cpp
+++ b/Engine/Plugins/Developer/Concert/ConcertUI/ConcertSharedSlate/Source/ConcertSharedSlate/Private/Replication/ReplicationWidgetFactories.cpp
@@ -34,6 +34,8 @@ namespace UE::ConcertSharedSlate
 			.ObjectHierarchy(MoveTemp(ViewerParams.ObjectHierarchy))
 			.NameModel(MoveTemp(ViewerParams.NameModel))
 			.OnExtendObjectsContextMenu(MoveTemp(ViewerParams.OnExtendObjectsContextMenu))
+			.OnPreAddSelectedObjectsDelegate(MoveTemp(EditorParams.OnPreAddSelectedObjectsDelegate))
+			.OnPostAddSelectedObjectsDelegate(MoveTemp(EditorParams.OnPostAddSelectedObjectsDelegate))
 			.ShouldDisplayObject(MoveTemp(ViewerParams.ShouldDisplayObjectDelegate))
 			.LeftOfObjectSearchBar() [ MoveTemp(ViewerParams.LeftOfObjectSearchBar.Widget) ]
 			.RightOfObjectSearchBar() [ MoveTemp(ViewerParams.RightOfObjectSearchBar.Widget) ]
diff --git a/Engine/Plugins/Developer/Concert/ConcertUI/ConcertSharedSlate/Source/ConcertSharedSlate/Public/Replication/Editor/Model/IEditableReplicationStreamModel.h b/Engine/Plugins/Developer/Concert/ConcertUI/ConcertSharedSlate/Source/ConcertSharedSlate/Public/Replication/Editor/Model/IEditableReplicationStreamModel.h
--- a/Engine/Plugins/Developer/Concert/ConcertUI/ConcertSharedSlate/Source/ConcertSharedSlate/Public/Replication/Editor/Model/IEditableReplicationStreamModel.h
+++ b/Engine/Plugins/Developer/Concert/ConcertUI/ConcertSharedSlate/Source/ConcertSharedSlate/Public/Replication/Editor/Model/IEditableReplicationStreamModel.h
@@ -25,6 +25,13 @@ namespace UE::ConcertSharedSlate
 		 */
 		ExternalChange
 	};
+
+	enum class EReplicatedAddObjectFlags : uint8
+	{
+		None,
+
+		SkipStreamExtension
+	};
 	
 	/**
 	 * Abstracts the concept of mapping objects to properties. This allows writing.
diff --git a/Engine/Plugins/Developer/Concert/ConcertUI/ConcertSharedSlate/Source/ConcertSharedSlate/Public/Replication/ReplicationWidgetFactories.h b/Engine/Plugins/Developer/Concert/ConcertUI/ConcertSharedSlate/Source/ConcertSharedSlate/Public/Replication/ReplicationWidgetFactories.h
--- a/Engine/Plugins/Developer/Concert/ConcertUI/ConcertSharedSlate/Source/ConcertSharedSlate/Public/Replication/ReplicationWidgetFactories.h
+++ b/Engine/Plugins/Developer/Concert/ConcertUI/ConcertSharedSlate/Source/ConcertSharedSlate/Public/Replication/ReplicationWidgetFactories.h
@@ -13,11 +13,6 @@
 #include "Misc/Attribute.h"
 #include "Templates/SharedPointer.h"
 
-namespace UE::ConcertSharedSlate
-{
-	class IPropertyAssignmentView;
-}
-
 class UObject;
 
 struct FConcertStreamObjectAutoBindingRules;
@@ -30,6 +25,7 @@ namespace UE::ConcertSharedSlate
 	class IMultiReplicationStreamEditor;
 	class IObjectNameModel;
 	class IObjectSelectionSourceModel;
+	class IPropertyAssignmentView;
 	class IPropertyTreeView;
 	class IReplicationStreamModel;
 	class IReplicationStreamEditor;
@@ -119,6 +115,11 @@ namespace UE::ConcertSharedSlate
 		TAttribute<bool> IsEditingEnabled;
 		/** Optional. Whenever IsEditingEnabled returns true, this tooltip is displayed for relevant, disabled UI. */
 		TAttribute<FText> EditingDisabledToolTipText;
+
+		/** Called just before the passed objects are added to the stream model. Called in response to the user selecting these objects from the combo button. */
+		FSelectObjectsFromComboButton OnPreAddSelectedObjectsDelegate;
+		/** Called just after the passed objects have been added to the stream model. Called in response to the user selecting these objects from the combo button. */
+		FSelectObjectsFromComboButton OnPostAddSelectedObjectsDelegate;
 	};
 
 	/**
@@ -175,6 +176,11 @@ namespace UE::ConcertSharedSlate
 
 		/** Optional. If set, the Add button should automatically assign the added object to stream returned by this callback. */
 		FGetAutoAssignTarget GetAutoAssignToStreamDelegate;
+		
+		/** Called just before the passed objects are added to the stream model. Called in response to the user selecting these objects from the combo button. */
+		FSelectObjectsFromComboButton OnPreAddSelectedObjectsDelegate;
+		/** Called just after the passed objects have been added to the stream model. Called in response to the user selecting these objects from the combo button. */
+		FSelectObjectsFromComboButton OnPostAddSelectedObjectsDelegate;
 	};
 
 	/**
diff --git a/Engine/Plugins/Developer/Concert/ConcertUI/ConcertSharedSlate/Source/ConcertSharedSlate/Public/Replication/Utils/ReplicationWidgetDelegates.h b/Engine/Plugins/Developer/Concert/ConcertUI/ConcertSharedSlate/Source/ConcertSharedSlate/Public/Replication/Utils/ReplicationWidgetDelegates.h
--- a/Engine/Plugins/Developer/Concert/ConcertUI/ConcertSharedSlate/Source/ConcertSharedSlate/Public/Replication/Utils/ReplicationWidgetDelegates.h
+++ b/Engine/Plugins/Developer/Concert/ConcertUI/ConcertSharedSlate/Source/ConcertSharedSlate/Public/Replication/Utils/ReplicationWidgetDelegates.h
@@ -20,6 +20,8 @@ namespace UE::ConcertSharedSlate
 	class IPropertySourceProcessor;
 	class IReplicationStreamModel;
 
+	struct FSelectableObjectInfo;
+
 	/** A predicate for determining Left < Right. */
 	DECLARE_DELEGATE_RetVal_TwoParams(bool, FSortPropertyPredicate, const FPropertyData& Left, const FPropertyData& Right);
 
@@ -31,4 +33,7 @@ namespace UE::ConcertSharedSlate
 
 	/** Delegate for deciding whether an object should be displayed. */
 	DECLARE_DELEGATE_RetVal_OneParam(bool, FShouldDisplayObject, const FSoftObjectPath& ObjectPath);
+
+	/** Delegate executed when object options are selected from the Add button to the left of the search bar in the top view of the replication panel. */
+	DECLARE_DELEGATE_OneParam(FSelectObjectsFromComboButton, TConstArrayView<FSelectableObjectInfo>);
 }
\ No newline at end of file
