Fix spamming of transactions when a client changes their replication state on the server.
This prevented users from undoing changes.

diff --git a/Engine/Plugins/Developer/Concert/ConcertApp/MultiUserClient/Source/MultiUserClient/Private/Replication/Misc/PropertySelection/UserPropertySelector.cpp b/Engine/Plugins/Developer/Concert/ConcertApp/MultiUserClient/Source/MultiUserClient/Private/Replication/Misc/PropertySelection/UserPropertySelector.cpp
--- a/Engine/Plugins/Developer/Concert/ConcertApp/MultiUserClient/Source/MultiUserClient/Private/Replication/Misc/PropertySelection/UserPropertySelector.cpp
+++ b/Engine/Plugins/Developer/Concert/ConcertApp/MultiUserClient/Source/MultiUserClient/Private/Replication/Misc/PropertySelection/UserPropertySelector.cpp
@@ -7,24 +7,18 @@
 #include "Replication/Client/ReplicationClientManager.h"
 #include "Replication/Editor/Model/IEditableReplicationStreamModel.h"
 
+#include "ScopedTransaction.h"
 #include "UObject/Package.h"
 #include "UObject/UObjectGlobals.h"
 
+#define LOCTEXT_NAMESPACE "FUserPropertySelector"
+
 namespace UE::MultiUserClient
 {
 	FUserPropertySelector::FUserPropertySelector(FReplicationClientManager& InClientManager)
 		: ClientManager(InClientManager)
 		, PropertySelection(NewObject<UMultiUserReplicationStream>(GetTransientPackage(), NAME_None, RF_Transient | RF_Transactional))
-		, SelectionEditModel(
-			// Transact ClientContentStorage
-			ConcertClientSharedSlate::CreateTransactionalStreamModel(
-				// Read & write the stream data in ClientContentStorage
-				ConcertSharedSlate::CreateBaseStreamModel(
-					PropertySelection->MakeReplicationMapGetterAttribute()
-					),
-				*PropertySelection
-				)
-			)
+		, SelectionEditModel(ConcertSharedSlate::CreateBaseStreamModel(PropertySelection->MakeReplicationMapGetterAttribute()))
 		, PropertyProcessor(MakeShared<FUserPropertySelectionSource>(*SelectionEditModel, InClientManager))
 	{
 		RegisterClient(ClientManager.GetLocalClient());
@@ -45,24 +39,18 @@ namespace UE::MultiUserClient
 
 	void FUserPropertySelector::AddSelectedProperties(UObject* Object, TConstArrayView<FConcertPropertyChain> Properties)
 	{
-		SelectionEditModel->AddObjects({ Object });
-		SelectionEditModel->AddProperties({ Object }, Properties);
-
-		OnPropertySelectionChangedDelegate.Broadcast();
+		const FScopedTransaction Transaction(LOCTEXT("AddSelectedProperties", "Select replicated property"));
+		PropertySelection->Modify();
+		
+		InternalAddSelectedProperties(Object, Properties);
 	}
-
+	
 	void FUserPropertySelector::RemoveSelectedProperties(UObject* Object, TConstArrayView<FConcertPropertyChain> Properties)
 	{
-		SelectionEditModel->RemoveProperties({ Object }, Properties);
-		if (!SelectionEditModel->HasAnyPropertyAssigned(Object))
-		{
-			SelectionEditModel->RemoveObjects({ Object });
-		}
-
-		const TSharedRef<ConcertSharedSlate::IEditableReplicationStreamModel> ClientEditModel = ClientManager.GetLocalClient().GetClientEditModel();
-		ClientEditModel->RemoveProperties({ Object }, Properties);
+		const FScopedTransaction Transaction(LOCTEXT("RemoveSelectedProperties", "Deselect replicated property"));
+		PropertySelection->Modify();
 		
-		OnPropertySelectionChangedDelegate.Broadcast();
+		InternalRemoveSelectedProperties(Object, Properties);
 	}
 
 	bool FUserPropertySelector::IsPropertySelected(const FSoftObjectPath& Object, const FConcertPropertyChain& Property) const
@@ -110,12 +98,13 @@ namespace UE::MultiUserClient
 
 			for (const FConcertPropertyChain& Property : Pair.Value.PropertySelection.ReplicatedProperties)
 			{
-				AddSelectedProperties(Object, { Property });
+				// Do not transact this change: the user did not actively add these properties, so it should not show up in the undo history.
+				InternalAddSelectedProperties(Object, { Property });
 			}
 		}
 	}
 
-	void FUserPropertySelector::OnObjectTransacted(UObject* Object, const FTransactionObjectEvent& TransactionObjectEvent)
+	void FUserPropertySelector::OnObjectTransacted(UObject* Object, const FTransactionObjectEvent&) const
 	{
 		if (Object == PropertySelection)
 		{
@@ -123,4 +112,28 @@ namespace UE::MultiUserClient
 			OnPropertySelectionChangedDelegate.Broadcast();
 		}
 	}
+	
+	void FUserPropertySelector::InternalAddSelectedProperties(UObject* Object, TConstArrayView<FConcertPropertyChain> Properties)
+	{
+		SelectionEditModel->AddObjects({ Object });
+		SelectionEditModel->AddProperties({ Object }, Properties);
+
+		OnPropertySelectionChangedDelegate.Broadcast();
+	}
+	
+	void FUserPropertySelector::InternalRemoveSelectedProperties(UObject* Object, TConstArrayView<FConcertPropertyChain> Properties)
+	{
+		SelectionEditModel->RemoveProperties({ Object }, Properties);
+		if (!SelectionEditModel->HasAnyPropertyAssigned(Object))
+		{
+			SelectionEditModel->RemoveObjects({ Object });
+		}
+
+		const TSharedRef<ConcertSharedSlate::IEditableReplicationStreamModel> ClientEditModel = ClientManager.GetLocalClient().GetClientEditModel();
+		ClientEditModel->RemoveProperties({ Object }, Properties);
+		
+		OnPropertySelectionChangedDelegate.Broadcast();
+	}
 }
+
+#undef LOCTEXT_NAMESPACE
\ No newline at end of file
diff --git a/Engine/Plugins/Developer/Concert/ConcertApp/MultiUserClient/Source/MultiUserClient/Private/Replication/Misc/PropertySelection/UserPropertySelector.h b/Engine/Plugins/Developer/Concert/ConcertApp/MultiUserClient/Source/MultiUserClient/Private/Replication/Misc/PropertySelection/UserPropertySelector.h
--- a/Engine/Plugins/Developer/Concert/ConcertApp/MultiUserClient/Source/MultiUserClient/Private/Replication/Misc/PropertySelection/UserPropertySelector.h
+++ b/Engine/Plugins/Developer/Concert/ConcertApp/MultiUserClient/Source/MultiUserClient/Private/Replication/Misc/PropertySelection/UserPropertySelector.h
@@ -74,7 +74,8 @@ namespace UE::MultiUserClient
 
 		/** Broadcasts when the user property selection changes. */
 		FOnPropertySelectionChanged OnPropertySelectionChangedDelegate;
-		
+
+		/** Called when a remote client joins. */
 		void OnClientAdded(FRemoteReplicationClient& Client) { RegisterClient(Client); }
 		/** Ensures that whenever the client's server state changes, its properties are tracked as user selected. */
 		void RegisterClient(FReplicationClient& Client);
@@ -85,7 +86,10 @@ namespace UE::MultiUserClient
 		void TrackProperties(const FConcertObjectReplicationMap& ReplicationMap);
 
 		/** If PropertySelection is transacted, broadcast OnPropertySelectionChangedDelegate. */
-		void OnObjectTransacted(UObject* Object, const FTransactionObjectEvent& TransactionObjectEvent);
+		void OnObjectTransacted(UObject* Object, const FTransactionObjectEvent&) const;
+		
+		void InternalAddSelectedProperties(UObject* Object, TConstArrayView<FConcertPropertyChain> Properties);
+		void InternalRemoveSelectedProperties(UObject* Object, TConstArrayView<FConcertPropertyChain> Properties);
 	};
 }
 
