Rename UMultiUserReplicationClientPreset to UMultiUserReplicationContent to avoid confusion with presets, which will be added in a later CL.

diff --git a/Engine/Plugins/Developer/Concert/ConcertApp/MultiUserClient/Config/BaseMultiUserClient.ini b/Engine/Plugins/Developer/Concert/ConcertApp/MultiUserClient/Config/BaseMultiUserClient.ini
--- a/Engine/Plugins/Developer/Concert/ConcertApp/MultiUserClient/Config/BaseMultiUserClient.ini
+++ b/Engine/Plugins/Developer/Concert/ConcertApp/MultiUserClient/Config/BaseMultiUserClient.ini
@@ -1,3 +1,7 @@
 [/Script/MultiUserClient.MultiUserReplicationSettings]
 ReplicationEditorSettings=(DefaultPropertySelection=(("/Script/Engine.SceneComponent", (DefaultSelectedProperties=("RelativeLocation","RelativeLocation.X","RelativeLocation.Y","RelativeLocation.Z","RelativeRotation","RelativeRotation.Yaw","RelativeRotation.Pitch","RelativeRotation.Roll"))),("/Script/CinematicCamera.CineCameraComponent", ()),("/Script/Engine.CameraComponent", ())),DefaultAddedSubobjectRules=(SubobjectMatchingRules=(("/Script/Engine.Actor", (IncludeClasses=("/Script/Engine.StaticMeshComponent","/Script/Engine.CameraComponent","/Script/CinematicCamera.CineCameraComponent"))))))
-FrequencyRules=(DefaultObjectFrequencySettings=(ReplicationMode=SpecifiedRate,ReplicationRate=30),FrequencyRuleStack=((ReplicationSettings=(ReplicationRate=60),SubobjectMatchingRules=(SubobjectMatchingRules=(("/Script/Engine.Actor", (IncludeClasses=("/Script/CinematicCamera.CineCameraComponent"))))))))
\ No newline at end of file
+FrequencyRules=(DefaultObjectFrequencySettings=(ReplicationMode=SpecifiedRate,ReplicationRate=30),FrequencyRuleStack=((ReplicationSettings=(ReplicationRate=60),SubobjectMatchingRules=(SubobjectMatchingRules=(("/Script/Engine.Actor", (IncludeClasses=("/Script/CinematicCamera.CineCameraComponent"))))))))
+
+[CoreRedirects]
+; 5.4 -> 5.5
++ClassRedirects=(OldName="/Script/MultiUserClient.MultiUserReplicationClientPreset",NewName="/Script/MultiUserClient.MultiUserReplicationClientContent")
\ No newline at end of file
diff --git a/Engine/Plugins/Developer/Concert/ConcertApp/MultiUserClient/Source/MultiUserClient/Private/Replication/Client/LocalReplicationClient.cpp b/Engine/Plugins/Developer/Concert/ConcertApp/MultiUserClient/Source/MultiUserClient/Private/Replication/Client/LocalReplicationClient.cpp
--- a/Engine/Plugins/Developer/Concert/ConcertApp/MultiUserClient/Source/MultiUserClient/Private/Replication/Client/LocalReplicationClient.cpp
+++ b/Engine/Plugins/Developer/Concert/ConcertApp/MultiUserClient/Source/MultiUserClient/Private/Replication/Client/LocalReplicationClient.cpp
@@ -11,7 +11,7 @@ namespace UE::MultiUserClient
 	FLocalReplicationClient::FLocalReplicationClient(
 		FReplicationDiscoveryContainer& InDiscoveryContainer,
 		FGlobalAuthorityCache& InAuthorityCache,
-		UMultiUserReplicationClientPreset& InSessionContent,
+		UMultiUserReplicationClientContent& InSessionContent,
 		TUniquePtr<IClientStreamSynchronizer> InStreamSynchronizer,
 		TSharedRef<IConcertSyncClient> InClient
 		)
diff --git a/Engine/Plugins/Developer/Concert/ConcertApp/MultiUserClient/Source/MultiUserClient/Private/Replication/Client/LocalReplicationClient.h b/Engine/Plugins/Developer/Concert/ConcertApp/MultiUserClient/Source/MultiUserClient/Private/Replication/Client/LocalReplicationClient.h
--- a/Engine/Plugins/Developer/Concert/ConcertApp/MultiUserClient/Source/MultiUserClient/Private/Replication/Client/LocalReplicationClient.h
+++ b/Engine/Plugins/Developer/Concert/ConcertApp/MultiUserClient/Source/MultiUserClient/Private/Replication/Client/LocalReplicationClient.h
@@ -24,7 +24,7 @@ namespace UE::MultiUserClient
 		FLocalReplicationClient(
 			FReplicationDiscoveryContainer& InDiscoveryContainer,
 			FGlobalAuthorityCache& InAuthorityCache,
-			UMultiUserReplicationClientPreset& InSessionContent,
+			UMultiUserReplicationClientContent& InSessionContent,
 			TUniquePtr<IClientStreamSynchronizer> InStreamSynchronizer,
 			TSharedRef<IConcertSyncClient> InClient
 			);
diff --git a/Engine/Plugins/Developer/Concert/ConcertApp/MultiUserClient/Source/MultiUserClient/Private/Replication/Client/RemoteReplicationClient.cpp b/Engine/Plugins/Developer/Concert/ConcertApp/MultiUserClient/Source/MultiUserClient/Private/Replication/Client/RemoteReplicationClient.cpp
--- a/Engine/Plugins/Developer/Concert/ConcertApp/MultiUserClient/Source/MultiUserClient/Private/Replication/Client/RemoteReplicationClient.cpp
+++ b/Engine/Plugins/Developer/Concert/ConcertApp/MultiUserClient/Source/MultiUserClient/Private/Replication/Client/RemoteReplicationClient.cpp
@@ -3,7 +3,7 @@
 #include "RemoteReplicationClient.h"
 
 #include "IConcertClient.h"
-#include "Assets/MultiUserReplicationClientPreset.h"
+#include "Assets/MultiUserReplicationClientContent.h"
 #include "Replication/Authority/AuthoritySynchronizer_RemoteClient.h"
 #include "Replication/Stream/StreamSynchronizer_RemoteClient.h"
 #include "Replication/Submission/Remote/SubmissionWorkflow_RemoteClient.h"
@@ -15,7 +15,7 @@ namespace UE::MultiUserClient
 		FReplicationDiscoveryContainer& InDiscoveryContainer,
 		TSharedRef<IConcertClient> InClient,
 		FGlobalAuthorityCache& InAuthorityCache,
-		UMultiUserReplicationClientPreset& InSessionContent,
+		UMultiUserReplicationClientContent& InSessionContent,
 		FStreamAndAuthorityQueryService& QueryService
 		)
 		: FReplicationClient(
diff --git a/Engine/Plugins/Developer/Concert/ConcertApp/MultiUserClient/Source/MultiUserClient/Private/Replication/Client/RemoteReplicationClient.h b/Engine/Plugins/Developer/Concert/ConcertApp/MultiUserClient/Source/MultiUserClient/Private/Replication/Client/RemoteReplicationClient.h
--- a/Engine/Plugins/Developer/Concert/ConcertApp/MultiUserClient/Source/MultiUserClient/Private/Replication/Client/RemoteReplicationClient.h
+++ b/Engine/Plugins/Developer/Concert/ConcertApp/MultiUserClient/Source/MultiUserClient/Private/Replication/Client/RemoteReplicationClient.h
@@ -30,7 +30,7 @@ namespace UE::MultiUserClient
 			FReplicationDiscoveryContainer& InDiscoveryContainer UE_LIFETIMEBOUND,
 			TSharedRef<IConcertClient> InClient,
 			FGlobalAuthorityCache& InAuthorityCache UE_LIFETIMEBOUND,
-			UMultiUserReplicationClientPreset& InSessionContent UE_LIFETIMEBOUND,
+			UMultiUserReplicationClientContent& InSessionContent UE_LIFETIMEBOUND,
 			FStreamAndAuthorityQueryService& QueryService UE_LIFETIMEBOUND
 			);
 	};
diff --git a/Engine/Plugins/Developer/Concert/ConcertApp/MultiUserClient/Source/MultiUserClient/Private/Replication/Client/ReplicationClient.cpp b/Engine/Plugins/Developer/Concert/ConcertApp/MultiUserClient/Source/MultiUserClient/Private/Replication/Client/ReplicationClient.cpp
--- a/Engine/Plugins/Developer/Concert/ConcertApp/MultiUserClient/Source/MultiUserClient/Private/Replication/Client/ReplicationClient.cpp
+++ b/Engine/Plugins/Developer/Concert/ConcertApp/MultiUserClient/Source/MultiUserClient/Private/Replication/Client/ReplicationClient.cpp
@@ -2,7 +2,7 @@
 
 #include "ReplicationClient.h"
 
-#include "Assets/MultiUserReplicationClientPreset.h"
+#include "Assets/MultiUserReplicationClientContent.h"
 #include "Replication/Editor/Model/IEditableReplicationStreamModel.h"
 #include "Replication/ReplicationWidgetFactories.h"
 #include "Replication/Stream/Discovery/MultiUserStreamExtender.h"
@@ -19,7 +19,7 @@ namespace UE::MultiUserClient
 		const FGuid& EndpointId,
 		FReplicationDiscoveryContainer& InDiscoveryContainer,
 		FGlobalAuthorityCache& InAuthorityCache,
-		UMultiUserReplicationClientPreset& InSessionContent,
+		UMultiUserReplicationClientContent& InSessionContent,
 		TUniquePtr<IClientStreamSynchronizer> InStreamSynchronizer,
 		TUniquePtr<IClientAuthoritySynchronizer> InAuthoritySynchronizer,
 		TUniquePtr<ISubmissionWorkflow> InSubmissionWorkflow
diff --git a/Engine/Plugins/Developer/Concert/ConcertApp/MultiUserClient/Source/MultiUserClient/Private/Replication/Client/ReplicationClient.h b/Engine/Plugins/Developer/Concert/ConcertApp/MultiUserClient/Source/MultiUserClient/Private/Replication/Client/ReplicationClient.h
--- a/Engine/Plugins/Developer/Concert/ConcertApp/MultiUserClient/Source/MultiUserClient/Private/Replication/Client/ReplicationClient.h
+++ b/Engine/Plugins/Developer/Concert/ConcertApp/MultiUserClient/Source/MultiUserClient/Private/Replication/Client/ReplicationClient.h
@@ -18,7 +18,7 @@
 #include "Templates/UnrealTemplate.h"
 
 class IConcertClient;
-class UMultiUserReplicationClientPreset;
+class UMultiUserReplicationClientContent;
 
 namespace UE::ConcertSharedSlate
 {
@@ -53,14 +53,14 @@ namespace UE::MultiUserClient
 			const FGuid& EndpointId,
 			FReplicationDiscoveryContainer& InDiscoveryContainer UE_LIFETIMEBOUND,
 			FGlobalAuthorityCache& InAuthorityCache UE_LIFETIMEBOUND,
-			UMultiUserReplicationClientPreset& InSessionContent UE_LIFETIMEBOUND,
+			UMultiUserReplicationClientContent& InSessionContent UE_LIFETIMEBOUND,
 			TUniquePtr<IClientStreamSynchronizer> InStreamSynchronizer,
 			TUniquePtr<IClientAuthoritySynchronizer> InAuthoritySynchronizer,
 			TUniquePtr<ISubmissionWorkflow> InSubmissionWorkflow
 			);
 		~FReplicationClient();
 
-		UMultiUserReplicationClientPreset* GetClientContent() const { return ClientContentStorage; }
+		UMultiUserReplicationClientContent* GetClientContent() const { return ClientContentStorage; }
 		FMultiUserStreamExtender& GetStreamExtender() { return *StreamExtender; }
 		/**
 		 * This is used so the UI can construct the IReplicationStreamEditor.
@@ -114,7 +114,7 @@ namespace UE::MultiUserClient
 		const FGuid EndpointId;
 		
 		/** The state of the server is synched up with this object and displayed in the UI. */
-		TObjectPtr<UMultiUserReplicationClientPreset> ClientContentStorage;
+		TObjectPtr<UMultiUserReplicationClientContent> ClientContentStorage;
 		
 		/** Keeps the client's stream state on the server in sync. */
 		TUniquePtr<IClientStreamSynchronizer> StreamSynchronizer;
diff --git a/Engine/Plugins/Developer/Concert/ConcertApp/MultiUserClient/Source/MultiUserClient/Private/Replication/Client/ReplicationClientManager.cpp b/Engine/Plugins/Developer/Concert/ConcertApp/MultiUserClient/Source/MultiUserClient/Private/Replication/Client/ReplicationClientManager.cpp
--- a/Engine/Plugins/Developer/Concert/ConcertApp/MultiUserClient/Source/MultiUserClient/Private/Replication/Client/ReplicationClientManager.cpp
+++ b/Engine/Plugins/Developer/Concert/ConcertApp/MultiUserClient/Source/MultiUserClient/Private/Replication/Client/ReplicationClientManager.cpp
@@ -14,21 +14,28 @@
 
 namespace UE::MultiUserClient
 {
+	namespace Private
+	{
+		static UMultiUserReplicationClientContent* MakeClientContent()
+		{
+			return NewObject<UMultiUserReplicationClientContent>(GetTransientPackage(), NAME_None, RF_Transient | RF_Transactional);
+		}
+	}
+	
 	FReplicationClientManager::FReplicationClientManager(
 		const TSharedRef<IConcertSyncClient>& InClient,
 		const TSharedRef<IConcertClientSession>& InSession,
 		FReplicationDiscoveryContainer& InRegisteredExtenders,
 		FStreamAndAuthorityQueryService& InQueryService
 		)
-		: SessionContent(NewObject<UMultiUserReplicationSessionPreset>(GetTransientPackage(), NAME_None, RF_Transient))
-		, ConcertClient(InClient)
+		: ConcertClient(InClient)
 		, Session(InSession)
 		, RegisteredExtenders(InRegisteredExtenders)
 		, QueryService(InQueryService)
 		, AuthorityCache(*this)
 		, LocalClient([this, InClient]()
 		{
-			UMultiUserReplicationClientPreset* ClientPreset = SessionContent->AddClient();
+			UMultiUserReplicationClientContent* ClientPreset = Private::MakeClientContent();
 			return FLocalReplicationClient(RegisteredExtenders, AuthorityCache, *ClientPreset, MakeUnique<FStreamSynchronizer_LocalClient>(InClient, ClientPreset->Stream->StreamId), InClient);
 		}())
 		, ReassignmentLogic(*this)
@@ -120,7 +127,13 @@ namespace UE::MultiUserClient
 
 	void FReplicationClientManager::AddReferencedObjects(FReferenceCollector& Collector)
 	{
-		Collector.AddReferencedObject(SessionContent);
+		ForEachClient([&Collector](const FReplicationClient& Client)
+		{
+			TObjectPtr<UMultiUserReplicationClientContent> Content = Client.GetClientContent();
+			Collector.AddReferencedObject(Content);
+			ensureMsgf(Content == Client.GetClientContent(), TEXT("Did not expect reference to be obliterated"));
+			return EBreakBehavior::Continue;
+		});
 	}
 
 	void FReplicationClientManager::OnSessionClientChanged(IConcertClientSession&, EConcertClientStatus NewStatus, const FConcertSessionClientInfo& ClientInfo)
@@ -148,7 +161,6 @@ namespace UE::MultiUserClient
 				{
 					const TUniquePtr<FRemoteReplicationClient> Client = MoveTemp(RemoteClients[Index]);
 					OnPreRemoteClientRemovedDelegate.Broadcast(*Client.Get());
-					SessionContent->RemoveClient(*Client->GetClientContent());
 					RemoteClients.RemoveAtSwap(Index);
 				}
 				// We want to broadcast after the client has been fully cleaned up
@@ -169,7 +181,7 @@ namespace UE::MultiUserClient
 			RegisteredExtenders,
 			ConcertClient->GetConcertClient(),
 			AuthorityCache,
-			*SessionContent->AddClient(),
+			*Private::MakeClientContent(),
 			QueryService
 			);
 		FRemoteReplicationClient& RemoteClient = *RemoteClientPtr;
diff --git a/Engine/Plugins/Developer/Concert/ConcertApp/MultiUserClient/Source/MultiUserClient/Private/Replication/Client/ReplicationClientManager.h b/Engine/Plugins/Developer/Concert/ConcertApp/MultiUserClient/Source/MultiUserClient/Private/Replication/Client/ReplicationClientManager.h
--- a/Engine/Plugins/Developer/Concert/ConcertApp/MultiUserClient/Source/MultiUserClient/Private/Replication/Client/ReplicationClientManager.h
+++ b/Engine/Plugins/Developer/Concert/ConcertApp/MultiUserClient/Source/MultiUserClient/Private/Replication/Client/ReplicationClientManager.h
@@ -109,9 +109,6 @@ namespace UE::MultiUserClient
 		//~ End FGCObject Interface
 		
 	private:
-		
-		/** The state of the server is synchronized up with this object and displayed in the UI. */
-		TObjectPtr<UMultiUserReplicationSessionPreset> SessionContent;
 
 		/** The local Concert client */
 		TSharedRef<IConcertSyncClient> ConcertClient;
diff --git a/Engine/Plugins/Developer/Concert/ConcertApp/MultiUserClient/Source/MultiUserClient/Private/Replication/Stream/StreamSynchronizer_RemoteClient.cpp b/Engine/Plugins/Developer/Concert/ConcertApp/MultiUserClient/Source/MultiUserClient/Private/Replication/Stream/StreamSynchronizer_RemoteClient.cpp
--- a/Engine/Plugins/Developer/Concert/ConcertApp/MultiUserClient/Source/MultiUserClient/Private/Replication/Stream/StreamSynchronizer_RemoteClient.cpp
+++ b/Engine/Plugins/Developer/Concert/ConcertApp/MultiUserClient/Source/MultiUserClient/Private/Replication/Stream/StreamSynchronizer_RemoteClient.cpp
@@ -2,7 +2,7 @@
 
 #include "StreamSynchronizer_RemoteClient.h"
 
-#include "Assets/MultiUserReplicationClientPreset.h"
+#include "Assets/MultiUserReplicationClientContent.h"
 #include "Replication/Misc/Query/StreamAndAuthorityQueryService.h"
 
 namespace UE::MultiUserClient
@@ -26,7 +26,7 @@ namespace UE::MultiUserClient
 
 	FGuid FStreamSynchronizer_RemoteClient::GetStreamId() const
 	{
-		return UMultiUserReplicationClientPreset::MultiUserStreamID;
+		return UMultiUserReplicationClientContent::MultiUserStreamID;
 	}
 
 	void FStreamSynchronizer_RemoteClient::HandleStreamQuery(const TArray<FConcertBaseStreamInfo>& Streams)
@@ -35,7 +35,7 @@ namespace UE::MultiUserClient
 		// That handles the (unlikely) case in which some external logic had added streams to the same client which we must differentiate
 		const FConcertBaseStreamInfo* MultiUserDescription = Streams.FindByPredicate([](const FConcertBaseStreamInfo& Description)
 		{
-			return Description.Identifier == UMultiUserReplicationClientPreset::MultiUserStreamID;
+			return Description.Identifier == UMultiUserReplicationClientContent::MultiUserStreamID;
 		});
 		if (MultiUserDescription
 			// Avoid unnecessary Broadcast()s
diff --git a/Engine/Plugins/Developer/Concert/ConcertApp/MultiUserClient/Source/MultiUserReplicationEditor/Internal/Assets/MultiUserReplicationClientContent.h b/Engine/Plugins/Developer/Concert/ConcertApp/MultiUserClient/Source/MultiUserReplicationEditor/Internal/Assets/MultiUserReplicationClientContent.h
new file mode 100644
--- /dev/null
+++ b/Engine/Plugins/Developer/Concert/ConcertApp/MultiUserClient/Source/MultiUserReplicationEditor/Internal/Assets/MultiUserReplicationClientContent.h
@@ -0,0 +1,28 @@
+// Copyright Epic Games, Inc. All Rights Reserved.
+
+#pragma once
+
+#include "UObject/Object.h"
+#include "MultiUserReplicationStream.h"
+#include "Replication/Data/ReplicationStream.h"
+#include "MultiUserReplicationClientContent.generated.h"
+
+/** The transactable content of a FReplicationClient. */
+UCLASS()
+class MULTIUSERREPLICATIONEDITOR_API UMultiUserReplicationClientContent : public UObject
+{
+	GENERATED_BODY()
+public:
+
+	/** The stream ID that is used by all Multi-User streams. */
+	static constexpr FGuid MultiUserStreamID { 0xAAAAAAAA, 0xBBBBBBBB, 0xCCCCCCCC, 0xDDDDDDDD };
+
+	/** The stream this client is managing */
+	UPROPERTY(Instanced)
+	TObjectPtr<UMultiUserReplicationStream> Stream;
+
+	UMultiUserReplicationClientContent();
+
+	/** Generates a description that can be sent to the MU server. */
+	FConcertReplicationStream GenerateDescription() const;
+};
diff --git a/Engine/Plugins/Developer/Concert/ConcertApp/MultiUserClient/Source/MultiUserReplicationEditor/Internal/Assets/MultiUserReplicationClientPreset.h b/Engine/Plugins/Developer/Concert/ConcertApp/MultiUserClient/Source/MultiUserReplicationEditor/Internal/Assets/MultiUserReplicationClientPreset.h
deleted file mode 100644
diff --git a/Engine/Plugins/Developer/Concert/ConcertApp/MultiUserClient/Source/MultiUserReplicationEditor/Internal/Assets/MultiUserReplicationSessionPreset.h b/Engine/Plugins/Developer/Concert/ConcertApp/MultiUserClient/Source/MultiUserReplicationEditor/Internal/Assets/MultiUserReplicationSessionPreset.h
--- a/Engine/Plugins/Developer/Concert/ConcertApp/MultiUserClient/Source/MultiUserReplicationEditor/Internal/Assets/MultiUserReplicationSessionPreset.h
+++ b/Engine/Plugins/Developer/Concert/ConcertApp/MultiUserClient/Source/MultiUserReplicationEditor/Internal/Assets/MultiUserReplicationSessionPreset.h
@@ -2,7 +2,7 @@
 
 #pragma once
 
-#include "MultiUserReplicationClientPreset.h"
+#include "MultiUserReplicationClientContent.h"
 #include "UObject/Object.h"
 #include "MultiUserReplicationSessionPreset.generated.h"
 
@@ -12,23 +12,16 @@
 UCLASS()
 class MULTIUSERREPLICATIONEDITOR_API UMultiUserReplicationSessionPreset : public UObject
 {
-	GENERATED_UCLASS_BODY()
+	GENERATED_BODY()
 public:
 
-	UMultiUserReplicationClientPreset* AddClient();
-	
-	void RemoveClient(UMultiUserReplicationClientPreset& Client);
-	void ClearClients();
+	UMultiUserReplicationClientContent* AddClient();
+	void RemoveClient(UMultiUserReplicationClientContent& Client);
 
-	const TArray<TObjectPtr<UMultiUserReplicationClientPreset>>& GetClientPresets() const { return ClientPresets; }
-	UMultiUserReplicationClientPreset* GetUnassignedClient() const { return UnassignedClient; }
+	const TArray<TObjectPtr<UMultiUserReplicationClientContent>>& GetClientPresets() const { return ClientPresets; }
 
 private:
 	
 	UPROPERTY(Instanced)
-	TArray<TObjectPtr<UMultiUserReplicationClientPreset>> ClientPresets;
-
-	/** Special preset that is assigned to an invalid client - objects without ownership go here. This is not in ClientPresets. */
-	UPROPERTY(Instanced)
-	TObjectPtr<UMultiUserReplicationClientPreset> UnassignedClient;
+	TArray<TObjectPtr<UMultiUserReplicationClientContent>> ClientPresets;
 };
diff --git a/Engine/Plugins/Developer/Concert/ConcertApp/MultiUserClient/Source/MultiUserReplicationEditor/Internal/Assets/MultiUserReplicationStream.h b/Engine/Plugins/Developer/Concert/ConcertApp/MultiUserClient/Source/MultiUserReplicationEditor/Internal/Assets/MultiUserReplicationStream.h
--- a/Engine/Plugins/Developer/Concert/ConcertApp/MultiUserClient/Source/MultiUserReplicationEditor/Internal/Assets/MultiUserReplicationStream.h
+++ b/Engine/Plugins/Developer/Concert/ConcertApp/MultiUserClient/Source/MultiUserReplicationEditor/Internal/Assets/MultiUserReplicationStream.h
@@ -23,9 +23,11 @@ public:
 	FConcertObjectReplicationMap ReplicationMap;
 
 	UMultiUserReplicationStream();
-	
+
+	/** Util that generates the description of this stream for network requests. */
 	FConcertReplicationStream GenerateDescription() const;
-	
+
+	/** Util that returns ReplicationMap. */
 	TAttribute<FConcertObjectReplicationMap*> MakeReplicationMapGetterAttribute()
 	{
 		return TAttribute<FConcertObjectReplicationMap*>::CreateLambda([WeakThis = TWeakObjectPtr<UMultiUserReplicationStream>(this)]()
diff --git a/Engine/Plugins/Developer/Concert/ConcertApp/MultiUserClient/Source/MultiUserReplicationEditor/Private/Assets/MultiUserReplicationClientContent.cpp b/Engine/Plugins/Developer/Concert/ConcertApp/MultiUserClient/Source/MultiUserReplicationEditor/Private/Assets/MultiUserReplicationClientContent.cpp
new file mode 100644
--- /dev/null
+++ b/Engine/Plugins/Developer/Concert/ConcertApp/MultiUserClient/Source/MultiUserReplicationEditor/Private/Assets/MultiUserReplicationClientContent.cpp
@@ -0,0 +1,18 @@
+// Copyright Epic Games, Inc. All Rights Reserved.
+
+#include "Assets/MultiUserReplicationClientContent.h"
+
+#include "Assets/MultiUserReplicationStream.h"
+#include "Replication/Data/ReplicationStream.h"
+
+UMultiUserReplicationClientContent::UMultiUserReplicationClientContent()
+{
+	Stream = CreateDefaultSubobject<UMultiUserReplicationStream>(TEXT("ReplicationList"));
+	Stream->StreamId = MultiUserStreamID;
+	Stream->SetFlags(RF_Transactional);
+}
+
+FConcertReplicationStream UMultiUserReplicationClientContent::GenerateDescription() const
+{
+	return Stream->GenerateDescription();
+}
diff --git a/Engine/Plugins/Developer/Concert/ConcertApp/MultiUserClient/Source/MultiUserReplicationEditor/Private/Assets/MultiUserReplicationClientPreset.cpp b/Engine/Plugins/Developer/Concert/ConcertApp/MultiUserClient/Source/MultiUserReplicationEditor/Private/Assets/MultiUserReplicationClientPreset.cpp
deleted file mode 100644
diff --git a/Engine/Plugins/Developer/Concert/ConcertApp/MultiUserClient/Source/MultiUserReplicationEditor/Private/Assets/MultiUserReplicationSessionPreset.cpp b/Engine/Plugins/Developer/Concert/ConcertApp/MultiUserClient/Source/MultiUserReplicationEditor/Private/Assets/MultiUserReplicationSessionPreset.cpp
--- a/Engine/Plugins/Developer/Concert/ConcertApp/MultiUserClient/Source/MultiUserReplicationEditor/Private/Assets/MultiUserReplicationSessionPreset.cpp
+++ b/Engine/Plugins/Developer/Concert/ConcertApp/MultiUserClient/Source/MultiUserReplicationEditor/Private/Assets/MultiUserReplicationSessionPreset.cpp
@@ -2,25 +2,14 @@
 
 #include "Assets/MultiUserReplicationSessionPreset.h"
 
-UMultiUserReplicationSessionPreset::UMultiUserReplicationSessionPreset(const FObjectInitializer& ObjectInitializer)
+UMultiUserReplicationClientContent* UMultiUserReplicationSessionPreset::AddClient()
 {
-	UnassignedClient = ObjectInitializer.CreateDefaultSubobject<UMultiUserReplicationClientPreset>(this, TEXT("UnassignedClient"));
-}
-
-UMultiUserReplicationClientPreset* UMultiUserReplicationSessionPreset::AddClient()
-{
-	UMultiUserReplicationClientPreset* Result = NewObject<UMultiUserReplicationClientPreset>(this);
+	UMultiUserReplicationClientContent* Result = NewObject<UMultiUserReplicationClientContent>(this);
 	ClientPresets.Add(Result);
 	return Result;
 }
 
-void UMultiUserReplicationSessionPreset::RemoveClient(UMultiUserReplicationClientPreset& Client)
+void UMultiUserReplicationSessionPreset::RemoveClient(UMultiUserReplicationClientContent& Client)
 {
 	ClientPresets.RemoveSingle(&Client);
 }
-
-void UMultiUserReplicationSessionPreset::ClearClients()
-{
-	ClientPresets.Empty();
-	UnassignedClient->ClearClient();
-}
diff --git a/Engine/Plugins/Developer/Concert/ConcertApp/MultiUserClient/Source/MultiUserReplicationEditor/Private/StreamEditor/ReplicationSessionPresetEditorToolkit.cpp b/Engine/Plugins/Developer/Concert/ConcertApp/MultiUserClient/Source/MultiUserReplicationEditor/Private/StreamEditor/ReplicationSessionPresetEditorToolkit.cpp
--- a/Engine/Plugins/Developer/Concert/ConcertApp/MultiUserClient/Source/MultiUserReplicationEditor/Private/StreamEditor/ReplicationSessionPresetEditorToolkit.cpp
+++ b/Engine/Plugins/Developer/Concert/ConcertApp/MultiUserClient/Source/MultiUserReplicationEditor/Private/StreamEditor/ReplicationSessionPresetEditorToolkit.cpp
@@ -2,7 +2,7 @@
 
 #include "ReplicationSessionPresetEditorToolkit.h"
 
-#include "Assets/MultiUserReplicationClientPreset.h"
+#include "Assets/MultiUserReplicationClientContent.h"
 #include "Assets/MultiUserReplicationSessionPreset.h"
 #include "Replication/Editor/Model/PropertySource/SelectPropertyFromUClassModel.h"
 #include "Replication/Editor/View/IReplicationStreamEditor.h"
