Upon joining replication, clients now restore what they had registered when they left the session.

diff --git a/Engine/Plugins/Developer/Concert/ConcertApp/MultiUserClient/Source/MultiUserClient/Private/Replication/MultiUserReplicationManager.cpp b/Engine/Plugins/Developer/Concert/ConcertApp/MultiUserClient/Source/MultiUserClient/Private/Replication/MultiUserReplicationManager.cpp
--- a/Engine/Plugins/Developer/Concert/ConcertApp/MultiUserClient/Source/MultiUserClient/Private/Replication/MultiUserReplicationManager.cpp
+++ b/Engine/Plugins/Developer/Concert/ConcertApp/MultiUserClient/Source/MultiUserClient/Private/Replication/MultiUserReplicationManager.cpp
@@ -8,7 +8,12 @@
 #include "Replication/IConcertClientReplicationManager.h"
 
 #include "Containers/Ticker.h"
+#include "Framework/Application/SlateApplication.h"
+#include "Framework/Notifications/NotificationManager.h"
 #include "UObject/Package.h"
+#include "Widgets/Notifications/SNotificationList.h"
+
+#define LOCTEXT_NAMESPACE "FMultiUserReplicationManager"
 
 namespace UE::MultiUserClient
 {
@@ -36,7 +41,6 @@ namespace UE::MultiUserClient
 		}
 
 		ConnectionState = EMultiUserReplicationConnectionState::Connecting;
-		// For now we join without any initial data - this will likely change in the future (5.5+)
 		Manager->JoinReplicationSession({})
 			.Next([WeakThis = AsWeak()](ConcertSyncClient::Replication::FJoinReplicatedSessionResult&& JoinSessionResult)
 			{
@@ -89,6 +93,9 @@ namespace UE::MultiUserClient
 			ConnectedState.Emplace(Client, DiscoveryContainer);
 			SetupClientConnectionEvents();
 			SetConnectionStateAndBroadcast(EMultiUserReplicationConnectionState::Connected);
+
+			// For convenience, the client should attempt to restore the content when they last left.
+			RestoreContentFromLastTime();
 		}
 		else
 		{
@@ -102,6 +109,37 @@ namespace UE::MultiUserClient
 		OnReplicationConnectionStateChangedDelegate.Broadcast(ConnectionState);
 	}
 
+	void FMultiUserReplicationManager::RestoreContentFromLastTime()
+	{
+		Client->GetReplicationManager()->RestoreContent(
+			{
+				.Flags = EConcertReplicationRestoreContentFlags::StreamsAndAuthority | EConcertReplicationRestoreContentFlags::ValidateUniqueClient
+			}
+		)
+		.Next([ClientInfo = Client->GetConcertClient()->GetClientInfo()](FConcertReplication_RestoreContent_Response&& Response)
+		{
+			UE_LOG(LogConcert, Log, TEXT("Content restoration completed with result '%s'"), *ConcertSyncCore::LexToString(Response.ErrorCode));
+			const bool bIsTimeout = !IsInGameThread() || Response.ErrorCode == EConcertReplicationRestoreErrorCode::Timeout;
+			if (Response.IsSuccess() || bIsTimeout || !FSlateApplication::IsInitialized())
+			{
+				return;
+			}
+
+			FSlateNotificationManager& NotificationManager = FSlateNotificationManager::Get();
+			FNotificationInfo NotificationInfo(LOCTEXT("RestoreFailed.Main", "Replication Content Restore"));
+			NotificationInfo.SubText = FText::Format(
+				LOCTEXT("RestoreFailed.SubTextFmt", "Display name {0} and device name {1} already taken by another client in session."),
+				FText::FromString(ClientInfo.DisplayName),
+				FText::FromString(ClientInfo.DeviceName)
+				);
+			NotificationInfo.bFireAndForget = true;
+			NotificationInfo.bUseSuccessFailIcons = true;
+			NotificationInfo.ExpireDuration = 4.f;
+			NotificationManager.AddNotification(NotificationInfo)
+				->SetCompletionState(SNotificationItem::CS_Fail);
+		});
+	}
+
 	void FMultiUserReplicationManager::SetupClientConnectionEvents()
 	{
 		FReplicationClientManager& ClientManager = ConnectedState->ClientManager;
@@ -117,7 +155,6 @@ namespace UE::MultiUserClient
 
 	void FMultiUserReplicationManager::OnClientAuthorityServerStateChanged(const FGuid EndpointId) const
 	{
-		
 		UE_LOG(LogConcert, Verbose, TEXT("Client %s authority changed"), *EndpointId.ToString());
 		OnAuthorityServerStateChangedDelegate.Broadcast(EndpointId);
 	}
@@ -206,3 +243,5 @@ namespace UE::MultiUserClient
 		, UserNotifier(ClientManager, MuteManager)
 	{}
 }
+
+#undef LOCTEXT_NAMESPACE
\ No newline at end of file
diff --git a/Engine/Plugins/Developer/Concert/ConcertApp/MultiUserClient/Source/MultiUserClient/Private/Replication/MultiUserReplicationManager.h b/Engine/Plugins/Developer/Concert/ConcertApp/MultiUserClient/Source/MultiUserClient/Private/Replication/MultiUserReplicationManager.h
--- a/Engine/Plugins/Developer/Concert/ConcertApp/MultiUserClient/Source/MultiUserClient/Private/Replication/MultiUserReplicationManager.h
+++ b/Engine/Plugins/Developer/Concert/ConcertApp/MultiUserClient/Source/MultiUserClient/Private/Replication/MultiUserReplicationManager.h
@@ -157,6 +157,8 @@ namespace UE::MultiUserClient
 		void HandleReplicationSessionJoined(const ConcertSyncClient::Replication::FJoinReplicatedSessionResult& JoinSessionResult);
 		/** Sets the current connection state and triggers OnReplicationConnectionStateChangedDelegate. */
 		void SetConnectionStateAndBroadcast(EMultiUserReplicationConnectionState NewState);
+		/** Restores the content the client had when they last left the session. */
+		void RestoreContentFromLastTime();
 
 		/** Sets up delegates for implementing the broadcasting of OnStreamServerStateChangedDelegate and OnAuthorityServerStateChangedDelegate. */
 		void SetupClientConnectionEvents();
diff --git a/Engine/Plugins/Developer/Concert/ConcertSync/ConcertSyncCore/Source/ConcertSyncCore/Public/ConcertSyncSessionTypes.h b/Engine/Plugins/Developer/Concert/ConcertSync/ConcertSyncCore/Source/ConcertSyncCore/Public/ConcertSyncSessionTypes.h
--- a/Engine/Plugins/Developer/Concert/ConcertSync/ConcertSyncCore/Source/ConcertSyncCore/Public/ConcertSyncSessionTypes.h
+++ b/Engine/Plugins/Developer/Concert/ConcertSync/ConcertSyncCore/Source/ConcertSyncCore/Public/ConcertSyncSessionTypes.h
@@ -164,7 +164,7 @@ struct FConcertSyncPackageEventData
 	FConcertPackageDataStream PackageDataStream;
 };
 
-UENUM()
+UENUM(Flags)
 enum class EConcertSyncActivityFlags : uint8
 {
 	None = 0,
diff --git a/Engine/Plugins/Developer/Concert/ConcertSync/ConcertSyncCore/Source/ConcertSyncCore/Public/Replication/Messages/ChangeStream.h b/Engine/Plugins/Developer/Concert/ConcertSync/ConcertSyncCore/Source/ConcertSyncCore/Public/Replication/Messages/ChangeStream.h
--- a/Engine/Plugins/Developer/Concert/ConcertSync/ConcertSyncCore/Source/ConcertSyncCore/Public/Replication/Messages/ChangeStream.h
+++ b/Engine/Plugins/Developer/Concert/ConcertSync/ConcertSyncCore/Source/ConcertSyncCore/Public/Replication/Messages/ChangeStream.h
@@ -63,7 +63,7 @@ struct FConcertReplication_ChangeStream_PutObject
 	CONCERTSYNCCORE_API TOptional<FConcertReplicatedObjectInfo> MakeObjectInfoIfValid() const;
 };
 
-UENUM()
+UENUM(Flags)
 enum class EConcertReplicationChangeFrequencyFlags : uint8
 {
 	None = 0,
diff --git a/Engine/Plugins/Developer/Concert/ConcertSync/ConcertSyncCore/Source/ConcertSyncCore/Public/Replication/Messages/RestoreContent.h b/Engine/Plugins/Developer/Concert/ConcertSync/ConcertSyncCore/Source/ConcertSyncCore/Public/Replication/Messages/RestoreContent.h
--- a/Engine/Plugins/Developer/Concert/ConcertSync/ConcertSyncCore/Source/ConcertSyncCore/Public/Replication/Messages/RestoreContent.h
+++ b/Engine/Plugins/Developer/Concert/ConcertSync/ConcertSyncCore/Source/ConcertSyncCore/Public/Replication/Messages/RestoreContent.h
@@ -8,7 +8,7 @@
 #include "Replication/Messages/SyncControl.h"
 #include "RestoreContent.generated.h"
 
-UENUM()
+UENUM(Flags) 
 enum class EConcertReplicationRestoreContentFlags : uint8
 {
 	None = 0,
@@ -71,7 +71,7 @@ struct FConcertReplication_RestoreContent_Request
 	/** If EConcertReplicationRestoreContentFlags::RestoreAuthority is set, describes how to deal with authority conflicts. */
 	UPROPERTY()
 	EConcertReplicationAuthorityRestoreMode AuthorityRestorationMode = EConcertReplicationAuthorityRestoreMode::ExcludeAlreadyOwnedObjectPropertiesFromStream;
-
+	
 	/**
 	 * The ID of an activity that contains a client's replication state.
 	 * The activity must be a FConcertSyncReplicationActivity with whose EventData.ActivityType == EConcertSyncReplicationActivityType::LeaveReplication.
@@ -107,6 +107,24 @@ enum class EConcertReplicationRestoreErrorCode : uint8
 	AuthorityConflict
 };
 
+namespace UE::ConcertSyncCore
+{
+	inline FString LexToString(EConcertReplicationRestoreErrorCode ErrorCode)
+	{
+		switch (ErrorCode)
+		{
+		case EConcertReplicationRestoreErrorCode::Timeout: return TEXT("Timeout");
+		case EConcertReplicationRestoreErrorCode::Success: return TEXT("Success");
+		case EConcertReplicationRestoreErrorCode::Invalid: return TEXT("Invalid");
+		case EConcertReplicationRestoreErrorCode::NotSupported: return TEXT("NotSupported");
+		case EConcertReplicationRestoreErrorCode::NoSuchActivity: return TEXT("NotSupported");
+		case EConcertReplicationRestoreErrorCode::NameConflict: return TEXT("NameConflict");
+		case EConcertReplicationRestoreErrorCode::AuthorityConflict: return TEXT("AuthorityConflict");
+			default: checkNoEntry(); return TEXT("Unknown");
+		}
+	}
+}
+
 USTRUCT()
 struct FConcertReplication_RestoreContent_Response
 {
