Save and load presets for Multi User replication

diff --git a/Engine/Plugins/Developer/Concert/ConcertApp/MultiUserClient/Source/MultiUserClient/Private/Replication/MultiUserReplicationManager.cpp b/Engine/Plugins/Developer/Concert/ConcertApp/MultiUserClient/Source/MultiUserClient/Private/Replication/MultiUserReplicationManager.cpp
--- a/Engine/Plugins/Developer/Concert/ConcertApp/MultiUserClient/Source/MultiUserClient/Private/Replication/MultiUserReplicationManager.cpp
+++ b/Engine/Plugins/Developer/Concert/ConcertApp/MultiUserClient/Source/MultiUserClient/Private/Replication/MultiUserReplicationManager.cpp
@@ -237,6 +237,7 @@ namespace UE::MultiUserClient
 		, QueryService(*InClient)
 		, ClientManager(InClient, InClient->GetConcertClient()->GetCurrentSession().ToSharedRef(), InDiscoveryContainer, QueryService.GetStreamAndAuthorityQueryService())
 		, MuteManager(*InClient, QueryService.GetMuteStateQueryService(), ClientManager.GetAuthorityCache())
+		, PresetManager(*InClient, ClientManager)
 		, PropertySelector(ClientManager)
 		, ChangeLevelHandler(ClientManager.GetLocalClient().GetClientEditModel().Get())
 		, PreventReplicatedPropertyTransaction(*InClient, ClientManager, MuteManager)
diff --git a/Engine/Plugins/Developer/Concert/ConcertApp/MultiUserClient/Source/MultiUserClient/Private/Replication/MultiUserReplicationManager.h b/Engine/Plugins/Developer/Concert/ConcertApp/MultiUserClient/Source/MultiUserClient/Private/Replication/MultiUserReplicationManager.h
--- a/Engine/Plugins/Developer/Concert/ConcertApp/MultiUserClient/Source/MultiUserClient/Private/Replication/MultiUserReplicationManager.h
+++ b/Engine/Plugins/Developer/Concert/ConcertApp/MultiUserClient/Source/MultiUserClient/Private/Replication/MultiUserReplicationManager.h
@@ -6,14 +6,15 @@
 #include "IConcertSession.h"
 #include "Misc/ChangeLevelHandler.h"
 #include "Misc/Notification/ReplicationUserNotifier.h"
+#include "Misc/PreventReplicatedPropertyTransaction.h"
+#include "Misc/PropertySelection/UserPropertySelector.h"
 #include "Misc/Query/RegularQueryService.h"
 #include "Muting/MuteStateManager.h"
+#include "Preset/PresetManager.h"
 #include "Replication/IMultiUserReplication.h"
 #include "Replication/Stream/Discovery/ReplicationDiscoveryContainer.h"
 
 #include "Misc/Optional.h"
-#include "Misc/PreventReplicatedPropertyTransaction.h"
-#include "Misc/PropertySelection/UserPropertySelector.h"
 #include "Templates/SharedPointer.h"
 #include "Templates/UnrealTemplate.h"
 
@@ -68,6 +69,10 @@ namespace UE::MultiUserClient
 		FMuteStateManager* GetMuteManager() { return ConnectedState ? &ConnectedState->MuteManager : nullptr; }
 		const FMuteStateManager* GetMuteManager() const { return ConnectedState ? &ConnectedState->MuteManager : nullptr; }
 
+		/** @note You're not supposed to keep any reference to the PresetManager since it can become invalid depending on connection state. */
+		FPresetManager* GetPresetManager() { return ConnectedState ? &ConnectedState->PresetManager : nullptr; }
+		const FPresetManager* GetPresetManager() const { return ConnectedState ? &ConnectedState->PresetManager : nullptr; }
+
 		/** @note You're not supposed to keep any reference to the PropertySelector since it can become invalid depending on connection state. */
 		FUserPropertySelector* GetUserPropertySelector() { return ConnectedState ? &ConnectedState->PropertySelector : nullptr; }
 		const FUserPropertySelector* GetUserPropertySelector() const { return ConnectedState ? &ConnectedState->PropertySelector : nullptr; }
@@ -119,6 +124,8 @@ namespace UE::MultiUserClient
 			FReplicationClientManager ClientManager;
 			/** Interacts with the mute global server mute system. */
 			FMuteStateManager MuteManager;
+			/** Saves and loads presets for the session. Accessed by UI. */
+			FPresetManager PresetManager;
 
 			/**
 			 * Manages the properties the user is iterating on in the replication session.
diff --git a/Engine/Plugins/Developer/Concert/ConcertApp/MultiUserClient/Source/MultiUserClient/Private/Replication/Preset/Data/ReplacedWithPresetEvent.h b/Engine/Plugins/Developer/Concert/ConcertApp/MultiUserClient/Source/MultiUserClient/Private/Replication/Preset/Data/ReplacedWithPresetEvent.h
new file mode 100644
--- /dev/null
+++ b/Engine/Plugins/Developer/Concert/ConcertApp/MultiUserClient/Source/MultiUserClient/Private/Replication/Preset/Data/ReplacedWithPresetEvent.h
@@ -0,0 +1,27 @@
+﻿// Copyright Epic Games, Inc. All Rights Reserved.
+
+#pragma once
+
+#include "ConcertMessageData.h"
+
+#include "Containers/Array.h"
+#include "HAL/Platform.h"
+
+namespace UE::MultiUserClient
+{
+	/** Describes why changing a client failed. */
+	struct FFailToReplaceWithPresetInfo
+	{
+		FConcertSessionClientInfo ClientInfo;
+	};
+
+	/** Describes the outcome of replacing a session's content with a preset. */
+	struct FReplacedWithPresetEvent
+	{
+		/** The clients for which an error occured. */
+		TArray<FFailToReplaceWithPresetInfo> FailedClients;
+
+		bool IsSuccess() const { return FailedClients.IsEmpty(); }
+	};
+}
+
diff --git a/Engine/Plugins/Developer/Concert/ConcertApp/MultiUserClient/Source/MultiUserClient/Private/Replication/Preset/Operations/ReplaceSessionContentOperation.cpp b/Engine/Plugins/Developer/Concert/ConcertApp/MultiUserClient/Source/MultiUserClient/Private/Replication/Preset/Operations/ReplaceSessionContentOperation.cpp
new file mode 100644
--- /dev/null
+++ b/Engine/Plugins/Developer/Concert/ConcertApp/MultiUserClient/Source/MultiUserClient/Private/Replication/Preset/Operations/ReplaceSessionContentOperation.cpp
@@ -0,0 +1,3 @@
+﻿// Copyright Epic Games, Inc. All Rights Reserved.
+
+#include "ReplaceSessionContentOperation.h"
diff --git a/Engine/Plugins/Developer/Concert/ConcertApp/MultiUserClient/Source/MultiUserClient/Private/Replication/Preset/Operations/ReplaceSessionContentOperation.h b/Engine/Plugins/Developer/Concert/ConcertApp/MultiUserClient/Source/MultiUserClient/Private/Replication/Preset/Operations/ReplaceSessionContentOperation.h
new file mode 100644
--- /dev/null
+++ b/Engine/Plugins/Developer/Concert/ConcertApp/MultiUserClient/Source/MultiUserClient/Private/Replication/Preset/Operations/ReplaceSessionContentOperation.h
@@ -0,0 +1,15 @@
+﻿// Copyright Epic Games, Inc. All Rights Reserved.
+
+#pragma once
+
+#include "Templates/UnrealTemplate.h"
+
+namespace UE::MultiUserClient
+{
+	class FReplaceSessionContentOperation : public FNoncopyable
+	{
+	public:
+	
+	};
+}
+
diff --git a/Engine/Plugins/Developer/Concert/ConcertApp/MultiUserClient/Source/MultiUserClient/Private/Replication/Preset/PresetManager.cpp b/Engine/Plugins/Developer/Concert/ConcertApp/MultiUserClient/Source/MultiUserClient/Private/Replication/Preset/PresetManager.cpp
new file mode 100644
--- /dev/null
+++ b/Engine/Plugins/Developer/Concert/ConcertApp/MultiUserClient/Source/MultiUserClient/Private/Replication/Preset/PresetManager.cpp
@@ -0,0 +1,232 @@
+﻿// Copyright Epic Games, Inc. All Rights Reserved.
+
+#include "PresetManager.h"
+
+#include "ConcertLogGlobal.h"
+#include "Assets/MultiUserReplicationSessionPreset.h"
+
+#include "FileHelpers.h"
+#include "IConcertSyncClient.h"
+#include "Replication/Client/ReplicationClientManager.h"
+#include "Widgets/ActiveSession/Replication/Client/ClientUtils.h"
+
+namespace UE::MultiUserClient
+{
+	namespace Private
+	{
+		static void RemoveEmptyObjectsFromRequest(FConcertReplicationStream& Stream)
+		{
+			TMap<FSoftObjectPath, FConcertReplicatedObjectInfo>& ReplicationMap = Stream.BaseDescription.ReplicationMap.ReplicatedObjects;
+			for (auto It = ReplicationMap.CreateIterator(); It; ++It)
+			{
+				const bool bIsEmpty = It->Value.PropertySelection.ReplicatedProperties.IsEmpty();
+				if (bIsEmpty)
+				{
+					It.RemoveCurrent();
+				}
+			}
+		}
+		
+		static FConcertReplication_PutState_Request BuildRequest(
+			const UMultiUserReplicationSessionPreset& Preset,
+			const IConcertClientSession& Session,
+			EApplyPresetFlags Flags
+			)
+		{
+			FConcertReplication_PutState_Request Request;
+			
+			const auto AddClient = [&Preset, &Request, Flags](const FConcertSessionClientInfo& ClientSessionInfo)
+			{
+				const UMultiUserReplicationClientContent* ClientSessionContent = Preset.GetClientContent(ClientSessionInfo.ClientInfo);
+				if (!ClientSessionContent)
+				{
+					if (EnumHasAnyFlags(Flags, EApplyPresetFlags::ClearUnreferencedClients))
+					{
+						Request.NewStreams.Add(ClientSessionInfo.ClientEndpointId, {});
+					}
+					return;
+				}
+
+				const UMultiUserReplicationStream* SavedStream = ClientSessionContent->Stream;
+				if (SavedStream->ReplicationMap.IsEmpty())
+				{
+					// This causes the client's content to be cleared.
+					Request.NewStreams.Add(ClientSessionInfo.ClientEndpointId, {});
+					return;
+				}
+				
+				const FGuid& StreamId = SavedStream->StreamId;
+				FConcertReplicationStream Stream { { .Identifier = StreamId, .ReplicationMap = SavedStream->ReplicationMap } };
+				// Empty objects will be rejected by the server
+				RemoveEmptyObjectsFromRequest(Stream); 
+				Request.NewStreams.Add(ClientSessionInfo.ClientEndpointId, { TArray{ Stream } });
+				
+				// MU automatically requests authority when it adds an object.
+				// We'll assume that that authority was granted - if it actually was not, our request may fail due to overlapping authority.
+				TArray<FConcertObjectInStreamID>& OwnedObjects = Request.NewAuthorityState.Add(ClientSessionInfo.ClientEndpointId).Objects;
+				Algo::Transform(Stream.BaseDescription.ReplicationMap.ReplicatedObjects, OwnedObjects, [&StreamId](const TPair<FSoftObjectPath, FConcertReplicatedObjectInfo>& Pair)
+				{
+					return FConcertObjectInStreamID{ StreamId, Pair.Key };
+				});
+				
+				// TODO UE-219639: Load frequency settings
+			};
+			AddClient({ Session.GetSessionClientEndpointId(), Session.GetLocalClientInfo() });
+			for (const FConcertSessionClientInfo& ClientSessionInfo : Session.GetSessionClients())
+			{
+				AddClient(ClientSessionInfo);
+			}
+			
+			// TODO UE-219427: When implementing mute state, you must skip muting objects that may no be referenced due to some client not being present.
+			return Request;
+		}
+
+		static EReplaceSessionContentErrorCode ExtractErrorCode(const FConcertReplication_PutState_Response& Response)
+		{
+			switch (Response.ResponseCode)
+			{
+			case EConcertReplicationPutStateResponseCode::Success: return EReplaceSessionContentErrorCode::Success;
+			case EConcertReplicationPutStateResponseCode::Timeout: return EReplaceSessionContentErrorCode::Timeout;
+			case EConcertReplicationPutStateResponseCode::FeatureDisabled: return EReplaceSessionContentErrorCode::FeatureDisabled;
+							
+			case EConcertReplicationPutStateResponseCode::ClientUnknown: [[fallthrough]];
+			case EConcertReplicationPutStateResponseCode::StreamError: [[fallthrough]];
+			case EConcertReplicationPutStateResponseCode::AuthorityConflict: [[fallthrough]];
+			case EConcertReplicationPutStateResponseCode::MuteError: [[fallthrough]];
+			default: return EReplaceSessionContentErrorCode::Rejected; 
+			}
+		}
+
+		static void RemoveEmptyObjectsFromLocalClient(ConcertSharedSlate::IEditableReplicationStreamModel& EditModel)
+		{
+			TArray<FSoftObjectPath> EmptyObjects;
+			EditModel.ForEachReplicatedObject([&EditModel, &EmptyObjects](const FSoftObjectPath& Object)
+			{
+				if (EditModel.GetNumProperties(Object) == 0)
+				{
+					EmptyObjects.Add(Object);
+				}
+				return EBreakBehavior::Continue;
+			});
+			EditModel.RemoveObjects(EmptyObjects);
+		}
+	}
+	
+	FPresetManager::FPresetManager(const IConcertSyncClient& SyncClient, const FReplicationClientManager& ClientManager)
+		: SyncClient(SyncClient)
+		, ClientManager(ClientManager)
+	{
+		IConcertClientReplicationManager& ReplicationManager = *SyncClient.GetReplicationManager();
+		ReplicationManager.OnPostRemoteEditApplied().AddRaw(this, &FPresetManager::OnPostRemoteEditApplied);
+	}
+
+	FPresetManager::~FPresetManager()
+	{
+		IConcertClientReplicationManager& ReplicationManager = *SyncClient.GetReplicationManager();
+		ReplicationManager.OnPostRemoteEditApplied().RemoveAll(this);
+	
+		if (InProgressSessionReplacementOp)
+		{
+			InProgressSessionReplacementOp->EmplaceValue(FReplaceSessionContentResult{ EReplaceSessionContentErrorCode::Cancelled });
+			InProgressSessionReplacementOp.Reset();
+		}
+	}
+
+	TFuture<FReplaceSessionContentResult> FPresetManager::ReplaceSessionContentWithPreset(const UMultiUserReplicationSessionPreset& Preset, EApplyPresetFlags Flags)
+	{
+		if (!ensure(!IsPresetChangeInProgress()))
+		{
+			return MakeFulfilledPromise<FReplaceSessionContentResult>(EReplaceSessionContentErrorCode::InProgress).GetFuture();
+		}
+		
+		const TSharedPtr<IConcertClientSession> Session = SyncClient.GetConcertClient()->GetCurrentSession();
+		if (!ensure(Session))
+		{
+			return MakeFulfilledPromise<FReplaceSessionContentResult>(EReplaceSessionContentErrorCode::Timeout).GetFuture();
+		}
+		
+		InProgressSessionReplacementOp = MakeShared<TPromise<FReplaceSessionContentResult>>();
+		SyncClient.GetReplicationManager()
+			->PutClientState(Private::BuildRequest(Preset, *Session, Flags))
+			.Next(
+				[this, WeakPromise = TWeakPtr<TPromise<FReplaceSessionContentResult>>(InProgressSessionReplacementOp)]
+				(FConcertReplication_PutState_Response&& Response)
+			{
+				// If WeakPromise is stale, the request completes after our owning FPresetManager has been destroyed.
+				// In that case, it is not safe to access this.
+				if (const TSharedPtr<TPromise<FReplaceSessionContentResult>> PromisePin = WeakPromise.Pin())
+				{
+					// Destroy before emplacing promise because its future may trigger another ReplaceSessionContentWithPreset call (unlikely); PromisePin will keep it alive for now.
+					InProgressSessionReplacementOp.Reset();
+					PromisePin->EmplaceValue(Private::ExtractErrorCode(Response));
+
+					// The client may have added an object via the Add button but not assigned any properties.
+					// Those empty objects exist locally only and were never submitted to the server.
+					// Remove those because it is a client expectation that the final list only contains the objects that were in the preset.
+					ConcertSharedSlate::IEditableReplicationStreamModel& EditModel = *ClientManager.GetLocalClient().GetClientEditModel();
+					Private::RemoveEmptyObjectsFromLocalClient(EditModel);
+				}
+			});
+		
+		return InProgressSessionReplacementOp->GetFuture();
+	}
+
+	void FPresetManager::ExportToPresetAndSaveAs()
+	{
+		UMultiUserReplicationSessionPreset* Preset = ExportToPreset();
+		
+		TArray<UObject*> SavedAssets;
+		FEditorFileUtils::SaveAssetsAs({ Preset }, SavedAssets);
+	}
+
+	UMultiUserReplicationSessionPreset* FPresetManager::ExportToPreset() const
+	{
+		const TSharedPtr<IConcertClientSession> Session = SyncClient.GetConcertClient()->GetCurrentSession();
+		checkf(Session, TEXT("FPresetManager is only supposed to exist while in a session"));
+		
+		UMultiUserReplicationSessionPreset* Preset = NewObject<UMultiUserReplicationSessionPreset>(
+			GetTransientPackage(),
+			MakeUniqueObjectName(GetTransientPackage(), UMultiUserReplicationSessionPreset::StaticClass(), TEXT("ReplicationPreset")),
+			// Mark as transient so FEditorFileUtils::SaveAssetsAs creates a new package for the object.
+			RF_Transient
+			);
+		ClientManager.ForEachClient([this, &Session, &Preset](const FReplicationClient& Client)
+		{
+			FConcertClientInfo ClientInfo;
+			const bool bGotClientInfo = ClientUtils::GetClientDisplayInfo(*Session, Client.GetEndpointId(), ClientInfo);
+			if (!ensure(bGotClientInfo))
+			{
+				return EBreakBehavior::Continue;
+			}
+			
+			UMultiUserReplicationClientContent* ClientContent_ToCopy = Client.GetClientContent();
+			UMultiUserReplicationClientContent* ClientContent_InPreset = Preset->AddClientIfUnique(ClientInfo);
+			if (!ClientContent_InPreset)
+			{
+				UE_LOG(LogConcert, Warning,
+					TEXT("There are multiple clients with display name %s and device name %s in the session. Only the 1st encountered will be saved into the preset. Did you perhaps launch 2 editors on the same machine (if so you can use -CONCERTDISPLAYNAME)?"),
+					*ClientInfo.DisplayName,
+					*ClientInfo.DeviceName
+					);
+				return EBreakBehavior::Continue;
+			}
+
+			ClientContent_InPreset->Stream->Copy(*ClientContent_ToCopy->Stream);
+			return EBreakBehavior::Continue;
+		});
+		
+		return Preset;
+	}
+
+	void FPresetManager::OnPostRemoteEditApplied(const ConcertSyncClient::Replication::FRemoteEditEvent& Event) const
+	{
+		if (Event.Reason == EConcertReplicationChangeClientReason::PutRequest)
+		{
+			// The client may have added an object via the Add button but not assigned any properties.
+			// Those empty objects exist locally only and were never submitted to the server.
+			// Remove those because it is a client expectation that the final list only contains the objects that were in the preset.
+			ConcertSharedSlate::IEditableReplicationStreamModel& EditModel = *ClientManager.GetLocalClient().GetClientEditModel();
+			Private::RemoveEmptyObjectsFromLocalClient(EditModel);
+		}
+	}
+}
diff --git a/Engine/Plugins/Developer/Concert/ConcertApp/MultiUserClient/Source/MultiUserClient/Private/Replication/Preset/PresetManager.h b/Engine/Plugins/Developer/Concert/ConcertApp/MultiUserClient/Source/MultiUserClient/Private/Replication/Preset/PresetManager.h
new file mode 100644
--- /dev/null
+++ b/Engine/Plugins/Developer/Concert/ConcertApp/MultiUserClient/Source/MultiUserClient/Private/Replication/Preset/PresetManager.h
@@ -0,0 +1,96 @@
+﻿// Copyright Epic Games, Inc. All Rights Reserved.
+
+#pragma once
+
+#include "Async/Future.h"
+#include "HAL/Platform.h"
+#include "Templates/UnrealTemplate.h"
+
+class IConcertSyncClient;
+class UMultiUserReplicationSessionPreset;
+enum class EMultiUserClientPresetLoadMode : uint8;
+
+namespace UE::ConcertSyncClient::Replication { struct FRemoteEditEvent; }
+
+namespace UE::MultiUserClient
+{
+	class FReplicationClientManager;
+
+	enum class EReplaceSessionContentErrorCode : uint8
+	{
+		/** Request completed successfully. */
+		Success,
+		
+		/** Request cancelled because FPresetManager was destroyed - probably because the use left the session during the request. */
+		Cancelled,
+		/** Another locally initiated operation is already in progress. */
+		InProgress,
+		
+		/** Request timed out */
+		Timeout,
+		/** The feature is not enabled (i.e. EConcertSyncSessionFlags::ShouldEnableRemoteEditing or EConcertSyncSessionFlags::ShouldAllowGlobalMuting were not set on the server).  */
+		FeatureDisabled,
+		/** Server rejected the change because it was not valid */
+		Rejected
+	};
+	
+	/** Result of FPresetManager::ReplaceSessionContentWithPreset. */
+	struct FReplaceSessionContentResult
+	{
+		EReplaceSessionContentErrorCode ErrorCode;
+
+		FReplaceSessionContentResult(EReplaceSessionContentErrorCode ErrorCode = EReplaceSessionContentErrorCode::Success)
+			: ErrorCode(ErrorCode)
+		{}
+
+		bool IsSuccess() const { return ErrorCode == EReplaceSessionContentErrorCode::Success; }
+	};
+
+	enum class EApplyPresetFlags : uint8
+	{
+		None,
+		/** If set clients that were not in the session when the preset was created will get their content reset, too. */
+		ClearUnreferencedClients = 1 << 0
+	};
+	ENUM_CLASS_FLAGS(EApplyPresetFlags);
+	
+	/**
+	 * Implements all logic for managing presets in the MU session: saving and loading presets.
+	 * The UI directly interfaces with this class.
+	 */
+	class FPresetManager : public FNoncopyable
+	{
+	public:
+		
+		FPresetManager(const IConcertSyncClient& SyncClient UE_LIFETIMEBOUND, const FReplicationClientManager& ClientManager UE_LIFETIMEBOUND);
+		~FPresetManager();
+
+		/** @return Whether any preset is currently being applied. */
+		bool IsPresetChangeInProgress() const { return InProgressSessionReplacementOp.IsValid(); }
+		
+		/** Applies Preset to all clients in the session. */
+		TFuture<FReplaceSessionContentResult> ReplaceSessionContentWithPreset(const UMultiUserReplicationSessionPreset& Preset, EApplyPresetFlags Flags = EApplyPresetFlags::None);
+
+		/**
+		 * Exports the current session content to a preset, asks the user where to save it, then saves it.
+		 * @return A future that finishes when saving has completed.
+		 */
+		void ExportToPresetAndSaveAs();
+
+	private:
+
+		/** Used to get display information of clients in the session. */
+		const IConcertSyncClient& SyncClient;
+		/** Used to get the clients' replication content. */
+		const FReplicationClientManager& ClientManager;
+
+		/** Non-null for as long as the ReplaceSessionContentWithPreset network request takes. */
+		TSharedPtr<TPromise<FReplaceSessionContentResult>> InProgressSessionReplacementOp;
+
+		/** Exports the current session content to a preset. */
+		UMultiUserReplicationSessionPreset* ExportToPreset() const;
+
+		/** Called when the local client receives a remote edit. */
+		void OnPostRemoteEditApplied(const ConcertSyncClient::Replication::FRemoteEditEvent&) const;
+	};
+}
diff --git a/Engine/Plugins/Developer/Concert/ConcertApp/MultiUserClient/Source/MultiUserClient/Private/Widgets/ActiveSession/Replication/Client/Multi/SMultiClientView.cpp b/Engine/Plugins/Developer/Concert/ConcertApp/MultiUserClient/Source/MultiUserClient/Private/Widgets/ActiveSession/Replication/Client/Multi/SMultiClientView.cpp
--- a/Engine/Plugins/Developer/Concert/ConcertApp/MultiUserClient/Source/MultiUserClient/Private/Widgets/ActiveSession/Replication/Client/Multi/SMultiClientView.cpp
+++ b/Engine/Plugins/Developer/Concert/ConcertApp/MultiUserClient/Source/MultiUserClient/Private/Widgets/ActiveSession/Replication/Client/Multi/SMultiClientView.cpp
@@ -19,6 +19,7 @@
 #include "Widgets/ActiveSession/Replication/Client/Context/ContextMenuUtils.h"
 #include "Widgets/ActiveSession/Replication/Client/Multi/Columns/MultiStreamColumns.h"
 #include "Widgets/ActiveSession/Replication/Client/PropertySelection/SPropertySelectionComboButton.h"
+#include "Widgets/ActiveSession/Replication/Client/SPresetComboButton.h"
 #include "Widgets/ActiveSession/Replication/Client/SReplicationStatus.h"
 
 #include "Widgets/Layout/SBox.h"
@@ -149,6 +150,7 @@ namespace UE::MultiUserClient
 			},
 			.ShouldDisplayObjectDelegate = FShouldDisplayObject::CreateSP(this, &SMultiClientView::ShouldDisplayObject),
 		};
+		ViewerParams.RightOfObjectSearchBar.Widget = SNew(SPresetComboButton, *InMultiUserReplicationManager.GetPresetManager());
 		StreamEditor = CreateBaseMultiStreamEditor(MoveTemp(Params), MoveTemp(ViewerParams));
 		check(StreamEditor);
 		
diff --git a/Engine/Plugins/Developer/Concert/ConcertApp/MultiUserClient/Source/MultiUserClient/Private/Widgets/ActiveSession/Replication/Client/SPresetComboButton.cpp b/Engine/Plugins/Developer/Concert/ConcertApp/MultiUserClient/Source/MultiUserClient/Private/Widgets/ActiveSession/Replication/Client/SPresetComboButton.cpp
new file mode 100644
--- /dev/null
+++ b/Engine/Plugins/Developer/Concert/ConcertApp/MultiUserClient/Source/MultiUserClient/Private/Widgets/ActiveSession/Replication/Client/SPresetComboButton.cpp
@@ -0,0 +1,170 @@
+﻿// Copyright Epic Games, Inc. All Rights Reserved.
+
+#include "SPresetComboButton.h"
+
+#include "Assets/MultiUserReplicationSessionPreset.h"
+#include "Replication/Preset/PresetManager.h"
+
+#include "Containers/Ticker.h"
+#include "ContentBrowserModule.h"
+#include "IContentBrowserSingleton.h"
+#include "Framework/MultiBox/MultiBoxBuilder.h"
+#include "Framework/Notifications/NotificationManager.h"
+#include "SSimpleComboButton.h"
+#include "Styling/AppStyle.h"
+#include "Widgets/Layout/SBorder.h"
+#include "Widgets/Notifications/SNotificationList.h"
+
+#define LOCTEXT_NAMESPACE "SPresetComboButton"
+
+namespace UE::MultiUserClient
+{
+	namespace Private
+	{
+		static FText MakeTitle(const FReplaceSessionContentResult& Result, const FText& PresetText)
+		{
+			const FText Format = Result.IsSuccess()
+				? LOCTEXT("ApplyPreset.Title.SuccessFmt", "Applied {0} preset")
+				: LOCTEXT("ApplyPreset.Title.FailFmt", "Failed to apply {0} preset");
+			return FText::Format(Format, PresetText);
+		}
+
+		static FText MakeSubText(const FReplaceSessionContentResult& Result)
+		{
+			switch (Result.ErrorCode)
+			{
+			case EReplaceSessionContentErrorCode::Success: return FText::GetEmpty();
+			case EReplaceSessionContentErrorCode::Cancelled: return LOCTEXT("ApplyPreset.SubText.Success", "Disconnected from session.");
+			case EReplaceSessionContentErrorCode::InProgress: return LOCTEXT("ApplyPreset.SubText.InProgress", "Another operation is already in progress.");
+			case EReplaceSessionContentErrorCode::Timeout: return LOCTEXT("ApplyPreset.SubText.Timeout", "Request timed out.");
+			case EReplaceSessionContentErrorCode::FeatureDisabled: return LOCTEXT("ApplyPreset.SubText.FeatureDisabled", "This session does not support presets.");
+			case EReplaceSessionContentErrorCode::Rejected: return LOCTEXT("ApplyPreset.SubText.Rejected", "Rejected by server.");
+			default: checkNoEntry(); return FText::GetEmpty();
+			}
+		}
+	}
+	
+	void SPresetComboButton::Construct(const FArguments& InArgs, FPresetManager& InPresetManager)
+	{
+		PresetManager = &InPresetManager;
+		ChildSlot
+		[
+			SNew(SSimpleComboButton)
+			.Icon(FAppStyle::Get().GetBrush("AssetEditor.SaveAsset"))
+			.Text(LOCTEXT("Presets", "Presets"))
+			.OnGetMenuContent(this, &SPresetComboButton::CreateMenuContent)
+			.HasDownArrow(true)
+		];
+	}
+
+	TSharedRef<SWidget> SPresetComboButton::CreateMenuContent()
+	{
+		FMenuBuilder MenuBuilder(true, nullptr);
+
+		CreateSaveMenuContent(MenuBuilder);
+
+		MenuBuilder.BeginSection(NAME_None, LOCTEXT("Section.Import", "Import preset"));
+		CreateLoadMenuContent(MenuBuilder);
+		MenuBuilder.EndSection();
+
+		return MenuBuilder.MakeWidget();
+	}
+
+	void SPresetComboButton::CreateSaveMenuContent(FMenuBuilder& MenuBuilder)
+	{
+		MenuBuilder.AddMenuEntry(
+			LOCTEXT("SavePresetAs.Label", "Save Preset as..."),
+			LOCTEXT("SavePresetAs.ToolTip", "Saves what each client was replicating as a preset."),
+			FSlateIcon(FAppStyle::Get().GetStyleSetName(), "AssetEditor.SaveAssetAs"),
+			FUIAction(FExecuteAction::CreateSP(this, &SPresetComboButton::SavePresetAs))	
+		);
+	}
+
+	void SPresetComboButton::CreateLoadMenuContent(FMenuBuilder& MenuBuilder)
+	{
+		MenuBuilder.AddMenuEntry(
+			LOCTEXT("ImportPreset.ClearOtherClients.Label", "Clear clients not in preset"),
+			LOCTEXT("ImportPreset.ClearOtherClients.ToolTip", "If checked, clients that were not in the session when the preset was created will get their content reset, too."),
+			FSlateIcon(),
+			FUIAction(
+				FExecuteAction::CreateLambda([this](){ Options.bResetAllOtherClients = !Options.bResetAllOtherClients; }),
+				FCanExecuteAction::CreateLambda([](){ return true; }),
+				FIsActionChecked::CreateLambda([this]{ return Options.bResetAllOtherClients; })
+				),
+			NAME_None,
+			EUserInterfaceActionType::ToggleButton
+			);
+			
+		FContentBrowserModule& ContentBrowserModule = FModuleManager::Get().LoadModuleChecked<FContentBrowserModule>(TEXT("ContentBrowser"));
+
+		FAssetPickerConfig AssetPickerConfig;
+		AssetPickerConfig.Filter.ClassPaths.Add(UMultiUserReplicationSessionPreset::StaticClass()->GetClassPathName());
+		AssetPickerConfig.bAllowNullSelection = false;
+		AssetPickerConfig.Filter.bRecursiveClasses = false;
+		AssetPickerConfig.OnAssetSelected.BindSP(this, &SPresetComboButton::LoadPreset);
+		AssetPickerConfig.InitialAssetViewType = EAssetViewType::List;
+		AssetPickerConfig.bAllowDragging = false;
+		
+		TSharedRef<SWidget> PresetPicker = SNew(SBox)
+			.HeightOverride(400)
+			.WidthOverride(300)
+			[
+				SNew(SBorder)
+				.BorderImage(FAppStyle::GetBrush("Menu.Background"))
+				[
+					ContentBrowserModule.Get().CreateAssetPicker(AssetPickerConfig)
+				]
+			];
+		MenuBuilder.AddWidget(PresetPicker, FText(), true, false);
+	}
+
+	void SPresetComboButton::SavePresetAs()
+	{
+		PresetManager->ExportToPresetAndSaveAs();
+	}
+
+	void SPresetComboButton::LoadPreset(const FAssetData& AssetData)
+	{
+		FSlateNotificationManager& NotificationManager = FSlateNotificationManager::Get();
+		if (UMultiUserReplicationSessionPreset* Preset = Cast<UMultiUserReplicationSessionPreset>(AssetData.GetAsset()))
+		{
+			FText PresetText = FText::FromName(Preset->GetFName());
+			FNotificationInfo Info(FText::Format(LOCTEXT("ApplyPreset.Title.InProgressFmt", "Applying preset {0}"), PresetText));
+			Info.ExpireDuration = 4.f;
+			TSharedPtr<SNotificationItem> Notification = NotificationManager.AddNotification(Info);
+			Notification->SetCompletionState(SNotificationItem::CS_Pending);
+
+			PresetManager->ReplaceSessionContentWithPreset(*Preset, BuildFlags())
+				.Next([PresetText = MoveTemp(PresetText), Notification = MoveTemp(Notification)](const FReplaceSessionContentResult& Result) mutable
+				{
+					ExecuteOnGameThread(TEXT("SPresetComboButton"), [PresetText = MoveTemp(PresetText), Notification = MoveTemp(Notification), Result]()
+					{
+						Notification->SetText(Private::MakeTitle(Result, PresetText));
+						Notification->SetSubText(Private::MakeSubText(Result));
+						Notification->SetCompletionState(Result.IsSuccess() ? SNotificationItem::CS_Success : SNotificationItem::CS_Fail);
+					});
+				});
+		}
+		else
+		{
+			FNotificationInfo Info(LOCTEXT("FailedToLoad", "Failed to load preset"));
+			Info.ExpireDuration = 4.f;
+			NotificationManager.AddNotification(Info)
+				->SetCompletionState(SNotificationItem::CS_Fail);
+		}
+	}
+
+	EApplyPresetFlags SPresetComboButton::BuildFlags() const
+	{
+		EApplyPresetFlags Flags = EApplyPresetFlags::None;
+
+		if (Options.bResetAllOtherClients)
+		{
+			Flags |= EApplyPresetFlags::ClearUnreferencedClients;
+		}
+
+		return Flags;
+	}
+}
+
+#undef LOCTEXT_NAMESPACE
\ No newline at end of file
diff --git a/Engine/Plugins/Developer/Concert/ConcertApp/MultiUserClient/Source/MultiUserClient/Private/Widgets/ActiveSession/Replication/Client/SPresetComboButton.h b/Engine/Plugins/Developer/Concert/ConcertApp/MultiUserClient/Source/MultiUserClient/Private/Widgets/ActiveSession/Replication/Client/SPresetComboButton.h
new file mode 100644
--- /dev/null
+++ b/Engine/Plugins/Developer/Concert/ConcertApp/MultiUserClient/Source/MultiUserClient/Private/Widgets/ActiveSession/Replication/Client/SPresetComboButton.h
@@ -0,0 +1,57 @@
+﻿// Copyright Epic Games, Inc. All Rights Reserved.
+
+#pragma once
+
+#include "Widgets/DeclarativeSyntaxSupport.h"
+#include "Widgets/SCompoundWidget.h"
+
+namespace UE::MultiUserClient
+{
+	enum class EApplyPresetFlags : uint8;
+}
+
+class FMenuBuilder;
+
+namespace UE::MultiUserClient
+{
+	class FPresetManager;
+
+	/**
+	 * A combo button displayed to the right of the object search bar.
+	 * Its menu allows the user to save and load replication presets.
+	 */
+	class SPresetComboButton : public SCompoundWidget
+	{
+	public:
+
+		SLATE_BEGIN_ARGS(SPresetComboButton) {}
+			
+		SLATE_END_ARGS()
+
+		void Construct(const FArguments& InArgs, FPresetManager& InPresetManager UE_LIFETIMEBOUND);
+
+	private:
+
+		/** Used to save & load preset. */
+		FPresetManager* PresetManager = nullptr;
+
+		struct FPresetOptions
+		{
+			/** Clients not mentioned by the preset will get their content wiped. */
+			bool bResetAllOtherClients = true;
+		} Options;
+
+		/** Creates the Save & Load options for the menu. */
+		TSharedRef<SWidget> CreateMenuContent();
+		void CreateSaveMenuContent(FMenuBuilder& MenuBuilder);
+		void CreateLoadMenuContent(FMenuBuilder& MenuBuilder);
+
+		/** Saves the current session content. */
+		void SavePresetAs();
+		/** Handles loading the preset. */
+		void LoadPreset(const FAssetData& AssetData);
+
+		EApplyPresetFlags BuildFlags() const;
+	};
+}
+
diff --git a/Engine/Plugins/Developer/Concert/ConcertApp/MultiUserClient/Source/MultiUserReplicationEditor/Internal/Assets/MultiUserReplicationClientContent.h b/Engine/Plugins/Developer/Concert/ConcertApp/MultiUserClient/Source/MultiUserReplicationEditor/Internal/Assets/MultiUserReplicationClientContent.h
--- a/Engine/Plugins/Developer/Concert/ConcertApp/MultiUserClient/Source/MultiUserReplicationEditor/Internal/Assets/MultiUserReplicationClientContent.h
+++ b/Engine/Plugins/Developer/Concert/ConcertApp/MultiUserClient/Source/MultiUserReplicationEditor/Internal/Assets/MultiUserReplicationClientContent.h
@@ -21,6 +21,13 @@ public:
 	UPROPERTY(Instanced)
 	TObjectPtr<UMultiUserReplicationStream> Stream;
 
+	/** The FConcertClientInfo::DisplayName of the client. */
+	UPROPERTY()
+	FString DisplayName;
+	/** The FConcertClientInfo::DeviceName of the client. */
+	UPROPERTY()
+	FString DeviceName;
+
 	UMultiUserReplicationClientContent();
 
 	/** Generates a description that can be sent to the MU server. */
diff --git a/Engine/Plugins/Developer/Concert/ConcertApp/MultiUserClient/Source/MultiUserReplicationEditor/Internal/Assets/MultiUserReplicationSessionPreset.h b/Engine/Plugins/Developer/Concert/ConcertApp/MultiUserClient/Source/MultiUserReplicationEditor/Internal/Assets/MultiUserReplicationSessionPreset.h
--- a/Engine/Plugins/Developer/Concert/ConcertApp/MultiUserClient/Source/MultiUserReplicationEditor/Internal/Assets/MultiUserReplicationSessionPreset.h
+++ b/Engine/Plugins/Developer/Concert/ConcertApp/MultiUserClient/Source/MultiUserReplicationEditor/Internal/Assets/MultiUserReplicationSessionPreset.h
@@ -6,17 +6,27 @@
 #include "UObject/Object.h"
 #include "MultiUserReplicationSessionPreset.generated.h"
 
-/**
- * 
- */
+struct FConcertClientInfo;
+
+/** Stores per-client replication settings so it can be loaded by a user to quickly set up a session. */
 UCLASS()
 class MULTIUSERREPLICATIONEDITOR_API UMultiUserReplicationSessionPreset : public UObject
 {
 	GENERATED_BODY()
 public:
 
-	UMultiUserReplicationClientContent* AddClient();
-	void RemoveClient(UMultiUserReplicationClientContent& Client);
+	/** @return The client preset that matches ClientInfo.DisplayName. If there are multiple, returns the one that matches ClientInfo.DeviceName, as well. */
+	UMultiUserReplicationClientContent* GetClientContent(const FConcertClientInfo& ClientInfo) const;
+	/** @return The client preset that matches both the display and device name. */
+	UMultiUserReplicationClientContent* GetExactClientContent(const FConcertClientInfo& ClientInfo) const;
+	
+	/** @return Whether a client that matches ClientInfo.DisplayName. */
+	bool ContainsClient(const FConcertClientInfo& ClientInfo) const { return GetClientContent(ClientInfo) != nullptr; }
+	/** @return Whether a client that matches both the display and device name is saved in this preset. */
+	bool ContainsExactClient(const FConcertClientInfo& ClientInfo) const { return GetExactClientContent(ClientInfo) != nullptr; }
+	
+	/** Adds a client to the preset if it's not already present. */
+	UMultiUserReplicationClientContent* AddClientIfUnique(const FConcertClientInfo& ClientInfo);
 
 	const TArray<TObjectPtr<UMultiUserReplicationClientContent>>& GetClientPresets() const { return ClientPresets; }
 
diff --git a/Engine/Plugins/Developer/Concert/ConcertApp/MultiUserClient/Source/MultiUserReplicationEditor/Internal/Assets/MultiUserReplicationStream.h b/Engine/Plugins/Developer/Concert/ConcertApp/MultiUserClient/Source/MultiUserReplicationEditor/Internal/Assets/MultiUserReplicationStream.h
--- a/Engine/Plugins/Developer/Concert/ConcertApp/MultiUserClient/Source/MultiUserReplicationEditor/Internal/Assets/MultiUserReplicationStream.h
+++ b/Engine/Plugins/Developer/Concert/ConcertApp/MultiUserClient/Source/MultiUserReplicationEditor/Internal/Assets/MultiUserReplicationStream.h
@@ -35,4 +35,7 @@ public:
 			return WeakThis.IsValid() ? &WeakThis->ReplicationMap : nullptr;
 		});
 	}
+	
+	/** Copies the stream content of OtherStream. */
+	void Copy(UMultiUserReplicationStream& OtherStream);
 };
diff --git a/Engine/Plugins/Developer/Concert/ConcertApp/MultiUserClient/Source/MultiUserReplicationEditor/MultiUserReplicationEditor.Build.cs b/Engine/Plugins/Developer/Concert/ConcertApp/MultiUserClient/Source/MultiUserReplicationEditor/MultiUserReplicationEditor.Build.cs
--- a/Engine/Plugins/Developer/Concert/ConcertApp/MultiUserClient/Source/MultiUserReplicationEditor/MultiUserReplicationEditor.Build.cs
+++ b/Engine/Plugins/Developer/Concert/ConcertApp/MultiUserClient/Source/MultiUserReplicationEditor/MultiUserReplicationEditor.Build.cs
@@ -40,7 +40,8 @@ namespace UnrealBuildTool.Rules
 					// Concert
 					"Concert",
 					"ConcertClient",
-					"ConcertSharedSlate"
+					"ConcertSharedSlate",
+					"ConcertTransport"
 				}
 			);
 		}
diff --git a/Engine/Plugins/Developer/Concert/ConcertApp/MultiUserClient/Source/MultiUserReplicationEditor/Private/Assets/MultiUserReplicationSessionPreset.cpp b/Engine/Plugins/Developer/Concert/ConcertApp/MultiUserClient/Source/MultiUserReplicationEditor/Private/Assets/MultiUserReplicationSessionPreset.cpp
--- a/Engine/Plugins/Developer/Concert/ConcertApp/MultiUserClient/Source/MultiUserReplicationEditor/Private/Assets/MultiUserReplicationSessionPreset.cpp
+++ b/Engine/Plugins/Developer/Concert/ConcertApp/MultiUserClient/Source/MultiUserReplicationEditor/Private/Assets/MultiUserReplicationSessionPreset.cpp
@@ -2,14 +2,58 @@
 
 #include "Assets/MultiUserReplicationSessionPreset.h"
 
-UMultiUserReplicationClientContent* UMultiUserReplicationSessionPreset::AddClient()
+#include "ConcertLogGlobal.h"
+#include "ConcertMessageData.h"
+
+UMultiUserReplicationClientContent* UMultiUserReplicationSessionPreset::GetClientContent(const FConcertClientInfo& ClientInfo) const
 {
-	UMultiUserReplicationClientContent* Result = NewObject<UMultiUserReplicationClientContent>(this);
-	ClientPresets.Add(Result);
-	return Result;
+	const auto IsMatch = [&ClientInfo](UMultiUserReplicationClientContent* Content)
+	{
+		return Content && Content->DisplayName == ClientInfo.DisplayName;
+	};
+	const auto IsPerfectMatch = [&ClientInfo, &IsMatch](UMultiUserReplicationClientContent* Content)
+	{
+		return IsMatch(Content) && Content->DeviceName == ClientInfo.DeviceName;
+	};
+	
+	UMultiUserReplicationClientContent* BestMatch = nullptr;
+	bool bIsBestMatchPerfect = false;
+	for (UMultiUserReplicationClientContent* Content : ClientPresets)
+	{
+		if (IsPerfectMatch(Content))
+		{
+			UE_CLOG(bIsBestMatchPerfect, LogConcert, Warning, TEXT("Preset %s contained client (name: %s, device: %s) multiple times"), *GetPathName(), *ClientInfo.DisplayName, *ClientInfo.DeviceName);
+			bIsBestMatchPerfect = true;
+			BestMatch = Content;
+		}
+		else if (!bIsBestMatchPerfect && IsMatch(Content))
+		{
+			BestMatch = Content;
+		}
+	}
+	
+	return BestMatch;
+}
+
+UMultiUserReplicationClientContent* UMultiUserReplicationSessionPreset::GetExactClientContent(const FConcertClientInfo& ClientInfo) const
+{
+	const TObjectPtr<UMultiUserReplicationClientContent>* Result = ClientPresets.FindByPredicate([&ClientInfo](const TObjectPtr<UMultiUserReplicationClientContent>& Content)
+	{
+		return Content && Content->DisplayName == ClientInfo.DisplayName && Content->DeviceName == ClientInfo.DeviceName;
+	});
+	return Result ? *Result : nullptr;
 }
 
-void UMultiUserReplicationSessionPreset::RemoveClient(UMultiUserReplicationClientContent& Client)
+UMultiUserReplicationClientContent* UMultiUserReplicationSessionPreset::AddClientIfUnique(const FConcertClientInfo& ClientInfo)
 {
-	ClientPresets.RemoveSingle(&Client);
+	if (ContainsExactClient(ClientInfo))
+	{
+		return nullptr;
+	}
+	
+	UMultiUserReplicationClientContent* Result = NewObject<UMultiUserReplicationClientContent>(this);
+	Result->DisplayName = ClientInfo.DisplayName;
+	Result->DeviceName = ClientInfo.DeviceName;
+	ClientPresets.Add(Result);
+	return Result;
 }
diff --git a/Engine/Plugins/Developer/Concert/ConcertApp/MultiUserClient/Source/MultiUserReplicationEditor/Private/Assets/MultiUserReplicationStream.cpp b/Engine/Plugins/Developer/Concert/ConcertApp/MultiUserClient/Source/MultiUserReplicationEditor/Private/Assets/MultiUserReplicationStream.cpp
--- a/Engine/Plugins/Developer/Concert/ConcertApp/MultiUserClient/Source/MultiUserReplicationEditor/Private/Assets/MultiUserReplicationStream.cpp
+++ b/Engine/Plugins/Developer/Concert/ConcertApp/MultiUserClient/Source/MultiUserReplicationEditor/Private/Assets/MultiUserReplicationStream.cpp
@@ -16,3 +16,9 @@ FConcertReplicationStream UMultiUserReplicationStream::GenerateDescription() con
 {
 	return { StreamId, ReplicationMap };
 }
+
+void UMultiUserReplicationStream::Copy(UMultiUserReplicationStream& OtherStream)
+{
+	StreamId = OtherStream.StreamId;
+	ReplicationMap = OtherStream.ReplicationMap;
+}
