Replication presets now have the option to exclude certain clients from being saved.
Adds a tick box to include all clients. If unchecked, the user can use a submenu to exclude clients.

diff --git a/Engine/Plugins/Developer/Concert/ConcertApp/MultiUserClient/Source/MultiUserClient/Private/Replication/Preset/PresetManager.cpp b/Engine/Plugins/Developer/Concert/ConcertApp/MultiUserClient/Source/MultiUserClient/Private/Replication/Preset/PresetManager.cpp
--- a/Engine/Plugins/Developer/Concert/ConcertApp/MultiUserClient/Source/MultiUserClient/Private/Replication/Preset/PresetManager.cpp
+++ b/Engine/Plugins/Developer/Concert/ConcertApp/MultiUserClient/Source/MultiUserClient/Private/Replication/Preset/PresetManager.cpp
@@ -110,6 +110,34 @@ namespace UE::MultiUserClient
 			});
 			EditModel.RemoveObjects(EmptyObjects);
 		}
+
+		static TArray<TPair<const FReplicationClient*, FConcertClientInfo>> DetermineSavedClients(
+			const FReplicationClientManager& ClientManager,
+			const IConcertClientSession& Session,
+			const FSavePresetOptions& Options
+			)
+		{
+			TArray<TPair<const FReplicationClient*, FConcertClientInfo>> IncludedClients;
+			ClientManager.ForEachClient([&Session, &IncludedClients, &Options](const FReplicationClient& Client)
+			{
+				FConcertClientInfo ClientInfo;
+				const bool bGotClientInfo = ClientUtils::GetClientDisplayInfo(Session, Client.GetEndpointId(), ClientInfo);
+				if (!ensure(bGotClientInfo))
+				{
+					return EBreakBehavior::Continue;
+				}
+
+				const bool bIsFilteredOut = Options.ClientFilterDelegate.IsBound() && Options.ClientFilterDelegate.Execute(ClientInfo) == EFilterResult::Exclude; 
+				if (bIsFilteredOut)
+				{
+					return EBreakBehavior::Continue;
+				}
+			
+				IncludedClients.Emplace(&Client, ClientInfo);
+				return EBreakBehavior::Continue;
+			});
+			return IncludedClients;
+		}
 	}
 	
 	FPresetManager::FPresetManager(const IConcertSyncClient& SyncClient, const FReplicationClientManager& ClientManager)
@@ -171,18 +199,38 @@ namespace UE::MultiUserClient
 		return InProgressSessionReplacementOp->GetFuture();
 	}
 
-	void FPresetManager::ExportToPresetAndSaveAs()
+	ECanSaveResult FPresetManager::CanSavePreset(const FSavePresetOptions& Options) const
 	{
-		UMultiUserReplicationSessionPreset* Preset = ExportToPreset();
+		const TSharedPtr<IConcertClientSession> Session = SyncClient.GetConcertClient()->GetCurrentSession();
+		checkf(Session, TEXT("FPresetManager is only supposed to exist while in a session"));
 		
-		TArray<UObject*> SavedAssets;
-		FEditorFileUtils::SaveAssetsAs({ Preset }, SavedAssets);
+		const TArray<TPair<const FReplicationClient*, FConcertClientInfo>> IncludedClients = Private::DetermineSavedClients(ClientManager, *Session, Options);
+		return IncludedClients.IsEmpty()
+			? ECanSaveResult::NoClients
+			: ECanSaveResult::Yes;
 	}
 
-	UMultiUserReplicationSessionPreset* FPresetManager::ExportToPreset() const
+	UMultiUserReplicationSessionPreset* FPresetManager::ExportToPresetAndSaveAs(const FSavePresetOptions& Options)
+	{
+		UMultiUserReplicationSessionPreset* Preset = ExportToPreset(Options);
+		if (Preset)
+		{
+			TArray<UObject*> SavedAssets;
+			FEditorFileUtils::SaveAssetsAs({ Preset }, SavedAssets);
+		}
+		return Preset;
+	}
+
+	UMultiUserReplicationSessionPreset* FPresetManager::ExportToPreset(const FSavePresetOptions& Options) const
 	{
 		const TSharedPtr<IConcertClientSession> Session = SyncClient.GetConcertClient()->GetCurrentSession();
 		checkf(Session, TEXT("FPresetManager is only supposed to exist while in a session"));
+
+		const TArray<TPair<const FReplicationClient*, FConcertClientInfo>> IncludedClients = Private::DetermineSavedClients(ClientManager, *Session, Options);
+		if (IncludedClients.IsEmpty())
+		{
+			return nullptr;
+		}
 		
 		UMultiUserReplicationSessionPreset* Preset = NewObject<UMultiUserReplicationSessionPreset>(
 			GetTransientPackage(),
@@ -190,16 +238,10 @@ namespace UE::MultiUserClient
 			// Mark as transient so FEditorFileUtils::SaveAssetsAs creates a new package for the object.
 			RF_Transient
 			);
-		ClientManager.ForEachClient([this, &Session, &Preset](const FReplicationClient& Client)
+		for (const TPair<const FReplicationClient*, FConcertClientInfo>& ClientData : IncludedClients)
 		{
-			FConcertClientInfo ClientInfo;
-			const bool bGotClientInfo = ClientUtils::GetClientDisplayInfo(*Session, Client.GetEndpointId(), ClientInfo);
-			if (!ensure(bGotClientInfo))
-			{
-				return EBreakBehavior::Continue;
-			}
-			
-			UMultiUserReplicationClientContent* ClientContent_ToCopy = Client.GetClientContent();
+			const auto[Client, ClientInfo] = ClientData;
+			UMultiUserReplicationClientContent* ClientContent_ToCopy = Client->GetClientContent();
 			UMultiUserReplicationClientContent* ClientContent_InPreset = Preset->AddClientIfUnique(ClientInfo);
 			if (!ClientContent_InPreset)
 			{
@@ -208,12 +250,11 @@ namespace UE::MultiUserClient
 					*ClientInfo.DisplayName,
 					*ClientInfo.DeviceName
 					);
-				return EBreakBehavior::Continue;
+				continue;
 			}
 
 			ClientContent_InPreset->Stream->Copy(*ClientContent_ToCopy->Stream);
-			return EBreakBehavior::Continue;
-		});
+		}
 		
 		return Preset;
 	}
diff --git a/Engine/Plugins/Developer/Concert/ConcertApp/MultiUserClient/Source/MultiUserClient/Private/Replication/Preset/PresetManager.h b/Engine/Plugins/Developer/Concert/ConcertApp/MultiUserClient/Source/MultiUserClient/Private/Replication/Preset/PresetManager.h
--- a/Engine/Plugins/Developer/Concert/ConcertApp/MultiUserClient/Source/MultiUserClient/Private/Replication/Preset/PresetManager.h
+++ b/Engine/Plugins/Developer/Concert/ConcertApp/MultiUserClient/Source/MultiUserClient/Private/Replication/Preset/PresetManager.h
@@ -3,9 +3,12 @@
 #pragma once
 
 #include "Async/Future.h"
+#include "Delegates/Delegate.h"
 #include "HAL/Platform.h"
+#include "Misc/Optional.h"
 #include "Templates/UnrealTemplate.h"
 
+struct FConcertClientInfo;
 class IConcertSyncClient;
 class UMultiUserReplicationSessionPreset;
 enum class EMultiUserClientPresetLoadMode : uint8;
@@ -53,6 +56,23 @@ namespace UE::MultiUserClient
 		ClearUnreferencedClients = 1 << 0
 	};
 	ENUM_CLASS_FLAGS(EApplyPresetFlags);
+
+	enum class EFilterResult : uint8 { Include, Exclude };
+	DECLARE_DELEGATE_RetVal_OneParam(EFilterResult, FFilterClientForPreset, const FConcertClientInfo&)
+	/** Optionsl for saving a preset */
+	struct FSavePresetOptions
+	{
+		/** Filter that decides whether a client should be included in the preset. */
+		FFilterClientForPreset ClientFilterDelegate; 
+	};
+
+	enum class ECanSaveResult : uint8
+	{
+		/** Yes, a preset can be saved. */
+		Yes,
+		/** There are no clients to save for */
+		NoClients
+	};
 	
 	/**
 	 * Implements all logic for managing presets in the MU session: saving and loading presets.
@@ -71,11 +91,13 @@ namespace UE::MultiUserClient
 		/** Applies Preset to all clients in the session. */
 		TFuture<FReplaceSessionContentResult> ReplaceSessionContentWithPreset(const UMultiUserReplicationSessionPreset& Preset, EApplyPresetFlags Flags = EApplyPresetFlags::None);
 
+		/** @return Whether a preset can be saved (i.e. at least one client is included). */
+		ECanSaveResult CanSavePreset(const FSavePresetOptions& Options = {}) const;
 		/**
 		 * Exports the current session content to a preset, asks the user where to save it, then saves it.
 		 * @return A future that finishes when saving has completed.
 		 */
-		void ExportToPresetAndSaveAs();
+		UMultiUserReplicationSessionPreset* ExportToPresetAndSaveAs(const FSavePresetOptions& Options = {});
 
 	private:
 
@@ -88,7 +110,7 @@ namespace UE::MultiUserClient
 		TSharedPtr<TPromise<FReplaceSessionContentResult>> InProgressSessionReplacementOp;
 
 		/** Exports the current session content to a preset. */
-		UMultiUserReplicationSessionPreset* ExportToPreset() const;
+		UMultiUserReplicationSessionPreset* ExportToPreset(const FSavePresetOptions& Options) const;
 
 		/** Called when the local client receives a remote edit. */
 		void OnPostRemoteEditApplied(const ConcertSyncClient::Replication::FRemoteEditEvent&) const;
diff --git a/Engine/Plugins/Developer/Concert/ConcertApp/MultiUserClient/Source/MultiUserClient/Private/Widgets/ActiveSession/Replication/Client/Multi/SMultiClientView.cpp b/Engine/Plugins/Developer/Concert/ConcertApp/MultiUserClient/Source/MultiUserClient/Private/Widgets/ActiveSession/Replication/Client/Multi/SMultiClientView.cpp
--- a/Engine/Plugins/Developer/Concert/ConcertApp/MultiUserClient/Source/MultiUserClient/Private/Widgets/ActiveSession/Replication/Client/Multi/SMultiClientView.cpp
+++ b/Engine/Plugins/Developer/Concert/ConcertApp/MultiUserClient/Source/MultiUserClient/Private/Widgets/ActiveSession/Replication/Client/Multi/SMultiClientView.cpp
@@ -150,7 +150,7 @@ namespace UE::MultiUserClient
 			},
 			.ShouldDisplayObjectDelegate = FShouldDisplayObject::CreateSP(this, &SMultiClientView::ShouldDisplayObject),
 		};
-		ViewerParams.RightOfObjectSearchBar.Widget = SNew(SPresetComboButton, *InMultiUserReplicationManager.GetPresetManager());
+		ViewerParams.RightOfObjectSearchBar.Widget = SNew(SPresetComboButton, *InConcertClient, *InMultiUserReplicationManager.GetPresetManager());
 		StreamEditor = CreateBaseMultiStreamEditor(MoveTemp(Params), MoveTemp(ViewerParams));
 		check(StreamEditor);
 		
diff --git a/Engine/Plugins/Developer/Concert/ConcertApp/MultiUserClient/Source/MultiUserClient/Private/Widgets/ActiveSession/Replication/Client/SPresetComboButton.cpp b/Engine/Plugins/Developer/Concert/ConcertApp/MultiUserClient/Source/MultiUserClient/Private/Widgets/ActiveSession/Replication/Client/SPresetComboButton.cpp
--- a/Engine/Plugins/Developer/Concert/ConcertApp/MultiUserClient/Source/MultiUserClient/Private/Widgets/ActiveSession/Replication/Client/SPresetComboButton.cpp
+++ b/Engine/Plugins/Developer/Concert/ConcertApp/MultiUserClient/Source/MultiUserClient/Private/Widgets/ActiveSession/Replication/Client/SPresetComboButton.cpp
@@ -12,6 +12,7 @@
 #include "Framework/Notifications/NotificationManager.h"
 #include "SSimpleComboButton.h"
 #include "Styling/AppStyle.h"
+#include "Widgets/Client/SClientName.h"
 #include "Widgets/Layout/SBorder.h"
 #include "Widgets/Notifications/SNotificationList.h"
 
@@ -42,10 +43,13 @@ namespace UE::MultiUserClient
 			default: checkNoEntry(); return FText::GetEmpty();
 			}
 		}
+
+		static FText GetSavedClientsText() { return LOCTEXT("Save.IncludedClients.Label", "Saved clients"); }
 	}
 	
-	void SPresetComboButton::Construct(const FArguments& InArgs, FPresetManager& InPresetManager)
+	void SPresetComboButton::Construct(const FArguments& InArgs, const IConcertClient& InClient, FPresetManager& InPresetManager)
 	{
+		Client = &InClient;
 		PresetManager = &InPresetManager;
 		ChildSlot
 		[
@@ -59,40 +63,146 @@ namespace UE::MultiUserClient
 
 	TSharedRef<SWidget> SPresetComboButton::CreateMenuContent()
 	{
-		FMenuBuilder MenuBuilder(true, nullptr);
+		FMenuBuilder MenuBuilder(false, nullptr);
 
-		CreateSaveMenuContent(MenuBuilder);
+		MenuBuilder.BeginSection(NAME_None, LOCTEXT("Section.Save", "Save preset"));
+		BuildSaveMenuContent(MenuBuilder);
+		MenuBuilder.EndSection();
 
 		MenuBuilder.BeginSection(NAME_None, LOCTEXT("Section.Import", "Import preset"));
-		CreateLoadMenuContent(MenuBuilder);
+		BuildLoadMenuContent(MenuBuilder);
 		MenuBuilder.EndSection();
 
 		return MenuBuilder.MakeWidget();
 	}
 
-	void SPresetComboButton::CreateSaveMenuContent(FMenuBuilder& MenuBuilder)
+	void SPresetComboButton::BuildSaveMenuContent(FMenuBuilder& MenuBuilder)
 	{
 		MenuBuilder.AddMenuEntry(
-			LOCTEXT("SavePresetAs.Label", "Save Preset as..."),
-			LOCTEXT("SavePresetAs.ToolTip", "Saves what each client was replicating as a preset."),
+			LOCTEXT("Save.SavePresetAs.Label", "Save Preset as..."),
+			TAttribute<FText>::CreateLambda([this]()
+			{
+				const ECanSaveResult CanSaveResult = PresetManager->CanSavePreset(BuildSaveOptions());
+				switch (CanSaveResult)
+				{
+				case ECanSaveResult::Yes: return LOCTEXT("Save.SavePresetAs.ToolTip.Yes", "Saves what each client is replicating as a preset.");
+				case ECanSaveResult::NoClients: return FText::Format(LOCTEXT("Save.SavePresetAs.ToolTip.NoClients", "Selec the clients you want to save first in '{0}'"), Private::GetSavedClientsText());
+				default: return FText::GetEmpty();
+				}
+				
+			}),
 			FSlateIcon(FAppStyle::Get().GetStyleSetName(), "AssetEditor.SaveAssetAs"),
-			FUIAction(FExecuteAction::CreateSP(this, &SPresetComboButton::SavePresetAs))	
+			FUIAction(
+				FExecuteAction::CreateSP(this, &SPresetComboButton::SavePresetAs),
+				FCanExecuteAction::CreateLambda([this]{ return PresetManager->CanSavePreset(BuildSaveOptions()) == ECanSaveResult::Yes; })
+				)	
 		);
+
+		MenuBuilder.AddMenuEntry(
+			LOCTEXT("Save.IncludeAllClients.Label", "Save all clients"),
+			LOCTEXT("Save.IncludeAllClients.ToolTip", "Whether you want to include all clients in the session into the preset."),
+			FSlateIcon(),
+			FUIAction(
+				FExecuteAction::CreateLambda([this](){ Options.bIncludeAllClients = !Options.bIncludeAllClients; }),
+				FCanExecuteAction::CreateLambda([](){ return true; }),
+				FIsActionChecked::CreateLambda([this]{ return Options.bIncludeAllClients; })
+				),
+			NAME_None,
+			EUserInterfaceActionType::ToggleButton
+			);
+		
+		MenuBuilder.AddSubMenu(
+			Private::GetSavedClientsText(),
+			LOCTEXT("Save.IncludedClients.ToolTip", "Select the clients you want to save into the preset"),
+			FNewMenuDelegate::CreateRaw(this, &SPresetComboButton::BuildExcludedClientSubmenu),
+			FUIAction(
+				FExecuteAction::CreateLambda([]{}),
+				FCanExecuteAction::CreateLambda([this]{ return !Options.bIncludeAllClients; }),
+				FIsActionChecked(),
+				FIsActionButtonVisible::CreateLambda([this]{ return !Options.bIncludeAllClients; })
+				),
+			NAME_None,
+			EUserInterfaceActionType::None
+			);
+	}
+
+	void SPresetComboButton::BuildExcludedClientSubmenu(FMenuBuilder& MenuBuilder)
+	{
+		const TSharedPtr<IConcertClientSession> Session = Client->GetCurrentSession();
+		if (!ensure(Session))
+		{
+			return;
+		}
+
+		const auto GetDisplayText = [](const FConcertClientInfo& ClientInfo, bool bDisplayAsLocalClient = false)
+		{
+			return ConcertClientSharedSlate::SClientName::GetDisplayText(ClientInfo, bDisplayAsLocalClient);
+		};
+		
+		TArray<FConcertSessionClientInfo> RemoteClients = Session->GetSessionClients();
+		RemoteClients.Sort([&GetDisplayText](const FConcertSessionClientInfo& Left, const FConcertSessionClientInfo& Right)
+		{
+			return GetDisplayText(Left.ClientInfo).ToString() < GetDisplayText(Right.ClientInfo).ToString();
+		});
+
+		const auto AddClient = [this, &MenuBuilder, &GetDisplayText](const FConcertClientInfo& ClientInfo, bool bDisplayHasLocalClient = false)
+		{
+			const FText DisplayText = GetDisplayText(ClientInfo, bDisplayHasLocalClient);
+			MenuBuilder.AddMenuEntry(
+				FText::Format(LOCTEXT("Save.Client.LabelFmt", "{0}"), DisplayText),
+				FText::Format(LOCTEXT("Save.Client.ToolTipFmt", "Check if you want the {0} saved in the preset."), DisplayText),
+				FSlateIcon(),
+				FUIAction(
+					FExecuteAction::CreateLambda([this, ClientInfo]()
+					{
+						if (Options.IncludedClients.Contains(ClientInfo))
+						{
+							Options.IncludedClients.RemoveSingle(ClientInfo);
+						}
+						else
+						{
+							Options.IncludedClients.Add(ClientInfo);
+						}
+					}),
+					FCanExecuteAction::CreateLambda([](){ return true; }),
+					FIsActionChecked::CreateLambda([this, ClientInfo]{ return Options.IncludedClients.Contains(ClientInfo); })
+					),
+				NAME_None,
+				EUserInterfaceActionType::ToggleButton
+				);
+		};
+		AddClient(Session->GetLocalClientInfo(), true);
+		for (const FConcertSessionClientInfo& RemoteClient : RemoteClients)
+		{
+			AddClient(RemoteClient.ClientInfo);
+		}
 	}
 
-	void SPresetComboButton::CreateLoadMenuContent(FMenuBuilder& MenuBuilder)
+	void SPresetComboButton::BuildLoadMenuContent(FMenuBuilder& MenuBuilder)
 	{
 		MenuBuilder.AddMenuEntry(
 			LOCTEXT("ImportPreset.ClearOtherClients.Label", "Clear clients not in preset"),
-			LOCTEXT("ImportPreset.ClearOtherClients.ToolTip", "If checked, clients that were not in the session when the preset was created will get their content reset, too."),
+			LOCTEXT("ImportPreset.ClearOtherClients.ToolTip", "Clients that were not in the session when the preset was created will get their content reset, too."),
 			FSlateIcon(),
 			FUIAction(
-				FExecuteAction::CreateLambda([this](){ Options.bResetAllOtherClients = !Options.bResetAllOtherClients; }),
+				FExecuteAction::CreateLambda([this](){ Options.bResetAllOtherClients = true; }),
 				FCanExecuteAction::CreateLambda([](){ return true; }),
 				FIsActionChecked::CreateLambda([this]{ return Options.bResetAllOtherClients; })
 				),
 			NAME_None,
-			EUserInterfaceActionType::ToggleButton
+			EUserInterfaceActionType::RadioButton
+			);
+		MenuBuilder.AddMenuEntry(
+			LOCTEXT("ImportPreset.AdditivelyAdd.Label", "Only change clients in preset"),
+			LOCTEXT("ImportPreset.AdditivelyAdd.ToolTip", "Clients that were not in the session when the preset was created will not be modified by this preset."),
+			FSlateIcon(),
+			FUIAction(
+				FExecuteAction::CreateLambda([this](){ Options.bResetAllOtherClients = false; }),
+				FCanExecuteAction::CreateLambda([](){ return true; }),
+				FIsActionChecked::CreateLambda([this]{ return !Options.bResetAllOtherClients; })
+				),
+			NAME_None,
+			EUserInterfaceActionType::RadioButton
 			);
 			
 		FContentBrowserModule& ContentBrowserModule = FModuleManager::Get().LoadModuleChecked<FContentBrowserModule>(TEXT("ContentBrowser"));
@@ -106,8 +216,8 @@ namespace UE::MultiUserClient
 		AssetPickerConfig.bAllowDragging = false;
 		
 		TSharedRef<SWidget> PresetPicker = SNew(SBox)
-			.HeightOverride(400)
-			.WidthOverride(300)
+			.HeightOverride(450)
+			.WidthOverride(320)
 			[
 				SNew(SBorder)
 				.BorderImage(FAppStyle::GetBrush("Menu.Background"))
@@ -118,12 +228,14 @@ namespace UE::MultiUserClient
 		MenuBuilder.AddWidget(PresetPicker, FText(), true, false);
 	}
 
-	void SPresetComboButton::SavePresetAs()
+	void SPresetComboButton::SavePresetAs() const
 	{
-		PresetManager->ExportToPresetAndSaveAs();
+		PresetManager->ExportToPresetAndSaveAs(
+			BuildSaveOptions()
+			);
 	}
 
-	void SPresetComboButton::LoadPreset(const FAssetData& AssetData)
+	void SPresetComboButton::LoadPreset(const FAssetData& AssetData) const
 	{
 		FSlateNotificationManager& NotificationManager = FSlateNotificationManager::Get();
 		if (UMultiUserReplicationSessionPreset* Preset = Cast<UMultiUserReplicationSessionPreset>(AssetData.GetAsset()))
@@ -165,6 +277,22 @@ namespace UE::MultiUserClient
 
 		return Flags;
 	}
+
+	FSavePresetOptions SPresetComboButton::BuildSaveOptions() const
+	{
+		FSavePresetOptions SavePresetOptions;
+		
+		if (!Options.bIncludeAllClients)
+		{
+			const auto Filter = [this](const FConcertClientInfo& Info)
+			{
+				return Options.IncludedClients.Contains(Info) ? EFilterResult::Include : EFilterResult::Exclude;
+			};
+			SavePresetOptions.ClientFilterDelegate.BindLambda(Filter);
+		}
+
+		return SavePresetOptions;
+	}
 }
 
 #undef LOCTEXT_NAMESPACE
\ No newline at end of file
diff --git a/Engine/Plugins/Developer/Concert/ConcertApp/MultiUserClient/Source/MultiUserClient/Private/Widgets/ActiveSession/Replication/Client/SPresetComboButton.h b/Engine/Plugins/Developer/Concert/ConcertApp/MultiUserClient/Source/MultiUserClient/Private/Widgets/ActiveSession/Replication/Client/SPresetComboButton.h
--- a/Engine/Plugins/Developer/Concert/ConcertApp/MultiUserClient/Source/MultiUserClient/Private/Widgets/ActiveSession/Replication/Client/SPresetComboButton.h
+++ b/Engine/Plugins/Developer/Concert/ConcertApp/MultiUserClient/Source/MultiUserClient/Private/Widgets/ActiveSession/Replication/Client/SPresetComboButton.h
@@ -4,13 +4,16 @@
 
 #include "Widgets/DeclarativeSyntaxSupport.h"
 #include "Widgets/SCompoundWidget.h"
+#include "Widgets/Disconnected/ConcertClientSessionBrowserController.h"
 
 namespace UE::MultiUserClient
 {
-	enum class EApplyPresetFlags : uint8;
+	struct FSavePresetOptions;
 }
 
 class FMenuBuilder;
+class IConcertClient;
+namespace UE::MultiUserClient { enum class EApplyPresetFlags : uint8; }
 
 namespace UE::MultiUserClient
 {
@@ -28,10 +31,12 @@ namespace UE::MultiUserClient
 			
 		SLATE_END_ARGS()
 
-		void Construct(const FArguments& InArgs, FPresetManager& InPresetManager UE_LIFETIMEBOUND);
+		void Construct(const FArguments& InArgs, const IConcertClient& InClient UE_LIFETIMEBOUND, FPresetManager& InPresetManager UE_LIFETIMEBOUND);
 
 	private:
 
+		/** Used to get clients in the session (for filtering purposes). */
+		const IConcertClient* Client = nullptr;
 		/** Used to save & load preset. */
 		FPresetManager* PresetManager = nullptr;
 
@@ -39,19 +44,27 @@ namespace UE::MultiUserClient
 		{
 			/** Clients not mentioned by the preset will get their content wiped. */
 			bool bResetAllOtherClients = true;
+
+			/** Whether the user want so to capture all clients in the preset. */
+			bool bIncludeAllClients = true;
+			
+			/** The clients that the preset should not be captured for */
+			TArray<FConcertClientInfo> IncludedClients;
 		} Options;
 
 		/** Creates the Save & Load options for the menu. */
 		TSharedRef<SWidget> CreateMenuContent();
-		void CreateSaveMenuContent(FMenuBuilder& MenuBuilder);
-		void CreateLoadMenuContent(FMenuBuilder& MenuBuilder);
+		void BuildSaveMenuContent(FMenuBuilder& MenuBuilder);
+		void BuildExcludedClientSubmenu(FMenuBuilder& MenuBuilder);
+		void BuildLoadMenuContent(FMenuBuilder& MenuBuilder);
 
 		/** Saves the current session content. */
-		void SavePresetAs();
+		void SavePresetAs() const;
 		/** Handles loading the preset. */
-		void LoadPreset(const FAssetData& AssetData);
+		void LoadPreset(const FAssetData& AssetData) const;
 
 		EApplyPresetFlags BuildFlags() const;
+		FSavePresetOptions BuildSaveOptions() const;
 	};
 }
 
