Move types into separate headers and remove unused headers.

diff --git a/Engine/Plugins/Developer/Concert/ConcertApp/MultiUserClient/Source/MultiUserClient/Private/Replication/Preset/Data/ReplacedWithPresetEvent.h b/Engine/Plugins/Developer/Concert/ConcertApp/MultiUserClient/Source/MultiUserClient/Private/Replication/Preset/Data/ReplacedWithPresetEvent.h
deleted file mode 100644
diff --git a/Engine/Plugins/Developer/Concert/ConcertApp/MultiUserClient/Source/MultiUserClient/Private/Replication/Preset/Operations/ReplaceSessionContentOperation.cpp b/Engine/Plugins/Developer/Concert/ConcertApp/MultiUserClient/Source/MultiUserClient/Private/Replication/Preset/Operations/ReplaceSessionContentOperation.cpp
deleted file mode 100644
diff --git a/Engine/Plugins/Developer/Concert/ConcertApp/MultiUserClient/Source/MultiUserClient/Private/Replication/Preset/Operations/ReplaceSessionContentOperation.h b/Engine/Plugins/Developer/Concert/ConcertApp/MultiUserClient/Source/MultiUserClient/Private/Replication/Preset/Operations/ReplaceSessionContentOperation.h
deleted file mode 100644
diff --git a/Engine/Plugins/Developer/Concert/ConcertApp/MultiUserClient/Source/MultiUserClient/Private/Replication/Preset/PresetManager.h b/Engine/Plugins/Developer/Concert/ConcertApp/MultiUserClient/Source/MultiUserClient/Private/Replication/Preset/PresetManager.h
--- a/Engine/Plugins/Developer/Concert/ConcertApp/MultiUserClient/Source/MultiUserClient/Private/Replication/Preset/PresetManager.h
+++ b/Engine/Plugins/Developer/Concert/ConcertApp/MultiUserClient/Source/MultiUserClient/Private/Replication/Preset/PresetManager.h
@@ -2,6 +2,9 @@
 
 #pragma once
 
+#include "ReplaceSessionContentResult.h"
+#include "SavePresetOptions.h"
+
 #include "Async/Future.h"
 #include "Delegates/Delegate.h"
 #include "HAL/Platform.h"
@@ -18,61 +21,6 @@ namespace UE::ConcertSyncClient::Replication { struct FRemoteEditEvent; }
 namespace UE::MultiUserClient
 {
 	class FReplicationClientManager;
-
-	enum class EReplaceSessionContentErrorCode : uint8
-	{
-		/** Request completed successfully. */
-		Success,
-		
-		/** Request cancelled because FPresetManager was destroyed - probably because the use left the session during the request. */
-		Cancelled,
-		/** Another locally initiated operation is already in progress. */
-		InProgress,
-		
-		/** Request timed out */
-		Timeout,
-		/** The feature is not enabled (i.e. EConcertSyncSessionFlags::ShouldEnableRemoteEditing or EConcertSyncSessionFlags::ShouldAllowGlobalMuting were not set on the server).  */
-		FeatureDisabled,
-		/** Server rejected the change because it was not valid */
-		Rejected
-	};
-	
-	/** Result of FPresetManager::ReplaceSessionContentWithPreset. */
-	struct FReplaceSessionContentResult
-	{
-		EReplaceSessionContentErrorCode ErrorCode;
-
-		FReplaceSessionContentResult(EReplaceSessionContentErrorCode ErrorCode = EReplaceSessionContentErrorCode::Success)
-			: ErrorCode(ErrorCode)
-		{}
-
-		bool IsSuccess() const { return ErrorCode == EReplaceSessionContentErrorCode::Success; }
-	};
-
-	enum class EApplyPresetFlags : uint8
-	{
-		None,
-		/** If set clients that were not in the session when the preset was created will get their content reset, too. */
-		ClearUnreferencedClients = 1 << 0
-	};
-	ENUM_CLASS_FLAGS(EApplyPresetFlags);
-
-	enum class EFilterResult : uint8 { Include, Exclude };
-	DECLARE_DELEGATE_RetVal_OneParam(EFilterResult, FFilterClientForPreset, const FConcertClientInfo&)
-	/** Optionsl for saving a preset */
-	struct FSavePresetOptions
-	{
-		/** Filter that decides whether a client should be included in the preset. */
-		FFilterClientForPreset ClientFilterDelegate; 
-	};
-
-	enum class ECanSaveResult : uint8
-	{
-		/** Yes, a preset can be saved. */
-		Yes,
-		/** There are no clients to save for */
-		NoClients
-	};
 	
 	/**
 	 * Implements all logic for managing presets in the MU session: saving and loading presets.
diff --git a/Engine/Plugins/Developer/Concert/ConcertApp/MultiUserClient/Source/MultiUserClient/Private/Replication/Preset/ReplaceSessionContentResult.h b/Engine/Plugins/Developer/Concert/ConcertApp/MultiUserClient/Source/MultiUserClient/Private/Replication/Preset/ReplaceSessionContentResult.h
new file mode 100644
--- /dev/null
+++ b/Engine/Plugins/Developer/Concert/ConcertApp/MultiUserClient/Source/MultiUserClient/Private/Replication/Preset/ReplaceSessionContentResult.h
@@ -0,0 +1,39 @@
+﻿// Copyright Epic Games, Inc. All Rights Reserved.
+
+#pragma once
+
+#include "HAL/Platform.h"
+
+namespace UE::MultiUserClient
+{
+	enum class EReplaceSessionContentErrorCode : uint8
+	{
+		/** Request completed successfully. */
+		Success,
+		
+		/** Request cancelled because FPresetManager was destroyed - probably because the use left the session during the request. */
+		Cancelled,
+		/** Another locally initiated operation is already in progress. */
+		InProgress,
+		
+		/** Request timed out */
+		Timeout,
+		/** The feature is not enabled (i.e. EConcertSyncSessionFlags::ShouldEnableRemoteEditing or EConcertSyncSessionFlags::ShouldAllowGlobalMuting were not set on the server).  */
+		FeatureDisabled,
+		/** Server rejected the change because it was not valid */
+		Rejected
+	};
+	
+	/** Result of FPresetManager::ReplaceSessionContentWithPreset. */
+	struct FReplaceSessionContentResult
+	{
+		EReplaceSessionContentErrorCode ErrorCode;
+
+		FReplaceSessionContentResult(EReplaceSessionContentErrorCode ErrorCode = EReplaceSessionContentErrorCode::Success)
+			: ErrorCode(ErrorCode)
+		{}
+
+		bool IsSuccess() const { return ErrorCode == EReplaceSessionContentErrorCode::Success; }
+	};
+}
+
diff --git a/Engine/Plugins/Developer/Concert/ConcertApp/MultiUserClient/Source/MultiUserClient/Private/Replication/Preset/SavePresetOptions.h b/Engine/Plugins/Developer/Concert/ConcertApp/MultiUserClient/Source/MultiUserClient/Private/Replication/Preset/SavePresetOptions.h
new file mode 100644
--- /dev/null
+++ b/Engine/Plugins/Developer/Concert/ConcertApp/MultiUserClient/Source/MultiUserClient/Private/Replication/Preset/SavePresetOptions.h
@@ -0,0 +1,38 @@
+﻿// Copyright Epic Games, Inc. All Rights Reserved.
+
+#pragma once
+
+#include "Delegates/Delegate.h"
+#include "HAL/Platform.h"
+#include "Misc/EnumClassFlags.h"
+
+struct FConcertClientInfo;
+
+namespace UE::MultiUserClient
+{
+	enum class EApplyPresetFlags : uint8
+	{
+		None,
+		/** If set clients that were not in the session when the preset was created will get their content reset, too. */
+		ClearUnreferencedClients = 1 << 0
+	};
+	ENUM_CLASS_FLAGS(EApplyPresetFlags);
+
+	enum class EFilterResult : uint8 { Include, Exclude };
+	DECLARE_DELEGATE_RetVal_OneParam(EFilterResult, FFilterClientForPreset, const FConcertClientInfo&);
+	
+	/** Options for saving a preset */
+	struct FSavePresetOptions
+	{
+		/** Filter that decides whether a client should be included in the preset. */
+		FFilterClientForPreset ClientFilterDelegate; 
+	};
+
+	enum class ECanSaveResult : uint8
+	{
+		/** Yes, a preset can be saved. */
+		Yes,
+		/** There are no clients to save for */
+		NoClients
+	};
+}
\ No newline at end of file
