Replication presets now also save & load mute state

diff --git a/Engine/Plugins/Developer/Concert/ConcertApp/MultiUserClient/Source/MultiUserClient/Private/Replication/MultiUserReplicationManager.cpp b/Engine/Plugins/Developer/Concert/ConcertApp/MultiUserClient/Source/MultiUserClient/Private/Replication/MultiUserReplicationManager.cpp
--- a/Engine/Plugins/Developer/Concert/ConcertApp/MultiUserClient/Source/MultiUserClient/Private/Replication/MultiUserReplicationManager.cpp
+++ b/Engine/Plugins/Developer/Concert/ConcertApp/MultiUserClient/Source/MultiUserClient/Private/Replication/MultiUserReplicationManager.cpp
@@ -237,7 +237,7 @@ namespace UE::MultiUserClient
 		, QueryService(*InClient)
 		, ClientManager(InClient, InClient->GetConcertClient()->GetCurrentSession().ToSharedRef(), InDiscoveryContainer, QueryService.GetStreamAndAuthorityQueryService())
 		, MuteManager(*InClient, QueryService.GetMuteStateQueryService(), ClientManager.GetAuthorityCache())
-		, PresetManager(*InClient, ClientManager)
+		, PresetManager(*InClient, ClientManager, MuteManager.GetSynchronizer())
 		, PropertySelector(ClientManager)
 		, ChangeLevelHandler(ClientManager.GetLocalClient().GetClientEditModel().Get())
 		, PreventReplicatedPropertyTransaction(*InClient, ClientManager, MuteManager)
diff --git a/Engine/Plugins/Developer/Concert/ConcertApp/MultiUserClient/Source/MultiUserClient/Private/Replication/Muting/MuteStateSynchronizer.cpp b/Engine/Plugins/Developer/Concert/ConcertApp/MultiUserClient/Source/MultiUserClient/Private/Replication/Muting/MuteStateSynchronizer.cpp
--- a/Engine/Plugins/Developer/Concert/ConcertApp/MultiUserClient/Source/MultiUserClient/Private/Replication/Muting/MuteStateSynchronizer.cpp
+++ b/Engine/Plugins/Developer/Concert/ConcertApp/MultiUserClient/Source/MultiUserClient/Private/Replication/Muting/MuteStateSynchronizer.cpp
@@ -29,10 +29,12 @@ namespace UE::MultiUserClient
 		for (const TPair<FSoftObjectPath, FConcertReplication_ObjectMuteSetting>& MutedObject : Request.ObjectsToMute)
 		{
 			MutedObjects.Add(MutedObject.Key);
+			ExplicitlyMutedObjects.Add(MutedObject.Key, MutedObject.Value);
 		}
 		for (const TPair<FSoftObjectPath, FConcertReplication_ObjectMuteSetting>& UnmutedObject : Request.ObjectsToUnmute)
 		{
 			MutedObjects.Remove(UnmutedObject.Key);
+			ExplicitlyUnmutedObjects.Add(UnmutedObject.Key, UnmutedObject.Value);
 		}
 		
 		OnMuteStateChangedDelegate.Broadcast();
@@ -45,6 +47,8 @@ namespace UE::MultiUserClient
 
 		Algo::Transform(NewMuteState.ExplicitlyMutedObjects, MutedObjects, [](const TPair<FSoftObjectPath, FConcertReplication_ObjectMuteSetting>& Pair){ return Pair.Key; });
 		MutedObjects.Append(NewMuteState.ImplicitlyMutedObjects);
+		ExplicitlyMutedObjects = NewMuteState.ExplicitlyMutedObjects;
+		ExplicitlyUnmutedObjects = NewMuteState.ExplicitlyUnmutedObjects;
 
 		OnMuteStateChangedDelegate.Broadcast();
 	}
diff --git a/Engine/Plugins/Developer/Concert/ConcertApp/MultiUserClient/Source/MultiUserClient/Private/Replication/Muting/MuteStateSynchronizer.h b/Engine/Plugins/Developer/Concert/ConcertApp/MultiUserClient/Source/MultiUserClient/Private/Replication/Muting/MuteStateSynchronizer.h
--- a/Engine/Plugins/Developer/Concert/ConcertApp/MultiUserClient/Source/MultiUserClient/Private/Replication/Muting/MuteStateSynchronizer.h
+++ b/Engine/Plugins/Developer/Concert/ConcertApp/MultiUserClient/Source/MultiUserClient/Private/Replication/Muting/MuteStateSynchronizer.h
@@ -2,7 +2,10 @@
 
 #pragma once
 
+#include "Replication/Messages/Muting.h"
+
 #include "Containers/Set.h"
+#include "Containers/Map.h"
 #include "Delegates/Delegate.h"
 #include "HAL/Platform.h"
 #include "Templates/UnrealTemplate.h"
@@ -31,6 +34,9 @@ namespace UE::MultiUserClient
 		/** @return Whether ObjectPath is muted. */
 		bool IsMuted(const FSoftObjectPath& ObjectPath) const { return MutedObjects.Contains(ObjectPath); }
 
+		const TMap<FSoftObjectPath, FConcertReplication_ObjectMuteSetting>& GetExplicitlyMutedObjects() const { return ExplicitlyMutedObjects; }
+		const TMap<FSoftObjectPath, FConcertReplication_ObjectMuteSetting>& GetExplicitlyUnmutedObjects() const { return ExplicitlyUnmutedObjects; }
+
 		/**
 		 * Update the mute state after the local application has successfully changed mute state.
 		 * FMuteStateQueryService will eventually notify us of the change but this applies it instantaneously.
@@ -44,8 +50,12 @@ namespace UE::MultiUserClient
 
 		/** Updates us with the new mute state from the server regularily. */
 		FMuteStateQueryService& MuteQueryService;
-
-		/** All the objects that are muted. */
+		
+		/** All explicitly muted objects on the server. */
+		TMap<FSoftObjectPath, FConcertReplication_ObjectMuteSetting> ExplicitlyMutedObjects;
+		/** All explicitly unmuted objects on the server.*/
+		TMap<FSoftObjectPath, FConcertReplication_ObjectMuteSetting> ExplicitlyUnmutedObjects;
+		/** A cached, final view of muted objects that combines ExplicitlyMutedObjects and ExplicitlyUnmutedObjects. */
 		TSet<FSoftObjectPath> MutedObjects;
 
 		/** Broadcasts whent the mute state changes (either because FMuteStateQueryService received updated state or because the local application has successfully made a change request). */
diff --git a/Engine/Plugins/Developer/Concert/ConcertApp/MultiUserClient/Source/MultiUserClient/Private/Replication/Preset/PresetManager.cpp b/Engine/Plugins/Developer/Concert/ConcertApp/MultiUserClient/Source/MultiUserClient/Private/Replication/Preset/PresetManager.cpp
--- a/Engine/Plugins/Developer/Concert/ConcertApp/MultiUserClient/Source/MultiUserClient/Private/Replication/Preset/PresetManager.cpp
+++ b/Engine/Plugins/Developer/Concert/ConcertApp/MultiUserClient/Source/MultiUserClient/Private/Replication/Preset/PresetManager.cpp
@@ -7,7 +7,10 @@
 
 #include "FileHelpers.h"
 #include "IConcertSyncClient.h"
+#include "Misc/ObjectPathOuterIterator.h"
 #include "Replication/Client/ReplicationClientManager.h"
+#include "Replication/Misc/ReplicationStreamUtils.h"
+#include "Replication/Muting/MuteStateManager.h"
 #include "Widgets/ActiveSession/Replication/Client/ClientUtils.h"
 
 namespace UE::MultiUserClient
@@ -26,21 +29,20 @@ namespace UE::MultiUserClient
 				}
 			}
 		}
-		
-		static FConcertReplication_PutState_Request BuildRequest(
+
+		static void FillStreamAndAuthorityRequest(
+			FConcertReplication_PutState_Request& Request,
 			const UMultiUserReplicationSessionPreset& Preset,
 			const IConcertClientSession& Session,
-			EApplyPresetFlags Flags
+			bool bClearUnreferencedClients
 			)
 		{
-			FConcertReplication_PutState_Request Request;
-			
-			const auto AddClient = [&Preset, &Request, Flags](const FConcertSessionClientInfo& ClientSessionInfo)
+			const auto AddClient = [&Preset, &Request, bClearUnreferencedClients](const FConcertSessionClientInfo& ClientSessionInfo)
 			{
 				const UMultiUserReplicationClientContent* ClientSessionContent = Preset.GetClientContent(ClientSessionInfo.ClientInfo);
 				if (!ClientSessionContent)
 				{
-					if (EnumHasAnyFlags(Flags, EApplyPresetFlags::ClearUnreferencedClients))
+					if (bClearUnreferencedClients)
 					{
 						Request.NewStreams.Add(ClientSessionInfo.ClientEndpointId, {});
 					}
@@ -76,8 +78,64 @@ namespace UE::MultiUserClient
 			{
 				AddClient(ClientSessionInfo);
 			}
+		}
+
+		static void FillMuteStateRequest(
+			FConcertReplication_PutState_Request& Request,
+			const UMultiUserReplicationSessionPreset& Preset,
+			const IConcertClientSession& Session
+			)
+		{
+			FConcertReplication_ChangeMuteState_Request& MuteRequest = Request.MuteChange;
+			MuteRequest.Flags = EConcertReplicationMuteRequestFlags::ClearMuteState;
+			
+			const auto IsReferencedByConnectedClient = [&Request, &Session](const FSoftObjectPath& ObjectPath)
+			{
+				return Algo::AnyOf(Request.NewStreams, [&Session, &ObjectPath](const TPair<FGuid, FConcertReplicationStreamArray>& ClientContent)
+				{
+					const FGuid& EndpointId = ClientContent.Key;
+					FConcertSessionClientInfo Dummy;
+					
+					const bool bIsConnected = Session.GetSessionClientEndpointId() == EndpointId || Session.FindSessionClient(EndpointId, Dummy);
+					// Case: User muted Floor but only Floor.StaticMeshComponent0 is replicated. Hence, also look for child objects being referenced.
+					const bool bIsReferenced = ConcertSyncCore::IsObjectOrChildReferenced(ClientContent.Value.Streams, ObjectPath);
+					
+					return bIsConnected && bIsReferenced;
+				});
+			};
+
+			const FMultiUserMuteSessionContent& MuteContent = Preset.GetMuteContent();
+			for (const TPair<FSoftObjectPath, FConcertReplication_ObjectMuteSetting>& MutedObject : MuteContent.MutedObjects)
+			{
+				if (IsReferencedByConnectedClient(MutedObject.Key))
+				{
+					MuteRequest.ObjectsToMute.Add(MutedObject.Key, MutedObject.Value);
+				}
+			}
+			for (const TPair<FSoftObjectPath, FConcertReplication_ObjectMuteSetting>& UnmutedObject : MuteContent.UnmutedObjects)
+			{
+				if (IsReferencedByConnectedClient(UnmutedObject.Key))
+				{
+					MuteRequest.ObjectsToUnmute.Add(UnmutedObject.Key, UnmutedObject.Value);
+				}
+			}
+		}
+		
+		static FConcertReplication_PutState_Request BuildRequest(
+			const UMultiUserReplicationSessionPreset& Preset,
+			const IConcertClientSession& Session,
+			EApplyPresetFlags Flags
+			)
+		{
+			FConcertReplication_PutState_Request Request;
+
+			const bool bClearUnreferencedClients = EnumHasAnyFlags(Flags, EApplyPresetFlags::ClearUnreferencedClients);
+			FillStreamAndAuthorityRequest(Request, Preset, Session, bClearUnreferencedClients);
+
+			// TODO UE-219829: Once the server allows sending the mute state disconnected clients should have when they rejoin,
+			// simply send over all mute state instead of doing filtering here.
+			FillMuteStateRequest(Request, Preset, Session);
 			
-			// TODO UE-219427: When implementing mute state, you must skip muting objects that may no be referenced due to some client not being present.
 			return Request;
 		}
 
@@ -140,9 +198,14 @@ namespace UE::MultiUserClient
 		}
 	}
 	
-	FPresetManager::FPresetManager(const IConcertSyncClient& SyncClient, const FReplicationClientManager& ClientManager)
+	FPresetManager::FPresetManager(
+		const IConcertSyncClient& SyncClient,
+		const FReplicationClientManager& ClientManager,
+		const FMuteStateSynchronizer& MuteStateSynchronizer
+		)
 		: SyncClient(SyncClient)
 		, ClientManager(ClientManager)
+		, MuteStateSynchronizer(MuteStateSynchronizer)
 	{
 		IConcertClientReplicationManager& ReplicationManager = *SyncClient.GetReplicationManager();
 		ReplicationManager.OnPostRemoteEditApplied().AddRaw(this, &FPresetManager::OnPostRemoteEditApplied);
@@ -255,6 +318,10 @@ namespace UE::MultiUserClient
 
 			ClientContent_InPreset->Stream->Copy(*ClientContent_ToCopy->Stream);
 		}
+
+		Preset->SetMuteContent(
+			FMultiUserMuteSessionContent(MuteStateSynchronizer.GetExplicitlyMutedObjects(), MuteStateSynchronizer.GetExplicitlyUnmutedObjects())
+			);
 		
 		return Preset;
 	}
diff --git a/Engine/Plugins/Developer/Concert/ConcertApp/MultiUserClient/Source/MultiUserClient/Private/Replication/Preset/PresetManager.h b/Engine/Plugins/Developer/Concert/ConcertApp/MultiUserClient/Source/MultiUserClient/Private/Replication/Preset/PresetManager.h
--- a/Engine/Plugins/Developer/Concert/ConcertApp/MultiUserClient/Source/MultiUserClient/Private/Replication/Preset/PresetManager.h
+++ b/Engine/Plugins/Developer/Concert/ConcertApp/MultiUserClient/Source/MultiUserClient/Private/Replication/Preset/PresetManager.h
@@ -20,6 +20,7 @@ namespace UE::ConcertSyncClient::Replication { struct FRemoteEditEvent; }
 
 namespace UE::MultiUserClient
 {
+	class FMuteStateSynchronizer;
 	class FReplicationClientManager;
 	
 	/**
@@ -30,7 +31,11 @@ namespace UE::MultiUserClient
 	{
 	public:
 		
-		FPresetManager(const IConcertSyncClient& SyncClient UE_LIFETIMEBOUND, const FReplicationClientManager& ClientManager UE_LIFETIMEBOUND);
+		FPresetManager(
+			const IConcertSyncClient& SyncClient UE_LIFETIMEBOUND,
+			const FReplicationClientManager& ClientManager UE_LIFETIMEBOUND,
+			const FMuteStateSynchronizer& MuteStateSynchronizer UE_LIFETIMEBOUND
+			);
 		~FPresetManager();
 
 		/** @return Whether any preset is currently being applied. */
@@ -53,6 +58,8 @@ namespace UE::MultiUserClient
 		const IConcertSyncClient& SyncClient;
 		/** Used to get the clients' replication content. */
 		const FReplicationClientManager& ClientManager;
+		/** Used to get the mute state when saving. */
+		const FMuteStateSynchronizer& MuteStateSynchronizer;
 
 		/** Non-null for as long as the ReplaceSessionContentWithPreset network request takes. */
 		TSharedPtr<TPromise<FReplaceSessionContentResult>> InProgressSessionReplacementOp;
diff --git a/Engine/Plugins/Developer/Concert/ConcertApp/MultiUserClient/Source/MultiUserReplicationEditor/Internal/Assets/MultiUserReplicationSessionPreset.h b/Engine/Plugins/Developer/Concert/ConcertApp/MultiUserClient/Source/MultiUserReplicationEditor/Internal/Assets/MultiUserReplicationSessionPreset.h
--- a/Engine/Plugins/Developer/Concert/ConcertApp/MultiUserClient/Source/MultiUserReplicationEditor/Internal/Assets/MultiUserReplicationSessionPreset.h
+++ b/Engine/Plugins/Developer/Concert/ConcertApp/MultiUserClient/Source/MultiUserReplicationEditor/Internal/Assets/MultiUserReplicationSessionPreset.h
@@ -3,11 +3,35 @@
 #pragma once
 
 #include "MultiUserReplicationClientContent.h"
+#include "Replication/Messages/Muting.h"
 #include "UObject/Object.h"
 #include "MultiUserReplicationSessionPreset.generated.h"
 
 struct FConcertClientInfo;
 
+USTRUCT()
+struct FMultiUserMuteSessionContent
+{
+	GENERATED_BODY()
+
+	/** The argument to put into FConcertReplication_ChangeMuteState_Request::ObjectsToMute. */
+	UPROPERTY()
+	TMap<FSoftObjectPath, FConcertReplication_ObjectMuteSetting> MutedObjects;
+
+	/** The argument to put into FConcertReplication_ChangeMuteState_Request::ObjectsToUnmute. */
+	UPROPERTY()
+	TMap<FSoftObjectPath, FConcertReplication_ObjectMuteSetting> UnmutedObjects;
+
+	FMultiUserMuteSessionContent() = default;
+	FMultiUserMuteSessionContent(
+		TMap<FSoftObjectPath, FConcertReplication_ObjectMuteSetting> MutedObjects,
+		TMap<FSoftObjectPath, FConcertReplication_ObjectMuteSetting> UnmutedObjects
+		)
+		: MutedObjects(MoveTemp(MutedObjects))
+		, UnmutedObjects(MoveTemp(UnmutedObjects))
+	{}
+};
+
 /** Stores per-client replication settings so it can be loaded by a user to quickly set up a session. */
 UCLASS()
 class MULTIUSERREPLICATIONEDITOR_API UMultiUserReplicationSessionPreset : public UObject
@@ -15,6 +39,8 @@ class MULTIUSERREPLICATIONEDITOR_API UMultiUserReplicationSessionPreset : public
 	GENERATED_BODY()
 public:
 
+	/********** Clients **********/
+
 	/** @return The client preset that matches ClientInfo.DisplayName. If there are multiple, returns the one that matches ClientInfo.DeviceName, as well. */
 	UMultiUserReplicationClientContent* GetClientContent(const FConcertClientInfo& ClientInfo) const;
 	/** @return The client preset that matches both the display and device name. */
@@ -27,11 +53,19 @@ public:
 	
 	/** Adds a client to the preset if it's not already present. */
 	UMultiUserReplicationClientContent* AddClientIfUnique(const FConcertClientInfo& ClientInfo);
-
+	
 	const TArray<TObjectPtr<UMultiUserReplicationClientContent>>& GetClientPresets() const { return ClientPresets; }
-
+	
+	/********** Muting **********/
+	
+	const FMultiUserMuteSessionContent& GetMuteContent() const { return MuteContent; }
+	void SetMuteContent(FMultiUserMuteSessionContent Content) { MuteContent = MoveTemp(Content); }
+	
 private:
 	
 	UPROPERTY(Instanced)
 	TArray<TObjectPtr<UMultiUserReplicationClientContent>> ClientPresets;
+
+	UPROPERTY()
+	FMultiUserMuteSessionContent MuteContent;
 };
diff --git a/Engine/Plugins/Developer/Concert/ConcertSync/ConcertSyncCore/Source/ConcertSyncCore/Public/Replication/Misc/ReplicationStreamUtils.h b/Engine/Plugins/Developer/Concert/ConcertSync/ConcertSyncCore/Source/ConcertSyncCore/Public/Replication/Misc/ReplicationStreamUtils.h
--- a/Engine/Plugins/Developer/Concert/ConcertSync/ConcertSyncCore/Source/ConcertSyncCore/Public/Replication/Misc/ReplicationStreamUtils.h
+++ b/Engine/Plugins/Developer/Concert/ConcertSync/ConcertSyncCore/Source/ConcertSyncCore/Public/Replication/Misc/ReplicationStreamUtils.h
@@ -4,12 +4,15 @@
 
 #include "Replication/Data/ObjectIds.h"
 #include "Replication/Data/ReplicationStream.h"
+#include "Misc/ObjectPathOuterIterator.h"
 
 #include "Containers/Array.h"
 #include "Containers/ArrayView.h"
 #include "Containers/ContainersFwd.h"
 #include "UObject/SoftObjectPath.h"
 
+#include <type_traits>
+
 namespace UE::ConcertSyncCore
 {
 	/** Util for finding a stream in an array of streams by id. */
@@ -31,6 +34,9 @@ namespace UE::ConcertSyncCore
 	const FConcertObjectReplicationSettings* FindObjectFrequency(const TConstArrayView<FConcertReplicationStream> Streams, const FConcertObjectInStreamID& ObjectId);
 	/** Util for finding an object's frequency override settings in a stream in an array of streams. */
 	FConcertObjectReplicationSettings* FindObjectFrequencyEditable(const TArrayView<FConcertReplicationStream> Streams, const FConcertObjectInStreamID& ObjectId);
+	
+	/** @return Whether ObjectPath or any its child objects, i.e. any object that ObjectPath is an indirect or direct outer of, are referenced by Streams. */
+	bool IsObjectOrChildReferenced(const TConstArrayView<FConcertReplicationStream> Streams, const FSoftObjectPath& ObjectPath);
 }
 
 namespace UE::ConcertSyncCore
@@ -74,4 +80,22 @@ namespace UE::ConcertSyncCore
 		FConcertReplicationStream* Stream = FindStreamEditable(Streams, ObjectId.StreamId);
 		return Stream ? Stream->BaseDescription.FrequencySettings.ObjectOverrides.Find(ObjectId.Object) : nullptr;
 	}
+
+	inline bool IsObjectOrChildReferenced(const TConstArrayView<FConcertReplicationStream> Streams, const FSoftObjectPath& ObjectPath)
+	{
+		for (const FConcertReplicationStream& Stream : Streams)
+		{
+			for (const TPair<FSoftObjectPath, FConcertReplicatedObjectInfo>& Pair : Stream.BaseDescription.ReplicationMap.ReplicatedObjects)
+			{
+				for (FObjectPathOuterIterator It(Pair.Key); It; ++It)
+				{
+					if (*It == ObjectPath)
+					{
+						return true;
+					}
+				}
+			}
+		}
+		return false;
+	}
 }
