Save & load frequency settings in replication presets.

diff --git a/Engine/Plugins/Developer/Concert/ConcertApp/MultiUserClient/Source/MultiUserClient/Private/Replication/Misc/Frequency/FrequencyChangeTracker.cpp b/Engine/Plugins/Developer/Concert/ConcertApp/MultiUserClient/Source/MultiUserClient/Private/Replication/Misc/Frequency/FrequencyChangeTracker.cpp
--- a/Engine/Plugins/Developer/Concert/ConcertApp/MultiUserClient/Source/MultiUserClient/Private/Replication/Misc/Frequency/FrequencyChangeTracker.cpp
+++ b/Engine/Plugins/Developer/Concert/ConcertApp/MultiUserClient/Source/MultiUserClient/Private/Replication/Misc/Frequency/FrequencyChangeTracker.cpp
@@ -23,11 +23,13 @@ namespace UE::MultiUserClient
 
 	void FFrequencyChangeTracker::AddOverride(FSoftObjectPath Object, FConcertObjectReplicationSettings NewSettings)
 	{
+		// TODO UE-219834: The override should be added to UMultiUserReplicationStream so it can be transacted.
 		RecordedChanges.OverridesToAdd.Emplace(MoveTemp(Object), MoveTemp(NewSettings));
 	}
 	
 	FFrequencyChangelist FFrequencyChangeTracker::BuildForSubmission(const FStreamChangelist& ObjectChanges)
 	{
+		// TODO UE-219834: A diff between UMultiUserReplicationStream and the server state should be generated instead
 		FFrequencyChangelist Result = RecordedChanges;
 
 		// For now we only support adding overrides.
diff --git a/Engine/Plugins/Developer/Concert/ConcertApp/MultiUserClient/Source/MultiUserClient/Private/Replication/Preset/PresetManager.cpp b/Engine/Plugins/Developer/Concert/ConcertApp/MultiUserClient/Source/MultiUserClient/Private/Replication/Preset/PresetManager.cpp
--- a/Engine/Plugins/Developer/Concert/ConcertApp/MultiUserClient/Source/MultiUserClient/Private/Replication/Preset/PresetManager.cpp
+++ b/Engine/Plugins/Developer/Concert/ConcertApp/MultiUserClient/Source/MultiUserClient/Private/Replication/Preset/PresetManager.cpp
@@ -60,18 +60,17 @@ namespace UE::MultiUserClient
 				const FGuid& StreamId = SavedStream->StreamId;
 				FConcertReplicationStream Stream { { .Identifier = StreamId, .ReplicationMap = SavedStream->ReplicationMap } };
 				// Empty objects will be rejected by the server
-				RemoveEmptyObjectsFromRequest(Stream); 
+				RemoveEmptyObjectsFromRequest(Stream);
+				Stream.BaseDescription.FrequencySettings = ClientSessionContent->Stream->FrequencySettings;
 				Request.NewStreams.Add(ClientSessionInfo.ClientEndpointId, { TArray{ Stream } });
 				
 				// MU automatically requests authority when it adds an object.
-				// We'll assume that that authority was granted - if it actually was not, our request may fail due to overlapping authority.
+				// We'll assume that that authority was granted when the preset was created - if it actually was not, our request may fail due to overlapping authority.
 				TArray<FConcertObjectInStreamID>& OwnedObjects = Request.NewAuthorityState.Add(ClientSessionInfo.ClientEndpointId).Objects;
 				Algo::Transform(Stream.BaseDescription.ReplicationMap.ReplicatedObjects, OwnedObjects, [&StreamId](const TPair<FSoftObjectPath, FConcertReplicatedObjectInfo>& Pair)
 				{
 					return FConcertObjectInStreamID{ StreamId, Pair.Key };
 				});
-				
-				// TODO UE-219639: Load frequency settings
 			};
 			AddClient({ Session.GetSessionClientEndpointId(), Session.GetLocalClientInfo() });
 			for (const FConcertSessionClientInfo& ClientSessionInfo : Session.GetSessionClients())
@@ -317,6 +316,8 @@ namespace UE::MultiUserClient
 			}
 
 			ClientContent_InPreset->Stream->Copy(*ClientContent_ToCopy->Stream);
+			// TODO UE-219834: Once UMultiUserReplicationStream::FrequencySettings reflect the server state, this can be removed.
+			ClientContent_InPreset->Stream->FrequencySettings = Client->GetStreamSynchronizer().GetFrequencySettings();
 		}
 
 		Preset->SetMuteContent(
diff --git a/Engine/Plugins/Developer/Concert/ConcertApp/MultiUserClient/Source/MultiUserReplicationEditor/Internal/Assets/MultiUserReplicationStream.h b/Engine/Plugins/Developer/Concert/ConcertApp/MultiUserClient/Source/MultiUserReplicationEditor/Internal/Assets/MultiUserReplicationStream.h
--- a/Engine/Plugins/Developer/Concert/ConcertApp/MultiUserClient/Source/MultiUserReplicationEditor/Internal/Assets/MultiUserReplicationStream.h
+++ b/Engine/Plugins/Developer/Concert/ConcertApp/MultiUserClient/Source/MultiUserReplicationEditor/Internal/Assets/MultiUserReplicationStream.h
@@ -7,7 +7,7 @@
 #include "Misc/Attribute.h"
 #include "MultiUserReplicationStream.generated.h"
 
-/** Wraps FConcertObjectReplicationMap so its edition can be transacted in the editor. */
+/** Wraps FConcertObjectReplicationMap so its edition can be transacted in the editor and saved in presets. */
 UCLASS()
 class MULTIUSERREPLICATIONEDITOR_API UMultiUserReplicationStream : public UObject
 {
@@ -21,6 +21,16 @@ public:
 	/** The objects this stream will modify. */
 	UPROPERTY()
 	FConcertObjectReplicationMap ReplicationMap;
+	
+	/**
+	 * The frequency setting the stream has.
+	 *
+	 * TODO UE-219834:
+	 * Currently, this is only written to by the preset system.
+	 * In the future, FFrequencyChangeTracker could / should use this so changes to frequencies can be transacted, as well.
+	 */
+	UPROPERTY()
+	FConcertStreamFrequencySettings FrequencySettings;
 
 	UMultiUserReplicationStream();
 
diff --git a/Engine/Plugins/Developer/Concert/ConcertApp/MultiUserClient/Source/MultiUserReplicationEditor/Private/Assets/MultiUserReplicationStream.cpp b/Engine/Plugins/Developer/Concert/ConcertApp/MultiUserClient/Source/MultiUserReplicationEditor/Private/Assets/MultiUserReplicationStream.cpp
--- a/Engine/Plugins/Developer/Concert/ConcertApp/MultiUserClient/Source/MultiUserReplicationEditor/Private/Assets/MultiUserReplicationStream.cpp
+++ b/Engine/Plugins/Developer/Concert/ConcertApp/MultiUserClient/Source/MultiUserReplicationEditor/Private/Assets/MultiUserReplicationStream.cpp
@@ -14,11 +14,12 @@ UMultiUserReplicationStream::UMultiUserReplicationStream()
 
 FConcertReplicationStream UMultiUserReplicationStream::GenerateDescription() const
 {
-	return { StreamId, ReplicationMap };
+	return { StreamId, ReplicationMap, FrequencySettings };
 }
 
 void UMultiUserReplicationStream::Copy(UMultiUserReplicationStream& OtherStream)
 {
 	StreamId = OtherStream.StreamId;
 	ReplicationMap = OtherStream.ReplicationMap;
+	FrequencySettings = OtherStream.FrequencySettings;
 }
