Replication UI now caches the registered object hierarchy of clients so it can effectively query registered object hierarchy

diff --git a/Engine/Plugins/Developer/Concert/ConcertApp/MultiUserClient/Source/MultiUserClient/Private/Replication/Misc/GlobalAuthorityCache.cpp b/Engine/Plugins/Developer/Concert/ConcertApp/MultiUserClient/Source/MultiUserClient/Private/Replication/Misc/GlobalAuthorityCache.cpp
--- a/Engine/Plugins/Developer/Concert/ConcertApp/MultiUserClient/Source/MultiUserClient/Private/Replication/Misc/GlobalAuthorityCache.cpp
+++ b/Engine/Plugins/Developer/Concert/ConcertApp/MultiUserClient/Source/MultiUserClient/Private/Replication/Misc/GlobalAuthorityCache.cpp
@@ -23,7 +23,7 @@ namespace UE::MultiUserClient
 
 	void FGlobalAuthorityCache::ForEachClientWithObjectInStream(const FSoftObjectPath& Object, TFunctionRef<EBreakBehavior(const FGuid& ClientId)> Callback) const
 	{
-		const TSet<FGuid>* Clients = RegisteredObjectsToClients.Find(Object);
+		const TSet<FGuid>* Clients = StreamObjectsToClients.Find(Object);
 		if (!Clients)
 		{
 			return;
@@ -41,7 +41,7 @@ namespace UE::MultiUserClient
 	bool FGlobalAuthorityCache::IsObjectOrChildReferenced(const FSoftObjectPath& Object) const
 	{
 		const FString ObjectPathString = Object.ToString();
-		for (const TPair<FSoftObjectPath, TSet<FGuid>>& Pair : RegisteredObjectsToClients)
+		for (const TPair<FSoftObjectPath, TSet<FGuid>>& Pair : StreamObjectsToClients)
 		{
 			if (Pair.Key.ToString().Contains(ObjectPathString))
 			{
@@ -53,7 +53,7 @@ namespace UE::MultiUserClient
 
 	void FGlobalAuthorityCache::ForEachClientWithAuthorityOverObject(const FSoftObjectPath& Object, TFunctionRef<EBreakBehavior(const FGuid& ClientId)> Callback) const
 	{
-		const TSet<FGuid>* Clients = OwnedObjectsToClients.Find(Object);
+		const TSet<FGuid>* Clients = AuthorityObjectsToClients.Find(Object);
 		if (!Clients)
 		{
 			return;
@@ -139,7 +139,7 @@ namespace UE::MultiUserClient
 
 	void FGlobalAuthorityCache::ForEachClientReferencingProperty(const FSoftObjectPath& Object, const FConcertPropertyChain& Property, TFunctionRef<EBreakBehavior(const FGuid& ClientId)> Callback) const
 	{
-		const TSet<FGuid>* ClientsWithStreams = RegisteredObjectsToClients.Find(Object);
+		const TSet<FGuid>* ClientsWithStreams = StreamObjectsToClients.Find(Object);
 		if (!ClientsWithStreams)
 		{
 			return;
@@ -237,17 +237,19 @@ namespace UE::MultiUserClient
 		for (const TPair<FSoftObjectPath, FConcertReplicatedObjectInfo>& StreamContents : ClientObjectMap.ReplicatedObjects)
 		{
 			const FSoftObjectPath& Object = StreamContents.Key;
-			RegisteredObjectsToClients.FindOrAdd(Object).Add(ClientId);
+			StreamObjectsToClients.FindOrAdd(Object).Add(ClientId);
+			StreamObjectHierarchy.AddObject(Object);
+			
 			if (ClientAuthority.HasAuthorityOver(Object))
 			{
-				OwnedObjectsToClients.FindOrAdd(Object).Add(ClientId);
+				AuthorityObjectsToClients.FindOrAdd(Object).Add(ClientId);
 			}
 		}
 	}
 
 	void FGlobalAuthorityCache::RemoveClient(const FGuid& ClientId)
 	{
-		for (auto It = OwnedObjectsToClients.CreateIterator(); It; ++It)
+		for (auto It = AuthorityObjectsToClients.CreateIterator(); It; ++It)
 		{
 			It->Value.Remove(ClientId);
 			if (It->Value.IsEmpty())
@@ -256,11 +258,12 @@ namespace UE::MultiUserClient
 			}
 		}
 
-		for (auto It = RegisteredObjectsToClients.CreateIterator(); It; ++It)
+		for (auto It = StreamObjectsToClients.CreateIterator(); It; ++It)
 		{
 			It->Value.Remove(ClientId);
 			if (It->Value.IsEmpty())
 			{
+				StreamObjectHierarchy.RemoveObject(It->Key);
 				It.RemoveCurrent();
 			}
 		}
diff --git a/Engine/Plugins/Developer/Concert/ConcertApp/MultiUserClient/Source/MultiUserClient/Private/Replication/Misc/GlobalAuthorityCache.h b/Engine/Plugins/Developer/Concert/ConcertApp/MultiUserClient/Source/MultiUserClient/Private/Replication/Misc/GlobalAuthorityCache.h
--- a/Engine/Plugins/Developer/Concert/ConcertApp/MultiUserClient/Source/MultiUserClient/Private/Replication/Misc/GlobalAuthorityCache.h
+++ b/Engine/Plugins/Developer/Concert/ConcertApp/MultiUserClient/Source/MultiUserClient/Private/Replication/Misc/GlobalAuthorityCache.h
@@ -7,6 +7,7 @@
 #include "Replication/AuthorityConflictSharedUtils.h"
 
 #include "Delegates/Delegate.h"
+#include "Misc/ObjectPathHierarchy.h"
 #include "Templates/Function.h"
 
 struct FConcertPropertyChain;
@@ -78,6 +79,9 @@ namespace UE::MultiUserClient
 		 * @param SendingClient The client that will send the request
 		 */
 		void CleanseConflictsFromStreamRequest(FConcertReplication_ChangeStream_Request& Request, const FGuid& SendingClient) const;
+
+		/** Allows efficient retrieving of registered child objects. */
+		const ConcertSyncCore::FObjectPathHierarchy& GetStreamObjectHierarchy() const { return StreamObjectHierarchy; }
 		
 		DECLARE_MULTICAST_DELEGATE_OneParam(FOnCacheChanged, const FGuid& ClientId);
 		/** Called when the cache changes for a specific client. */
@@ -88,10 +92,12 @@ namespace UE::MultiUserClient
 		/** Used to obtain the clients and their states */
 		FReplicationClientManager& ClientManager;
 		
-		/** Maps objects that are owned to the clients that own them (have authority) */
-		TMap<FSoftObjectPath, TSet<FGuid>> OwnedObjectsToClients;
 		/** Maps objects that are owned to the clients that have them in the stream */
-		TMap<FSoftObjectPath, TSet<FGuid>> RegisteredObjectsToClients;
+		TMap<FSoftObjectPath, TSet<FGuid>> StreamObjectsToClients;
+		/** Maps objects that are owned to the clients that own them (have authority) */
+		TMap<FSoftObjectPath, TSet<FGuid>> AuthorityObjectsToClients;
+		/** The hierarchy of objects that have been registered by all clients. Allows efficient retrieving of child objects. */
+		ConcertSyncCore::FObjectPathHierarchy StreamObjectHierarchy;
 		
 		/** Called when the cache changes for a specific client. */
 		FOnCacheChanged OnCacheChangedDelegate;
