Change right-click context menu of replicated actors. The menu now creates a sub-menu for each subobject in which the user can change the frequency settings.

diff --git a/Engine/Plugins/Developer/Concert/ConcertApp/MultiUserClient/Source/MultiUserClient/Private/Widgets/ActiveSession/Replication/Client/Context/ContextMenuUtils_Frequency.cpp b/Engine/Plugins/Developer/Concert/ConcertApp/MultiUserClient/Source/MultiUserClient/Private/Widgets/ActiveSession/Replication/Client/Context/ContextMenuUtils_Frequency.cpp
--- a/Engine/Plugins/Developer/Concert/ConcertApp/MultiUserClient/Source/MultiUserClient/Private/Widgets/ActiveSession/Replication/Client/Context/ContextMenuUtils_Frequency.cpp
+++ b/Engine/Plugins/Developer/Concert/ConcertApp/MultiUserClient/Source/MultiUserClient/Private/Widgets/ActiveSession/Replication/Client/Context/ContextMenuUtils_Frequency.cpp
@@ -11,6 +11,7 @@
 #include "Async/TaskGraphInterfaces.h"
 #include "Framework/Commands/UIAction.h"
 #include "Framework/MultiBox/MultiBoxBuilder.h"
+#include "Replication/ObjectNameUtils.h"
 #include "Textures/SlateIcon.h"
 #include "Widgets/Images/SThrobber.h"
 #include "Widgets/Input/SNumericEntryBox.h"
@@ -191,7 +192,7 @@ namespace UE::MultiUserClient::ContextMenuUtils
 						SettingToOverride.ReplicationMode = ModeToSet;
 					});
 			};
-			
+
 			MenuBuilder.AddMenuEntry(
 				LOCTEXT("ListView_Realtime.Label", "Realtime"),
 				GetRealtimeTooltip(),
@@ -237,23 +238,19 @@ namespace UE::MultiUserClient::ContextMenuUtils
 				);
 		}
 
-		static void SharedAppendFrequencyToMenu(FMenuBuilder& MenuBuilder, const FSoftObjectPath& ContextObject, TAttribute<FInlineClientArray> GetClientsAttribute, FReplicationClientManager& InClientManager)
+		static void AppendFrequencyEditingToMenu(
+			FMenuBuilder& MenuBuilder,
+			const FSoftObjectPath& ContextObject,
+			TAttribute<FInlineClientArray> GetClientsAttribute,
+			FReplicationClientManager& InClientManager
+			)
 		{
-			// No frequency section if the object is not registered by any clients
-			const bool bHasObjectRegistered = Algo::AnyOf(GetClientsAttribute.Get(), [&ContextObject, &InClientManager](const FGuid& ClientId)
-			{
-				const FReplicationClient* Client = InClientManager.FindClient(ClientId);
-				return !Client || Client->GetStreamSynchronizer().GetServerState().HasProperties(ContextObject);
-			});
-			if (!bHasObjectRegistered)
-			{
-				return;
-			}
-			
-			MenuBuilder.BeginSection(NAME_None, LOCTEXT("Frequency", "Frequency"));
-			
 			// Change to Realtime / Specified Rate
+			MenuBuilder.BeginSection(NAME_None, LOCTEXT("Frequency", "Mode"));
 			AddReplicationModeSubMenu(MenuBuilder, ContextObject, GetClientsAttribute, InClientManager);
+			MenuBuilder.EndSection();
+
+			MenuBuilder.AddSeparator();
 
 			// Replication rate, only enabled if Specified Rate
 			constexpr bool bNoIndent = false;
@@ -267,35 +264,79 @@ namespace UE::MultiUserClient::ContextMenuUtils
 				bSearchable,
 				GetEditBoxInstructionsTooltip()
 				);
+		}
+		
+		static void AddSubMenuForEachObjectInHierarchy(
+			FMenuBuilder& MenuBuilder,
+			const FSoftObjectPath& ContextObject,
+			TFunction<FInlineClientArray(const FSoftObjectPath& Object)> GetClientsForObjectFunc,
+			FReplicationClientManager& InClientManager
+			)
+		{
+			bool bAddedSubmenu = false;
+			const auto OnAddSubmenu = [&MenuBuilder, &bAddedSubmenu]()
+			{
+				if (!bAddedSubmenu)
+				{
+					bAddedSubmenu = true;
+					MenuBuilder.BeginSection(NAME_None, LOCTEXT("Frequency", "Change Frequency"));
+				}
+			};
+			ON_SCOPE_EXIT{ if(bAddedSubmenu){ MenuBuilder.EndSection(); } };
+			
+			const auto MakeSubmenu = [&MenuBuilder, &GetClientsForObjectFunc, &InClientManager, &OnAddSubmenu](const FSoftObjectPath& Object)
+			{
+				OnAddSubmenu();
 
-			MenuBuilder.EndSection();
+				constexpr bool bOpenOnClick = false;
+				constexpr bool bCloseWindowAfterMenuSelection = false;
+				MenuBuilder.AddSubMenu(
+					ConcertClientSharedSlate::GetObjectDisplayName(TSoftObjectPtr(Object)),
+					FText::GetEmpty(),
+					FNewMenuDelegate::CreateLambda([Object, GetClientsForObjectFunc, &InClientManager](FMenuBuilder& MenuBuilder)
+					{
+						TAttribute<FInlineClientArray> ClientAttribute = TAttribute<FInlineClientArray>::CreateLambda([Object, GetClientsForObjectFunc]()
+						{
+							return GetClientsForObjectFunc(Object);
+						});
+						AppendFrequencyEditingToMenu(MenuBuilder, Object, MoveTemp(ClientAttribute), InClientManager);
+					}), bOpenOnClick, FSlateIcon(), bCloseWindowAfterMenuSelection);
+			};
+
+			if (!GetClientsForObjectFunc(ContextObject).IsEmpty())
+			{
+				MakeSubmenu(ContextObject);
+			}
+			
+			using namespace ConcertSyncCore;
+			const FObjectPathHierarchy& Hierarchy = InClientManager.GetAuthorityCache().GetStreamObjectHierarchy();
+			Hierarchy.TraverseTopToBottom([&MakeSubmenu](const FChildRelation& Relation)
+			{
+				MakeSubmenu(Relation.Child.Object);
+				return ETreeTraversalBehavior::Continue;
+			}, ContextObject);
 		}
 	}
-	
-	void AddFrequencyOptionsForSingleClient(FMenuBuilder& MenuBuilder, const FSoftObjectPath& ContextObject, const FGuid& ClientId, FReplicationClientManager& InClientManager)
-	{
-		Private::SharedAppendFrequencyToMenu(MenuBuilder, ContextObject, TArray{ ClientId }, InClientManager);
-	}
 
 	void AddFrequencyOptionsForMultipleClients(FMenuBuilder& MenuBuilder, const FSoftObjectPath& ContextObject, FReplicationClientManager& InClientManager)
 	{
 		using namespace Private;
-		TAttribute<FInlineClientArray> GetClientsAttribute = TAttribute<FInlineClientArray>::CreateLambda([ContextObject, &InClientManager]()
+		const auto GetClientsForObject = [&InClientManager](const FSoftObjectPath& Object)
 		{
 			FInlineClientArray Result;
-			InClientManager.ForEachClient([&ContextObject, &Result](const FReplicationClient& Client)
+			InClientManager.ForEachClient([&Object, &Result](const FReplicationClient& Client)
 			{
 				// Only clients that have the object registered should be considered
-				if (Client.GetStreamSynchronizer().GetServerState().HasProperties(ContextObject))
+				if (Client.GetStreamSynchronizer().GetServerState().HasProperties(Object))
 				{
 					Result.Add(Client.GetEndpointId());
 				}
 				return EBreakBehavior::Continue;
 			});
 			return Result;
-		});
+		};
 		
-		SharedAppendFrequencyToMenu(MenuBuilder, ContextObject, MoveTemp(GetClientsAttribute), InClientManager);
+		AddSubMenuForEachObjectInHierarchy(MenuBuilder, ContextObject, GetClientsForObject, InClientManager);
 	}
 }
 
