Make FReplicationClient use UMultiUserReplicationStream instead of UMultiUserReplicationClientContent.
This separates the replication system to use UMultiUserReplicationStream while the presets can use UMultiUserReplicationClientContent, giving each class a single responsibility.
This is in preparation for a change to what is saved in the presets.

diff --git a/Engine/Plugins/Developer/Concert/ConcertApp/MultiUserClient/Source/MultiUserClient/Private/Replication/Client/LocalReplicationClient.cpp b/Engine/Plugins/Developer/Concert/ConcertApp/MultiUserClient/Source/MultiUserClient/Private/Replication/Client/LocalReplicationClient.cpp
--- a/Engine/Plugins/Developer/Concert/ConcertApp/MultiUserClient/Source/MultiUserClient/Private/Replication/Client/LocalReplicationClient.cpp
+++ b/Engine/Plugins/Developer/Concert/ConcertApp/MultiUserClient/Source/MultiUserClient/Private/Replication/Client/LocalReplicationClient.cpp
@@ -11,7 +11,7 @@ namespace UE::MultiUserClient
 	FLocalReplicationClient::FLocalReplicationClient(
 		FReplicationDiscoveryContainer& InDiscoveryContainer,
 		FGlobalAuthorityCache& InAuthorityCache,
-		UMultiUserReplicationClientContent& InSessionContent,
+		UMultiUserReplicationStream& InClientStreamContent,
 		TUniquePtr<IClientStreamSynchronizer> InStreamSynchronizer,
 		TSharedRef<IConcertSyncClient> InClient
 		)
@@ -19,7 +19,7 @@ namespace UE::MultiUserClient
 			InClient->GetConcertClient()->GetCurrentSession()->GetSessionClientEndpointId(),
 			InDiscoveryContainer,
 			InAuthorityCache,
-			InSessionContent,
+			InClientStreamContent,
 			MoveTemp(InStreamSynchronizer),
 			MakeUnique<FAuthoritySynchronizer_LocalClient>(InClient),
 			MakeUnique<FSubmissionWorkflow_LocalClient>(MoveTemp(InClient)))
diff --git a/Engine/Plugins/Developer/Concert/ConcertApp/MultiUserClient/Source/MultiUserClient/Private/Replication/Client/LocalReplicationClient.h b/Engine/Plugins/Developer/Concert/ConcertApp/MultiUserClient/Source/MultiUserClient/Private/Replication/Client/LocalReplicationClient.h
--- a/Engine/Plugins/Developer/Concert/ConcertApp/MultiUserClient/Source/MultiUserClient/Private/Replication/Client/LocalReplicationClient.h
+++ b/Engine/Plugins/Developer/Concert/ConcertApp/MultiUserClient/Source/MultiUserClient/Private/Replication/Client/LocalReplicationClient.h
@@ -17,14 +17,14 @@ namespace UE::MultiUserClient
 		/**
 		 * @param InDiscoveryContainer Used for auto-discovering properties added to this client's stream. The caller ensures it outlives the constructed instance.
 		 * @param InAuthorityCache Caches authority state of all clients. Passed to subsystems. The caller ensures it outlives the constructed instance.
-		 * @param InSessionContent Object that this client's stream changes are to be written into. The caller ensures it outlives the constructed instance.
+		 * @param InClientStreamContent Object that this client's stream changes are to be written into. The caller ensures it outlives the constructed instance.
 		 * @param InStreamSynchronizer Implementation for obtaining stream registered on server. The constructed instance takes ownership.
 		 * @param InClient The local editor's running client. Passed to subsystems so they can send messages to the current session. Keeps strong reference.
 		 */
 		FLocalReplicationClient(
 			FReplicationDiscoveryContainer& InDiscoveryContainer,
 			FGlobalAuthorityCache& InAuthorityCache,
-			UMultiUserReplicationClientContent& InSessionContent,
+			UMultiUserReplicationStream& InClientStreamContent,
 			TUniquePtr<IClientStreamSynchronizer> InStreamSynchronizer,
 			TSharedRef<IConcertSyncClient> InClient
 			);
diff --git a/Engine/Plugins/Developer/Concert/ConcertApp/MultiUserClient/Source/MultiUserClient/Private/Replication/Client/RemoteReplicationClient.cpp b/Engine/Plugins/Developer/Concert/ConcertApp/MultiUserClient/Source/MultiUserClient/Private/Replication/Client/RemoteReplicationClient.cpp
--- a/Engine/Plugins/Developer/Concert/ConcertApp/MultiUserClient/Source/MultiUserClient/Private/Replication/Client/RemoteReplicationClient.cpp
+++ b/Engine/Plugins/Developer/Concert/ConcertApp/MultiUserClient/Source/MultiUserClient/Private/Replication/Client/RemoteReplicationClient.cpp
@@ -15,14 +15,14 @@ namespace UE::MultiUserClient
 		FReplicationDiscoveryContainer& InDiscoveryContainer,
 		TSharedRef<IConcertClient> InClient,
 		FGlobalAuthorityCache& InAuthorityCache,
-		UMultiUserReplicationClientContent& InSessionContent,
+		UMultiUserReplicationStream& InClientStreamContent,
 		FStreamAndAuthorityQueryService& QueryService
 		)
 		: FReplicationClient(
 			InConcertClientId,
 			InDiscoveryContainer,
 			InAuthorityCache,
-			InSessionContent,
+			InClientStreamContent,
 			MakeUnique<FStreamSynchronizer_RemoteClient>(InConcertClientId, QueryService),
 			MakeUnique<FAuthoritySynchronizer_RemoteClient>(InConcertClientId,QueryService),
 			MakeUnique<FSubmissionWorkflow_RemoteClient>(InClient->GetCurrentSession().ToSharedRef(), InConcertClientId))
diff --git a/Engine/Plugins/Developer/Concert/ConcertApp/MultiUserClient/Source/MultiUserClient/Private/Replication/Client/RemoteReplicationClient.h b/Engine/Plugins/Developer/Concert/ConcertApp/MultiUserClient/Source/MultiUserClient/Private/Replication/Client/RemoteReplicationClient.h
--- a/Engine/Plugins/Developer/Concert/ConcertApp/MultiUserClient/Source/MultiUserClient/Private/Replication/Client/RemoteReplicationClient.h
+++ b/Engine/Plugins/Developer/Concert/ConcertApp/MultiUserClient/Source/MultiUserClient/Private/Replication/Client/RemoteReplicationClient.h
@@ -22,7 +22,7 @@ namespace UE::MultiUserClient
 		 * @param InDiscoveryContainer Used for auto-discovering properties added to this client's stream. The caller ensures it outlives the constructed instance.
 		 * @param InClient The local editor's running client. Passed to subsystems so they can send messages to the current session. Keeps strong reference.
 		 * @param InAuthorityCache Caches authority state of all clients. Passed to subsystems. The caller ensures it outlives the constructed instance.
-		 * @param InSessionContent Object that this client's stream changes are to be written into. The caller ensures it outlives the constructed instance.
+		 * @param InClientStreamContent Object that this client's stream changes are to be written into. The caller ensures it outlives the constructed instance.
 		 * @param QueryService Service shared across remote clients which periodically asks the server about this client's state. The caller ensures it outlives the constructed instance.
 		 */
 		FRemoteReplicationClient(
@@ -30,7 +30,7 @@ namespace UE::MultiUserClient
 			FReplicationDiscoveryContainer& InDiscoveryContainer UE_LIFETIMEBOUND,
 			TSharedRef<IConcertClient> InClient,
 			FGlobalAuthorityCache& InAuthorityCache UE_LIFETIMEBOUND,
-			UMultiUserReplicationClientContent& InSessionContent UE_LIFETIMEBOUND,
+			UMultiUserReplicationStream& InClientStreamContent UE_LIFETIMEBOUND,
 			FStreamAndAuthorityQueryService& QueryService UE_LIFETIMEBOUND
 			);
 	};
diff --git a/Engine/Plugins/Developer/Concert/ConcertApp/MultiUserClient/Source/MultiUserClient/Private/Replication/Client/ReplicationClient.cpp b/Engine/Plugins/Developer/Concert/ConcertApp/MultiUserClient/Source/MultiUserClient/Private/Replication/Client/ReplicationClient.cpp
--- a/Engine/Plugins/Developer/Concert/ConcertApp/MultiUserClient/Source/MultiUserClient/Private/Replication/Client/ReplicationClient.cpp
+++ b/Engine/Plugins/Developer/Concert/ConcertApp/MultiUserClient/Source/MultiUserClient/Private/Replication/Client/ReplicationClient.cpp
@@ -19,13 +19,13 @@ namespace UE::MultiUserClient
 		const FGuid& EndpointId,
 		FReplicationDiscoveryContainer& InDiscoveryContainer,
 		FGlobalAuthorityCache& InAuthorityCache,
-		UMultiUserReplicationClientContent& InSessionContent,
+		UMultiUserReplicationStream& InClientStreamContent,
 		TUniquePtr<IClientStreamSynchronizer> InStreamSynchronizer,
 		TUniquePtr<IClientAuthoritySynchronizer> InAuthoritySynchronizer,
 		TUniquePtr<ISubmissionWorkflow> InSubmissionWorkflow
 		)
 		: EndpointId(EndpointId)
-		, ClientContentStorage(&InSessionContent)
+		, ClientStreamContent(&InClientStreamContent)
 		, StreamSynchronizer(MoveTemp(InStreamSynchronizer))
 		, AuthoritySynchronizer(MoveTemp(InAuthoritySynchronizer))
 		, SubmissionWorkflow(MoveTemp(InSubmissionWorkflow))
@@ -40,14 +40,14 @@ namespace UE::MultiUserClient
 			// Transact ClientContentStorage
 			ConcertClientSharedSlate::CreateTransactionalStreamModel(
 				// Read & write the stream data in ClientContentStorage
-				CreateBaseStreamModel(ClientContentStorage->Stream->MakeReplicationMapGetterAttribute(), StreamExtender),
-				*ClientContentStorage->Stream
+				CreateBaseStreamModel(GetClientStreamObject()->MakeReplicationMapGetterAttribute(), StreamExtender),
+				*GetClientStreamObject()
 				)
 			)
 		, LocalClientStreamDiffer(
 			GetStreamSynchronizer(),
-			ClientContentStorage->Stream->MakeReplicationMapGetterAttribute(),
-			FStreamChangeTracker::FOnModifyReplicationMap::CreateLambda([this](){ ClientContentStorage->Stream->Modify(); })
+			GetClientStreamObject()->MakeReplicationMapGetterAttribute(),
+			FStreamChangeTracker::FOnModifyReplicationMap::CreateLambda([this](){ GetClientStreamObject()->Modify(); })
 			)
 		, LocalAuthorityDiffer(EndpointId, *AuthoritySynchronizer, InAuthorityCache)
 		, LocalFrequencyChangeTracker(*StreamSynchronizer)
@@ -98,7 +98,7 @@ namespace UE::MultiUserClient
 
 		// The UI adds empty actors. However, we never send them to the server...
 		TSet<FSoftObjectPath> StagedObjects;
-		for (const TPair<FSoftObjectPath, FConcertReplicatedObjectInfo>& Pair : GetClientContent()->Stream->ReplicationMap.ReplicatedObjects)
+		for (const TPair<FSoftObjectPath, FConcertReplicatedObjectInfo>& Pair : GetClientStreamObject()->ReplicationMap.ReplicatedObjects)
 		{
 			if (Pair.Value.PropertySelection.ReplicatedProperties.IsEmpty())
 			{
@@ -107,13 +107,13 @@ namespace UE::MultiUserClient
 		}
 		
 		// ... if the user removes the last property from the entire actor-component hierarchy, we want the hierarchy to continue to displayed...
-		GetClientContent()->Stream->ReplicationMap = GetStreamSynchronizer().GetServerState();
+		GetClientStreamObject()->ReplicationMap = GetStreamSynchronizer().GetServerState();
 		// ... so add back the staged objects
 		for (const FSoftObjectPath& StagedObject : StagedObjects)
 		{
 			if (const UObject* Object = StagedObject.ResolveObject())
 			{
-				GetClientContent()->Stream->ReplicationMap.ReplicatedObjects.Add(StagedObject) = FConcertReplicatedObjectInfo::Make(*Object);
+				GetClientStreamObject()->ReplicationMap.ReplicatedObjects.Add(StagedObject) = FConcertReplicatedObjectInfo::Make(*Object);
 			}
 		}
 		// To remove the hierarchy, the user must click the actor and delete it explicitly, which will call IEditableReplicationStream::RemoveObjects on the staged objects.
diff --git a/Engine/Plugins/Developer/Concert/ConcertApp/MultiUserClient/Source/MultiUserClient/Private/Replication/Client/ReplicationClient.h b/Engine/Plugins/Developer/Concert/ConcertApp/MultiUserClient/Source/MultiUserClient/Private/Replication/Client/ReplicationClient.h
--- a/Engine/Plugins/Developer/Concert/ConcertApp/MultiUserClient/Source/MultiUserClient/Private/Replication/Client/ReplicationClient.h
+++ b/Engine/Plugins/Developer/Concert/ConcertApp/MultiUserClient/Source/MultiUserClient/Private/Replication/Client/ReplicationClient.h
@@ -2,6 +2,7 @@
 
 #pragma once
 
+#include "Assets/MultiUserReplicationClientContent.h"
 #include "Replication/Authority/AuthorityChangeTracker.h"
 #include "Replication/Authority/IClientAuthoritySynchronizer.h"
 #include "Replication/Editor/UnrealEditor/ModifyObjectInLevelHandler.h"
@@ -44,7 +45,7 @@ namespace UE::MultiUserClient
 		 * @param EndpointId The endpoint ID this instance corresponds to.
 		 * @param InDiscoveryContainer Used for auto-discovering properties added to this client's stream. The caller ensures it outlives the constructed instance.
 		 * @param InAuthorityCache Caches authority state of all clients. Passed to subsystems. The caller ensures it outlives the constructed instance.
-		 * @param InSessionContent Object that this client's stream changes are to be written into. The caller ensures it outlives the constructed instance.
+		 * @param InClientStreamContent Object that this client's stream changes are to be written into. The caller ensures it outlives the constructed instance.
 		 * @param InStreamSynchronizer Implementation for obtaining stream registered on server. The constructed instance takes ownership.
 		 * @param InAuthoritySynchronizer Implementation for obtaining the client's authority state on server. The constructed instance takes ownership.
 		 * @param InSubmissionWorkflow Implementation for for changing client streams and authority on the server. The constructed instance takes ownership.
@@ -53,14 +54,14 @@ namespace UE::MultiUserClient
 			const FGuid& EndpointId,
 			FReplicationDiscoveryContainer& InDiscoveryContainer UE_LIFETIMEBOUND,
 			FGlobalAuthorityCache& InAuthorityCache UE_LIFETIMEBOUND,
-			UMultiUserReplicationClientContent& InSessionContent UE_LIFETIMEBOUND,
+			UMultiUserReplicationStream& InClientStreamContent UE_LIFETIMEBOUND,
 			TUniquePtr<IClientStreamSynchronizer> InStreamSynchronizer,
 			TUniquePtr<IClientAuthoritySynchronizer> InAuthoritySynchronizer,
 			TUniquePtr<ISubmissionWorkflow> InSubmissionWorkflow
 			);
 		~FReplicationClient();
 
-		UMultiUserReplicationClientContent* GetClientContent() const { return ClientContentStorage; }
+		UMultiUserReplicationStream* GetClientStreamObject() const { return ClientStreamContent; }
 		FMultiUserStreamExtender& GetStreamExtender() { return *StreamExtender; }
 		/**
 		 * This is used so the UI can construct the IReplicationStreamEditor.
@@ -114,7 +115,7 @@ namespace UE::MultiUserClient
 		const FGuid EndpointId;
 		
 		/** The state of the server is synched up with this object and displayed in the UI. */
-		TObjectPtr<UMultiUserReplicationClientContent> ClientContentStorage;
+		TObjectPtr<UMultiUserReplicationStream> ClientStreamContent;
 		
 		/** Keeps the client's stream state on the server in sync. */
 		TUniquePtr<IClientStreamSynchronizer> StreamSynchronizer;
diff --git a/Engine/Plugins/Developer/Concert/ConcertApp/MultiUserClient/Source/MultiUserClient/Private/Replication/Client/ReplicationClientManager.cpp b/Engine/Plugins/Developer/Concert/ConcertApp/MultiUserClient/Source/MultiUserClient/Private/Replication/Client/ReplicationClientManager.cpp
--- a/Engine/Plugins/Developer/Concert/ConcertApp/MultiUserClient/Source/MultiUserClient/Private/Replication/Client/ReplicationClientManager.cpp
+++ b/Engine/Plugins/Developer/Concert/ConcertApp/MultiUserClient/Source/MultiUserClient/Private/Replication/Client/ReplicationClientManager.cpp
@@ -2,12 +2,12 @@
 
 #include "ReplicationClientManager.h"
 
+#include "Assets/MultiUserReplicationStream.h"
 #include "IConcertSyncClient.h"
 #include "LocalReplicationClient.h"
 #include "RemoteReplicationClient.h"
-#include "Assets/MultiUserReplicationSessionPreset.h"
 #include "Replication/Stream/StreamSynchronizer_LocalClient.h"
-#include "Replication/Misc/Query/StreamAndAuthorityQueryService.h"
+#include "Replication/Stream/MultiUserStreamId.h"
 
 #include "UObject/Package.h"
 #include "UObject/UObjectGlobals.h"
@@ -16,9 +16,11 @@ namespace UE::MultiUserClient
 {
 	namespace Private
 	{
-		static UMultiUserReplicationClientContent* MakeClientContent()
+		static UMultiUserReplicationStream* MakeClientContent()
 		{
-			return NewObject<UMultiUserReplicationClientContent>(GetTransientPackage(), NAME_None, RF_Transient | RF_Transactional);
+			UMultiUserReplicationStream* Stream = NewObject<UMultiUserReplicationStream>(GetTransientPackage(), NAME_None, RF_Transient | RF_Transactional);
+			Stream->StreamId = MultiUserStreamID;
+			return Stream;
 		}
 	}
 	
@@ -35,8 +37,8 @@ namespace UE::MultiUserClient
 		, AuthorityCache(*this)
 		, LocalClient([this, InClient]()
 		{
-			UMultiUserReplicationClientContent* ClientPreset = Private::MakeClientContent();
-			return FLocalReplicationClient(RegisteredExtenders, AuthorityCache, *ClientPreset, MakeUnique<FStreamSynchronizer_LocalClient>(InClient, ClientPreset->Stream->StreamId), InClient);
+			UMultiUserReplicationStream* ClientPreset = Private::MakeClientContent();
+			return FLocalReplicationClient(RegisteredExtenders, AuthorityCache, *ClientPreset, MakeUnique<FStreamSynchronizer_LocalClient>(InClient, ClientPreset->StreamId), InClient);
 		}())
 		, ReassignmentLogic(*this)
 	{
@@ -129,9 +131,9 @@ namespace UE::MultiUserClient
 	{
 		ForEachClient([&Collector](const FReplicationClient& Client)
 		{
-			TObjectPtr<UMultiUserReplicationClientContent> Content = Client.GetClientContent();
+			TObjectPtr<UMultiUserReplicationStream> Content = Client.GetClientStreamObject();
 			Collector.AddReferencedObject(Content);
-			ensureMsgf(Content == Client.GetClientContent(), TEXT("Did not expect reference to be obliterated"));
+			ensureMsgf(Content == Client.GetClientStreamObject(), TEXT("Did not expect reference to be obliterated"));
 			return EBreakBehavior::Continue;
 		});
 	}
diff --git a/Engine/Plugins/Developer/Concert/ConcertApp/MultiUserClient/Source/MultiUserClient/Private/Replication/Preset/PresetManager.cpp b/Engine/Plugins/Developer/Concert/ConcertApp/MultiUserClient/Source/MultiUserClient/Private/Replication/Preset/PresetManager.cpp
--- a/Engine/Plugins/Developer/Concert/ConcertApp/MultiUserClient/Source/MultiUserClient/Private/Replication/Preset/PresetManager.cpp
+++ b/Engine/Plugins/Developer/Concert/ConcertApp/MultiUserClient/Source/MultiUserClient/Private/Replication/Preset/PresetManager.cpp
@@ -2,17 +2,17 @@
 
 #include "PresetManager.h"
 
-#include "ConcertLogGlobal.h"
 #include "Assets/MultiUserReplicationSessionPreset.h"
-
-#include "FileHelpers.h"
+#include "ConcertLogGlobal.h"
 #include "IConcertSyncClient.h"
-#include "Misc/ObjectPathOuterIterator.h"
 #include "Replication/Client/ReplicationClientManager.h"
 #include "Replication/Misc/ReplicationStreamUtils.h"
 #include "Replication/Muting/MuteStateManager.h"
+#include "Replication/Stream/MultiUserStreamId.h"
 #include "Widgets/ActiveSession/Replication/Client/ClientUtils.h"
 
+#include "FileHelpers.h"
+
 namespace UE::MultiUserClient
 {
 	namespace Private
@@ -303,9 +303,9 @@ namespace UE::MultiUserClient
 		for (const TPair<const FReplicationClient*, FConcertClientInfo>& ClientData : IncludedClients)
 		{
 			const auto[Client, ClientInfo] = ClientData;
-			UMultiUserReplicationClientContent* ClientContent_ToCopy = Client->GetClientContent();
-			UMultiUserReplicationClientContent* ClientContent_InPreset = Preset->AddClientIfUnique(ClientInfo);
-			if (!ClientContent_InPreset)
+			UMultiUserReplicationStream* CopiedClientStream = Client->GetClientStreamObject();
+			UMultiUserReplicationClientContent* TargetClientPreset = Preset->AddClientIfUnique(ClientInfo, MultiUserStreamID);
+			if (!TargetClientPreset)
 			{
 				UE_LOG(LogConcert, Warning,
 					TEXT("There are multiple clients with display name %s and device name %s in the session. Only the 1st encountered will be saved into the preset. Did you perhaps launch 2 editors on the same machine (if so you can use -CONCERTDISPLAYNAME)?"),
@@ -315,9 +315,9 @@ namespace UE::MultiUserClient
 				continue;
 			}
 
-			ClientContent_InPreset->Stream->Copy(*ClientContent_ToCopy->Stream);
+			TargetClientPreset->Stream->Copy(*CopiedClientStream);
 			// TODO UE-219834: Once UMultiUserReplicationStream::FrequencySettings reflect the server state, this can be removed.
-			ClientContent_InPreset->Stream->FrequencySettings = Client->GetStreamSynchronizer().GetFrequencySettings();
+			TargetClientPreset->Stream->FrequencySettings = Client->GetStreamSynchronizer().GetFrequencySettings();
 		}
 
 		Preset->SetMuteContent(
diff --git a/Engine/Plugins/Developer/Concert/ConcertApp/MultiUserClient/Source/MultiUserClient/Private/Replication/Stream/MultiUserStreamId.h b/Engine/Plugins/Developer/Concert/ConcertApp/MultiUserClient/Source/MultiUserClient/Private/Replication/Stream/MultiUserStreamId.h
new file mode 100644
--- /dev/null
+++ b/Engine/Plugins/Developer/Concert/ConcertApp/MultiUserClient/Source/MultiUserClient/Private/Replication/Stream/MultiUserStreamId.h
@@ -0,0 +1,11 @@
+﻿// Copyright Epic Games, Inc. All Rights Reserved.
+
+#pragma once
+
+#include "Misc/Guid.h"
+
+namespace UE::MultiUserClient
+{
+	/** The stream ID that is used by all Multi-User streams. */
+	constexpr FGuid MultiUserStreamID { 0xAAAAAAAA, 0xBBBBBBBB, 0xCCCCCCCC, 0xDDDDDDDD };
+}
\ No newline at end of file
diff --git a/Engine/Plugins/Developer/Concert/ConcertApp/MultiUserClient/Source/MultiUserClient/Private/Replication/Stream/StreamSynchronizer_RemoteClient.cpp b/Engine/Plugins/Developer/Concert/ConcertApp/MultiUserClient/Source/MultiUserClient/Private/Replication/Stream/StreamSynchronizer_RemoteClient.cpp
--- a/Engine/Plugins/Developer/Concert/ConcertApp/MultiUserClient/Source/MultiUserClient/Private/Replication/Stream/StreamSynchronizer_RemoteClient.cpp
+++ b/Engine/Plugins/Developer/Concert/ConcertApp/MultiUserClient/Source/MultiUserClient/Private/Replication/Stream/StreamSynchronizer_RemoteClient.cpp
@@ -3,6 +3,7 @@
 #include "StreamSynchronizer_RemoteClient.h"
 
 #include "Assets/MultiUserReplicationClientContent.h"
+#include "MultiUserStreamId.h"
 #include "Replication/Misc/Query/StreamAndAuthorityQueryService.h"
 
 namespace UE::MultiUserClient
@@ -26,7 +27,7 @@ namespace UE::MultiUserClient
 
 	FGuid FStreamSynchronizer_RemoteClient::GetStreamId() const
 	{
-		return UMultiUserReplicationClientContent::MultiUserStreamID;
+		return MultiUserStreamID;
 	}
 
 	void FStreamSynchronizer_RemoteClient::HandleStreamQuery(const TArray<FConcertBaseStreamInfo>& Streams)
@@ -35,7 +36,7 @@ namespace UE::MultiUserClient
 		// That handles the (unlikely) case in which some external logic had added streams to the same client which we must differentiate
 		const FConcertBaseStreamInfo* MultiUserDescription = Streams.FindByPredicate([](const FConcertBaseStreamInfo& Description)
 		{
-			return Description.Identifier == UMultiUserReplicationClientContent::MultiUserStreamID;
+			return Description.Identifier == MultiUserStreamID;
 		});
 		if (MultiUserDescription
 			// Avoid unnecessary Broadcast()s
diff --git a/Engine/Plugins/Developer/Concert/ConcertApp/MultiUserClient/Source/MultiUserReplicationEditor/Internal/Assets/MultiUserReplicationClientContent.h b/Engine/Plugins/Developer/Concert/ConcertApp/MultiUserClient/Source/MultiUserReplicationEditor/Internal/Assets/MultiUserReplicationClientContent.h
--- a/Engine/Plugins/Developer/Concert/ConcertApp/MultiUserClient/Source/MultiUserReplicationEditor/Internal/Assets/MultiUserReplicationClientContent.h
+++ b/Engine/Plugins/Developer/Concert/ConcertApp/MultiUserClient/Source/MultiUserReplicationEditor/Internal/Assets/MultiUserReplicationClientContent.h
@@ -14,9 +14,6 @@ class MULTIUSERREPLICATIONEDITOR_API UMultiUserReplicationClientContent : public
 	GENERATED_BODY()
 public:
 
-	/** The stream ID that is used by all Multi-User streams. */
-	static constexpr FGuid MultiUserStreamID { 0xAAAAAAAA, 0xBBBBBBBB, 0xCCCCCCCC, 0xDDDDDDDD };
-
 	/** The stream this client is managing */
 	UPROPERTY(Instanced)
 	TObjectPtr<UMultiUserReplicationStream> Stream;
diff --git a/Engine/Plugins/Developer/Concert/ConcertApp/MultiUserClient/Source/MultiUserReplicationEditor/Internal/Assets/MultiUserReplicationSessionPreset.h b/Engine/Plugins/Developer/Concert/ConcertApp/MultiUserClient/Source/MultiUserReplicationEditor/Internal/Assets/MultiUserReplicationSessionPreset.h
--- a/Engine/Plugins/Developer/Concert/ConcertApp/MultiUserClient/Source/MultiUserReplicationEditor/Internal/Assets/MultiUserReplicationSessionPreset.h
+++ b/Engine/Plugins/Developer/Concert/ConcertApp/MultiUserClient/Source/MultiUserReplicationEditor/Internal/Assets/MultiUserReplicationSessionPreset.h
@@ -52,7 +52,7 @@ public:
 	bool ContainsExactClient(const FConcertClientInfo& ClientInfo) const { return GetExactClientContent(ClientInfo) != nullptr; }
 	
 	/** Adds a client to the preset if it's not already present. */
-	UMultiUserReplicationClientContent* AddClientIfUnique(const FConcertClientInfo& ClientInfo);
+	UMultiUserReplicationClientContent* AddClientIfUnique(const FConcertClientInfo& ClientInfo, const FGuid& StreamId);
 	
 	const TArray<TObjectPtr<UMultiUserReplicationClientContent>>& GetClientPresets() const { return ClientPresets; }
 	
diff --git a/Engine/Plugins/Developer/Concert/ConcertApp/MultiUserClient/Source/MultiUserReplicationEditor/Private/Assets/MultiUserReplicationClientContent.cpp b/Engine/Plugins/Developer/Concert/ConcertApp/MultiUserClient/Source/MultiUserReplicationEditor/Private/Assets/MultiUserReplicationClientContent.cpp
--- a/Engine/Plugins/Developer/Concert/ConcertApp/MultiUserClient/Source/MultiUserReplicationEditor/Private/Assets/MultiUserReplicationClientContent.cpp
+++ b/Engine/Plugins/Developer/Concert/ConcertApp/MultiUserClient/Source/MultiUserReplicationEditor/Private/Assets/MultiUserReplicationClientContent.cpp
@@ -8,7 +8,6 @@
 UMultiUserReplicationClientContent::UMultiUserReplicationClientContent()
 {
 	Stream = CreateDefaultSubobject<UMultiUserReplicationStream>(TEXT("ReplicationList"));
-	Stream->StreamId = MultiUserStreamID;
 	Stream->SetFlags(RF_Transactional);
 }
 
diff --git a/Engine/Plugins/Developer/Concert/ConcertApp/MultiUserClient/Source/MultiUserReplicationEditor/Private/Assets/MultiUserReplicationSessionPreset.cpp b/Engine/Plugins/Developer/Concert/ConcertApp/MultiUserClient/Source/MultiUserReplicationEditor/Private/Assets/MultiUserReplicationSessionPreset.cpp
--- a/Engine/Plugins/Developer/Concert/ConcertApp/MultiUserClient/Source/MultiUserReplicationEditor/Private/Assets/MultiUserReplicationSessionPreset.cpp
+++ b/Engine/Plugins/Developer/Concert/ConcertApp/MultiUserClient/Source/MultiUserReplicationEditor/Private/Assets/MultiUserReplicationSessionPreset.cpp
@@ -44,7 +44,7 @@ UMultiUserReplicationClientContent* UMultiUserReplicationSessionPreset::GetExact
 	return Result ? *Result : nullptr;
 }
 
-UMultiUserReplicationClientContent* UMultiUserReplicationSessionPreset::AddClientIfUnique(const FConcertClientInfo& ClientInfo)
+UMultiUserReplicationClientContent* UMultiUserReplicationSessionPreset::AddClientIfUnique(const FConcertClientInfo& ClientInfo, const FGuid& StreamId)
 {
 	if (ContainsExactClient(ClientInfo))
 	{
@@ -54,6 +54,7 @@ UMultiUserReplicationClientContent* UMultiUserReplicationSessionPreset::AddClien
 	UMultiUserReplicationClientContent* Result = NewObject<UMultiUserReplicationClientContent>(this);
 	Result->DisplayName = ClientInfo.DisplayName;
 	Result->DeviceName = ClientInfo.DeviceName;
+	Result->Stream->StreamId = StreamId;
 	ClientPresets.Add(Result);
 	return Result;
 }
