Rename functions of SReplicationStatus to indicate that clients that may be replicating are displayed.
This CL is for a future change that introduces offline client objects.

diff --git a/Engine/Plugins/Developer/Concert/ConcertApp/MultiUserClient/Source/MultiUserClient/Private/Widgets/ActiveSession/Replication/Client/Multi/SMultiClientView.cpp b/Engine/Plugins/Developer/Concert/ConcertApp/MultiUserClient/Source/MultiUserClient/Private/Widgets/ActiveSession/Replication/Client/Multi/SMultiClientView.cpp
--- a/Engine/Plugins/Developer/Concert/ConcertApp/MultiUserClient/Source/MultiUserClient/Private/Widgets/ActiveSession/Replication/Client/Multi/SMultiClientView.cpp
+++ b/Engine/Plugins/Developer/Concert/ConcertApp/MultiUserClient/Source/MultiUserClient/Private/Widgets/ActiveSession/Replication/Client/Multi/SMultiClientView.cpp
@@ -62,8 +62,8 @@ namespace UE::MultiUserClient
 		
 		SReplicationStatus::AppendReplicationStatus(*Content, ClientManager->GetAuthorityCache(),
 			SReplicationStatus::FArguments()
-			.DisplayedClients(this, &SMultiClientView::GetDisplayClientIds)
-			.ForEachReplicatedObject(this, &SMultiClientView::EnumerateObjectsInStreams)
+			.ReplicatableClients(this, &SMultiClientView::GetReplicatableClientIds)
+			.ForEachObjectInStream(this, &SMultiClientView::EnumerateObjectsInStreams)
 			);
 
 		RebuildClientSubscriptions();
@@ -168,7 +168,7 @@ namespace UE::MultiUserClient
 			];
 	}
 
-	TSet<FGuid> SMultiClientView::GetDisplayClientIds() const
+	TSet<FGuid> SMultiClientView::GetReplicatableClientIds() const
 	{
 		TSet<FGuid> ClientIds;
 		StreamModel->ForEachClient([&ClientIds](const FReplicationClient* Client)
diff --git a/Engine/Plugins/Developer/Concert/ConcertApp/MultiUserClient/Source/MultiUserClient/Private/Widgets/ActiveSession/Replication/Client/Multi/SMultiClientView.h b/Engine/Plugins/Developer/Concert/ConcertApp/MultiUserClient/Source/MultiUserClient/Private/Widgets/ActiveSession/Replication/Client/Multi/SMultiClientView.h
--- a/Engine/Plugins/Developer/Concert/ConcertApp/MultiUserClient/Source/MultiUserClient/Private/Widgets/ActiveSession/Replication/Client/Multi/SMultiClientView.h
+++ b/Engine/Plugins/Developer/Concert/ConcertApp/MultiUserClient/Source/MultiUserClient/Private/Widgets/ActiveSession/Replication/Client/Multi/SMultiClientView.h
@@ -79,7 +79,9 @@ namespace UE::MultiUserClient
 		TSharedRef<SWidget> CreateNoPropertiesWarning() const;
 
 		// SClientToolbar attributes
-		TSet<FGuid> GetDisplayClientIds() const;
+		/** @return Gets the clients that may be replicating */
+		TSet<FGuid> GetReplicatableClientIds() const;
+		/** Calls Consumer for each object path that is in a stream - independent of whether it is being replicated or not. */
 		void EnumerateObjectsInStreams(TFunctionRef<void(const FSoftObjectPath&)> Consumer) const;
 		
 		void RebuildClientSubscriptions();
diff --git a/Engine/Plugins/Developer/Concert/ConcertApp/MultiUserClient/Source/MultiUserClient/Private/Widgets/ActiveSession/Replication/Client/SReplicationStatus.cpp b/Engine/Plugins/Developer/Concert/ConcertApp/MultiUserClient/Source/MultiUserClient/Private/Widgets/ActiveSession/Replication/Client/SReplicationStatus.cpp
--- a/Engine/Plugins/Developer/Concert/ConcertApp/MultiUserClient/Source/MultiUserClient/Private/Widgets/ActiveSession/Replication/Client/SReplicationStatus.cpp
+++ b/Engine/Plugins/Developer/Concert/ConcertApp/MultiUserClient/Source/MultiUserClient/Private/Widgets/ActiveSession/Replication/Client/SReplicationStatus.cpp
@@ -60,8 +60,8 @@ namespace UE::MultiUserClient
 			.Padding(4, 0, 4, 3)
 			[
 				SNew(SReplicationStatus, InAuthorityCache)
-				.DisplayedClients(InArgs._DisplayedClients)
-				.ForEachReplicatedObject(InArgs._ForEachReplicatedObject)
+				.ReplicatableClients(InArgs._ReplicatableClients)
+				.ForEachObjectInStream(InArgs._ForEachObjectInStream)
 			];
 	}
 
@@ -69,10 +69,10 @@ namespace UE::MultiUserClient
 	{
 		AuthorityCache = &InAuthorityCache;
 		
-		DisplayedClientsAttribute = InArgs._DisplayedClients;
-		check(DisplayedClientsAttribute.IsBound() || DisplayedClientsAttribute.IsSet());
-		ForEachReplicatedObjectDelegate = InArgs._ForEachReplicatedObject;
-		check(ForEachReplicatedObjectDelegate.IsBound());
+		ReplicatableClientsAttribute = InArgs._ReplicatableClients;
+		check(ReplicatableClientsAttribute.IsBound() || ReplicatableClientsAttribute.IsSet());
+		ForEachObjectInStreamDelegate = InArgs._ForEachObjectInStream;
+		check(ForEachObjectInStreamDelegate.IsBound());
 
 		AuthorityCache->OnCacheChanged().AddSP(this, &SReplicationStatus::OnAuthorityCacheChanged);
 		ChildSlot
@@ -114,15 +114,15 @@ namespace UE::MultiUserClient
 
 	void SReplicationStatus::RefreshStatusText()
 	{
-		const TSet<FGuid> DisplayedClients = DisplayedClientsAttribute.Get();
+		const TSet<FGuid> ReplicatingClients = ReplicatableClientsAttribute.Get();
 
 		TSet<FSoftObjectPath> ReplicatedActors;
 		TSet<FSoftObjectPath> ReplicatedObjects;
-		ForEachReplicatedObjectDelegate.Execute([this, &DisplayedClients, &ReplicatedActors, &ReplicatedObjects](const FSoftObjectPath& Path)
+		ForEachObjectInStreamDelegate.Execute([this, &ReplicatingClients, &ReplicatedActors, &ReplicatedObjects](const FSoftObjectPath& Path)
 		{
-			const bool bIsReplicated = Algo::AnyOf(AuthorityCache->GetClientsWithAuthorityOverObject(Path), [&DisplayedClients](const FGuid& ClientId)
+			const bool bIsReplicated = Algo::AnyOf(AuthorityCache->GetClientsWithAuthorityOverObject(Path), [&ReplicatingClients](const FGuid& ClientId)
 			{
-				return DisplayedClients.Contains(ClientId);
+				return ReplicatingClients.Contains(ClientId);
 			});
 			if (bIsReplicated)
 			{
diff --git a/Engine/Plugins/Developer/Concert/ConcertApp/MultiUserClient/Source/MultiUserClient/Private/Widgets/ActiveSession/Replication/Client/SReplicationStatus.h b/Engine/Plugins/Developer/Concert/ConcertApp/MultiUserClient/Source/MultiUserClient/Private/Widgets/ActiveSession/Replication/Client/SReplicationStatus.h
--- a/Engine/Plugins/Developer/Concert/ConcertApp/MultiUserClient/Source/MultiUserClient/Private/Widgets/ActiveSession/Replication/Client/SReplicationStatus.h
+++ b/Engine/Plugins/Developer/Concert/ConcertApp/MultiUserClient/Source/MultiUserClient/Private/Widgets/ActiveSession/Replication/Client/SReplicationStatus.h
@@ -14,7 +14,7 @@ namespace UE::MultiUserClient
 {
 	class FGlobalAuthorityCache;
 
-	DECLARE_DELEGATE_OneParam(FForEachReplicatedObject, TFunctionRef<void(const FSoftObjectPath&)> Consumer);
+	DECLARE_DELEGATE_OneParam(FForEachObjectInStream, TFunctionRef<void(const FSoftObjectPath&)> Consumer);
 	
 	/**
 	 * Displays a text "Replicating x Objects for y Actors".
@@ -29,11 +29,11 @@ namespace UE::MultiUserClient
 
 		SLATE_BEGIN_ARGS(SReplicationStatus)
 		{}
-			/** The clients to show statistics for */
-			SLATE_ATTRIBUTE(TSet<FGuid>, DisplayedClients)
+			/** The clients that may be replicating, for which we show the stats. */
+			SLATE_ATTRIBUTE(TSet<FGuid>, ReplicatableClients)
 
-			/** Delegate which enumerates every replicated object. */
-			SLATE_EVENT(FForEachReplicatedObject, ForEachReplicatedObject)
+			/** Delegate which enumerates every object registered in a stream - independent of whether it is being replicated or not. */
+			SLATE_EVENT(FForEachObjectInStream, ForEachObjectInStream)
 		SLATE_END_ARGS()
 
 		/** Adds a separator and SReplicationStatus to the bottom of the VerticalBox- */
@@ -50,17 +50,17 @@ namespace UE::MultiUserClient
 		/** Used to get authority state of objects and informs us when authority changes. */
 		FGlobalAuthorityCache* AuthorityCache = nullptr;
 
-		/** The clients to show statistics for */
-		TAttribute<TSet<FGuid>> DisplayedClientsAttribute;
-		/** Delegate which enumerates every replicated object. */
-		FForEachReplicatedObject ForEachReplicatedObjectDelegate;
+		/** The clients that may be replicating, for which we show the stats. */
+		TAttribute<TSet<FGuid>> ReplicatableClientsAttribute;
+		/** Delegate which enumerates every object registered in a stream - independent of whether it is being replicated or not. */
+		FForEachObjectInStream ForEachObjectInStreamDelegate;
 
 		/** Updated when authority changes. Displays subobjects in bold. */
 		TSharedPtr<STextBlock> ObjectsText;
 		/** Updated when authority changes. Displays actors in bold. */
 		TSharedPtr<STextBlock> ActorsText;
 
-		void OnAuthorityCacheChanged(const FGuid& ClientId) { RefreshStatusText(); }
+		void OnAuthorityCacheChanged(const FGuid&) { RefreshStatusText(); }
 	};
 }
 
