Remove IClientSelectionModel::ContainsClient

diff --git a/Engine/Plugins/Developer/Concert/ConcertApp/MultiUserClient/Source/MultiUserClient/Private/Widgets/ActiveSession/Replication/Client/Multi/AllClientsSelectionModel.h b/Engine/Plugins/Developer/Concert/ConcertApp/MultiUserClient/Source/MultiUserClient/Private/Widgets/ActiveSession/Replication/Client/Multi/AllClientsSelectionModel.h
--- a/Engine/Plugins/Developer/Concert/ConcertApp/MultiUserClient/Source/MultiUserClient/Private/Widgets/ActiveSession/Replication/Client/Multi/AllClientsSelectionModel.h
+++ b/Engine/Plugins/Developer/Concert/ConcertApp/MultiUserClient/Source/MultiUserClient/Private/Widgets/ActiveSession/Replication/Client/Multi/AllClientsSelectionModel.h
@@ -18,7 +18,6 @@ namespace UE::MultiUserClient
 
 		//~ Begin IClientSelectionModel Interface
 		virtual void ForEachSelectedClient(TFunctionRef<EBreakBehavior(FReplicationClient&)> ProcessClient) const override;
-		virtual bool ContainsClient(const FGuid& ClientId) const override { return true; }
 		virtual FOnSelectionChanged& OnSelectionChanged() override { return OnSelectionChangedDelegate; }
 		//~ End IClientSelectionModel Interface
 
diff --git a/Engine/Plugins/Developer/Concert/ConcertApp/MultiUserClient/Source/MultiUserClient/Private/Widgets/ActiveSession/Replication/Client/Multi/IClientSelectionModel.h b/Engine/Plugins/Developer/Concert/ConcertApp/MultiUserClient/Source/MultiUserClient/Private/Widgets/ActiveSession/Replication/Client/Multi/IClientSelectionModel.h
--- a/Engine/Plugins/Developer/Concert/ConcertApp/MultiUserClient/Source/MultiUserClient/Private/Widgets/ActiveSession/Replication/Client/Multi/IClientSelectionModel.h
+++ b/Engine/Plugins/Developer/Concert/ConcertApp/MultiUserClient/Source/MultiUserClient/Private/Widgets/ActiveSession/Replication/Client/Multi/IClientSelectionModel.h
@@ -19,9 +19,6 @@ namespace UE::MultiUserClient
 		/** Iterates through every client in the selection. */
 		virtual void ForEachSelectedClient(TFunctionRef<EBreakBehavior(FReplicationClient&)> ProcessClient) const = 0;
 
-		/** @return Whether the client is selcted. */
-		virtual bool ContainsClient(const FGuid& ClientId) const = 0;
-
 		DECLARE_MULTICAST_DELEGATE(FOnSelectionChanged);
 		/** Called when the clients ForEachSelectedClient enumerates has changed. */
 		virtual FOnSelectionChanged& OnSelectionChanged() = 0;
diff --git a/Engine/Plugins/Developer/Concert/ConcertApp/MultiUserClient/Source/MultiUserClient/Private/Widgets/ActiveSession/Replication/Client/Multi/SMultiClientView.cpp b/Engine/Plugins/Developer/Concert/ConcertApp/MultiUserClient/Source/MultiUserClient/Private/Widgets/ActiveSession/Replication/Client/Multi/SMultiClientView.cpp
--- a/Engine/Plugins/Developer/Concert/ConcertApp/MultiUserClient/Source/MultiUserClient/Private/Widgets/ActiveSession/Replication/Client/Multi/SMultiClientView.cpp
+++ b/Engine/Plugins/Developer/Concert/ConcertApp/MultiUserClient/Source/MultiUserClient/Private/Widgets/ActiveSession/Replication/Client/Multi/SMultiClientView.cpp
@@ -196,14 +196,10 @@ namespace UE::MultiUserClient
 	{
 		CleanClientSubscriptions();
 
-		ClientManager->ForEachClient([this](FReplicationClient& Client)
+		SelectionModel->ForEachSelectedClient([this](FReplicationClient& Client)
 		{
-			if (SelectionModel->ContainsClient(Client.GetEndpointId()))
-			{
-				Client.OnModelChanged().AddSP(this, &SMultiClientView::RefreshUI);
-				Client.OnHierarchyNeedsRefresh().AddRaw(this, &SMultiClientView::RefreshUI);
-			}
-			
+			Client.OnModelChanged().AddSP(this, &SMultiClientView::RefreshUI);
+			Client.OnHierarchyNeedsRefresh().AddRaw(this, &SMultiClientView::RefreshUI);
 			return EBreakBehavior::Continue;
 		});
 	}
