Move replication related symbols into UE::MultiUserClient::Replication from UE::MultiUserClient namespace.

diff --git a/Engine/Plugins/Developer/Concert/ConcertApp/MultiUserClient/Source/MultiUserClient/Private/MultiUserClientModule.cpp b/Engine/Plugins/Developer/Concert/ConcertApp/MultiUserClient/Source/MultiUserClient/Private/MultiUserClientModule.cpp
--- a/Engine/Plugins/Developer/Concert/ConcertApp/MultiUserClient/Source/MultiUserClient/Private/MultiUserClientModule.cpp
+++ b/Engine/Plugins/Developer/Concert/ConcertApp/MultiUserClient/Source/MultiUserClient/Private/MultiUserClientModule.cpp
@@ -624,7 +624,7 @@ public:
 
 		// Boot the client instance
 		MultiUserClient->Startup(ClientConfig, EConcertSyncSessionFlags::Default_MultiUserSession);
-		ReplicationManager = MakeShared<UE::MultiUserClient::FMultiUserReplicationManager>(MultiUserClient.ToSharedRef());
+		ReplicationManager = MakeShared<UE::MultiUserClient::Replication::FMultiUserReplicationManager>(MultiUserClient.ToSharedRef());
 
 		// Hook UI elements in the tool bar (and setup commands).
 		RegisterUI();
@@ -1409,7 +1409,7 @@ private:
 
 	TSharedPtr<IConcertSyncClient> MultiUserClient;
 	/** Interacts with the replication system on behalf of Multi-User. */
-	TSharedPtr<UE::MultiUserClient::FMultiUserReplicationManager> ReplicationManager;
+	TSharedPtr<UE::MultiUserClient::Replication::FMultiUserReplicationManager> ReplicationManager;
 
 	/** True if the tab spawners have been registered for this module */
 	bool bHasRegisteredTabSpawners = false;
diff --git a/Engine/Plugins/Developer/Concert/ConcertApp/MultiUserClient/Source/MultiUserClient/Private/Replication/Authority/AuthorityChangeTracker.cpp b/Engine/Plugins/Developer/Concert/ConcertApp/MultiUserClient/Source/MultiUserClient/Private/Replication/Authority/AuthorityChangeTracker.cpp
--- a/Engine/Plugins/Developer/Concert/ConcertApp/MultiUserClient/Source/MultiUserClient/Private/Replication/Authority/AuthorityChangeTracker.cpp
+++ b/Engine/Plugins/Developer/Concert/ConcertApp/MultiUserClient/Source/MultiUserClient/Private/Replication/Authority/AuthorityChangeTracker.cpp
@@ -6,7 +6,7 @@
 #include "IClientAuthoritySynchronizer.h"
 #include "Replication/Misc/GlobalAuthorityCache.h"
 
-namespace UE::MultiUserClient
+namespace UE::MultiUserClient::Replication
 {
 	FAuthorityChangeTracker::FAuthorityChangeTracker(
 		const FGuid& InClientId,
diff --git a/Engine/Plugins/Developer/Concert/ConcertApp/MultiUserClient/Source/MultiUserClient/Private/Replication/Authority/AuthorityChangeTracker.h b/Engine/Plugins/Developer/Concert/ConcertApp/MultiUserClient/Source/MultiUserClient/Private/Replication/Authority/AuthorityChangeTracker.h
--- a/Engine/Plugins/Developer/Concert/ConcertApp/MultiUserClient/Source/MultiUserClient/Private/Replication/Authority/AuthorityChangeTracker.h
+++ b/Engine/Plugins/Developer/Concert/ConcertApp/MultiUserClient/Source/MultiUserClient/Private/Replication/Authority/AuthorityChangeTracker.h
@@ -6,14 +6,14 @@
 #include "Containers/ContainersFwd.h"
 #include "Replication/IConcertClientReplicationManager.h"
 
-namespace UE::MultiUserClient
+namespace UE::MultiUserClient::Replication
 {
 	class FGlobalAuthorityCache;
 }
 
 struct FSoftObjectPath;
 
-namespace UE::MultiUserClient
+namespace UE::MultiUserClient::Replication
 {
 	class IClientAuthoritySynchronizer;
 	enum class EAuthorityMutability : uint8;
diff --git a/Engine/Plugins/Developer/Concert/ConcertApp/MultiUserClient/Source/MultiUserClient/Private/Replication/Authority/AuthoritySynchronizer_LocalClient.cpp b/Engine/Plugins/Developer/Concert/ConcertApp/MultiUserClient/Source/MultiUserClient/Private/Replication/Authority/AuthoritySynchronizer_LocalClient.cpp
--- a/Engine/Plugins/Developer/Concert/ConcertApp/MultiUserClient/Source/MultiUserClient/Private/Replication/Authority/AuthoritySynchronizer_LocalClient.cpp
+++ b/Engine/Plugins/Developer/Concert/ConcertApp/MultiUserClient/Source/MultiUserClient/Private/Replication/Authority/AuthoritySynchronizer_LocalClient.cpp
@@ -5,7 +5,7 @@
 #include "IConcertSyncClient.h"
 #include "Replication/IConcertClientReplicationManager.h"
 
-namespace UE::MultiUserClient
+namespace UE::MultiUserClient::Replication
 {
 	FAuthoritySynchronizer_LocalClient::FAuthoritySynchronizer_LocalClient(TSharedRef<IConcertSyncClient> InClient)
 		: Client(MoveTemp(InClient))
diff --git a/Engine/Plugins/Developer/Concert/ConcertApp/MultiUserClient/Source/MultiUserClient/Private/Replication/Misc/Frequency/FrequencyUtils.cpp b/Engine/Plugins/Developer/Concert/ConcertApp/MultiUserClient/Source/MultiUserClient/Private/Replication/Misc/Frequency/FrequencyUtils.cpp
--- a/Engine/Plugins/Developer/Concert/ConcertApp/MultiUserClient/Source/MultiUserClient/Private/Replication/Misc/Frequency/FrequencyUtils.cpp
+++ b/Engine/Plugins/Developer/Concert/ConcertApp/MultiUserClient/Source/MultiUserClient/Private/Replication/Misc/Frequency/FrequencyUtils.cpp
@@ -2,7 +2,7 @@
 
 #include "FrequencyUtils.h"
 
-namespace UE::MultiUserClient::FrequencyUtils
+namespace UE::MultiUserClient::Replication::FrequencyUtils
 {
 	TSharedPtr<IParallelSubmissionOperation> SetFrequencySettingForClients(
 		const FSoftObjectPath& ContextObject,
diff --git a/Engine/Plugins/Developer/Concert/ConcertApp/MultiUserClient/Source/MultiUserClient/Private/Replication/Misc/Frequency/FrequencyUtils.h b/Engine/Plugins/Developer/Concert/ConcertApp/MultiUserClient/Source/MultiUserClient/Private/Replication/Misc/Frequency/FrequencyUtils.h
--- a/Engine/Plugins/Developer/Concert/ConcertApp/MultiUserClient/Source/MultiUserClient/Private/Replication/Misc/Frequency/FrequencyUtils.h
+++ b/Engine/Plugins/Developer/Concert/ConcertApp/MultiUserClient/Source/MultiUserClient/Private/Replication/Misc/Frequency/FrequencyUtils.h
@@ -13,12 +13,12 @@
 
 #include <type_traits>
 
-namespace UE::MultiUserClient
+namespace UE::MultiUserClient::Replication
 {
 	class IParallelSubmissionOperation;
 }
 
-namespace UE::MultiUserClient::FrequencyUtils
+namespace UE::MultiUserClient::Replication::FrequencyUtils
 {
 	// Low effort optimization to minimize heap allocations
 	using FInlineClientArray = TArray<FGuid, TInlineAllocator<8>>;
@@ -94,7 +94,7 @@ namespace UE::MultiUserClient::FrequencyUtils
 	}
 }
 
-namespace UE::MultiUserClient::FrequencyUtils
+namespace UE::MultiUserClient::Replication::FrequencyUtils
 {
 	template<typename TSubSetting, typename TSelectSubSetting>
 	TOptional<TSubSetting> FindSharedFrequencySetting(
diff --git a/Engine/Plugins/Developer/Concert/ConcertApp/MultiUserClient/Source/MultiUserClient/Private/Replication/Misc/GlobalAuthorityCache.h b/Engine/Plugins/Developer/Concert/ConcertApp/MultiUserClient/Source/MultiUserClient/Private/Replication/Misc/GlobalAuthorityCache.h
--- a/Engine/Plugins/Developer/Concert/ConcertApp/MultiUserClient/Source/MultiUserClient/Private/Replication/Misc/GlobalAuthorityCache.h
+++ b/Engine/Plugins/Developer/Concert/ConcertApp/MultiUserClient/Source/MultiUserClient/Private/Replication/Misc/GlobalAuthorityCache.h
@@ -12,14 +12,14 @@
 
 struct FConcertPropertyChain;
 
-namespace UE::MultiUserClient
+namespace UE::MultiUserClient::Replication
 {
 	class FRemoteReplicationClient;
 	class FReplicationClient;
 	class FReplicationClientManager;
 }
 
-namespace UE::MultiUserClient
+namespace UE::MultiUserClient::Replication
 {
 	/**
 	 * Allows efficient look-up of which objects and properties are owned by which clients.
diff --git a/Engine/Plugins/Developer/Concert/ConcertApp/MultiUserClient/Source/MultiUserClient/Private/Replication/Misc/Util/StreamRequestUtils.cpp b/Engine/Plugins/Developer/Concert/ConcertApp/MultiUserClient/Source/MultiUserClient/Private/Replication/Misc/Util/StreamRequestUtils.cpp
--- a/Engine/Plugins/Developer/Concert/ConcertApp/MultiUserClient/Source/MultiUserClient/Private/Replication/Misc/Util/StreamRequestUtils.cpp
+++ b/Engine/Plugins/Developer/Concert/ConcertApp/MultiUserClient/Source/MultiUserClient/Private/Replication/Misc/Util/StreamRequestUtils.cpp
@@ -9,7 +9,7 @@
 
 #include "Containers/UnrealString.h"
 
-namespace UE::MultiUserClient::StreamRequestUtils
+namespace UE::MultiUserClient::Replication::StreamRequestUtils
 {
 	FConcertReplication_ChangeStream_Request BuildChangeRequest_CreateNewStream(
 		const FGuid& StreamId,
diff --git a/Engine/Plugins/Developer/Concert/ConcertApp/MultiUserClient/Source/MultiUserClient/Private/Replication/Misc/Util/StreamRequestUtils.h b/Engine/Plugins/Developer/Concert/ConcertApp/MultiUserClient/Source/MultiUserClient/Private/Replication/Misc/Util/StreamRequestUtils.h
--- a/Engine/Plugins/Developer/Concert/ConcertApp/MultiUserClient/Source/MultiUserClient/Private/Replication/Misc/Util/StreamRequestUtils.h
+++ b/Engine/Plugins/Developer/Concert/ConcertApp/MultiUserClient/Source/MultiUserClient/Private/Replication/Misc/Util/StreamRequestUtils.h
@@ -8,7 +8,7 @@ class IConcertClientSession;
 
 struct FConcertReplication_ChangeStream_Request;
 
-namespace UE::MultiUserClient
+namespace UE::MultiUserClient::Replication
 {
 	/** Describes changes that MU clients make to the objects in a stream. */
 	struct FStreamChangelist
@@ -28,7 +28,7 @@ namespace UE::MultiUserClient
 	};
 }
 
-namespace UE::MultiUserClient::StreamRequestUtils
+namespace UE::MultiUserClient::Replication::StreamRequestUtils
 {
 	/**
 	 * Builds a request for creating a new stream.
diff --git a/Engine/Plugins/Developer/Concert/ConcertApp/MultiUserClient/Source/MultiUserClient/Private/Replication/MultiUserReplicationManager.h b/Engine/Plugins/Developer/Concert/ConcertApp/MultiUserClient/Source/MultiUserClient/Private/Replication/MultiUserReplicationManager.h
--- a/Engine/Plugins/Developer/Concert/ConcertApp/MultiUserClient/Source/MultiUserClient/Private/Replication/MultiUserReplicationManager.h
+++ b/Engine/Plugins/Developer/Concert/ConcertApp/MultiUserClient/Source/MultiUserClient/Private/Replication/MultiUserReplicationManager.h
@@ -23,9 +23,9 @@ class IConcertSyncClient;
 enum class EConcertConnectionStatus : uint8;
 
 namespace UE::ConcertSharedSlate { class IEditableReplicationStreamModel; }
-namespace UE::ConcertSyncClient::Replication { struct FJoinReplicatedSessionResult; }
+namespace UE::ConcertSyncClient { struct FJoinReplicatedSessionResult; }
 
-namespace UE::MultiUserClient
+namespace UE::MultiUserClient::Replication
 {
 	class FReplicationClient;
 	
diff --git a/Engine/Plugins/Developer/Concert/ConcertApp/MultiUserClient/Source/MultiUserClient/Private/Replication/Submission/External/ClientChangeConversionUtils.cpp b/Engine/Plugins/Developer/Concert/ConcertApp/MultiUserClient/Source/MultiUserClient/Private/Replication/Submission/External/ClientChangeConversionUtils.cpp
--- a/Engine/Plugins/Developer/Concert/ConcertApp/MultiUserClient/Source/MultiUserClient/Private/Replication/Submission/External/ClientChangeConversionUtils.cpp
+++ b/Engine/Plugins/Developer/Concert/ConcertApp/MultiUserClient/Source/MultiUserClient/Private/Replication/Submission/External/ClientChangeConversionUtils.cpp
@@ -10,7 +10,7 @@
 #include "Algo/RemoveIf.h"
 #include "Misc/Guid.h"
 
-namespace UE::MultiUserClient::ClientChangeConversionUtils
+namespace UE::MultiUserClient::Replication::ClientChangeConversionUtils
 {
 	namespace Private
 	{
diff --git a/Engine/Plugins/Developer/Concert/ConcertApp/MultiUserClient/Source/MultiUserClient/Private/Replication/Submission/External/ClientChangeConversionUtils.h b/Engine/Plugins/Developer/Concert/ConcertApp/MultiUserClient/Source/MultiUserClient/Private/Replication/Submission/External/ClientChangeConversionUtils.h
--- a/Engine/Plugins/Developer/Concert/ConcertApp/MultiUserClient/Source/MultiUserClient/Private/Replication/Submission/External/ClientChangeConversionUtils.h
+++ b/Engine/Plugins/Developer/Concert/ConcertApp/MultiUserClient/Source/MultiUserClient/Private/Replication/Submission/External/ClientChangeConversionUtils.h
@@ -14,10 +14,14 @@ struct FGuid;
 
 namespace UE::MultiUserClient
 {
-	enum class EAuthoritySubmissionResponseErrorCode : uint8;
-	
 	struct FChangeClientAuthorityResponse;
 	struct FChangeClientStreamResponse;
+}
+
+namespace UE::MultiUserClient::Replication
+{
+	enum class EAuthoritySubmissionResponseErrorCode : uint8;
+	
 	struct FSubmitAuthorityChangesResponse;
 	struct FSubmitStreamChangesResponse;
 }
@@ -25,7 +29,7 @@ namespace UE::MultiUserClient
 /**
  * Utils for transforming the simplified params in ChangeOperationTypes.h to the ones used by ConcertSyncCore.
  */
-namespace UE::MultiUserClient::ClientChangeConversionUtils
+namespace UE::MultiUserClient::Replication::ClientChangeConversionUtils
 {
 	/** Transforms FChangeStreamRequest to FConcertReplication_ChangeStream_Request, if the request is correctly formatted. */
 	TOptional<FConcertReplication_ChangeStream_Request> Transform(
diff --git a/Engine/Plugins/Developer/Concert/ConcertApp/MultiUserClient/Source/MultiUserClient/Private/Replication/Submission/External/ExternalClientChangeRequestHandler.h b/Engine/Plugins/Developer/Concert/ConcertApp/MultiUserClient/Source/MultiUserClient/Private/Replication/Submission/External/ExternalClientChangeRequestHandler.h
--- a/Engine/Plugins/Developer/Concert/ConcertApp/MultiUserClient/Source/MultiUserClient/Private/Replication/Submission/External/ExternalClientChangeRequestHandler.h
+++ b/Engine/Plugins/Developer/Concert/ConcertApp/MultiUserClient/Source/MultiUserClient/Private/Replication/Submission/External/ExternalClientChangeRequestHandler.h
@@ -11,15 +11,18 @@
 
 namespace UE::MultiUserClient
 {
-	class FClientChangeOperation;
-	class FReplicationClientManager;
-	class FSubmissionQueue;
 	class IClientChangeOperation;
+	struct FChangeClientReplicationRequest;
 	
 	enum class EChangeAuthorityOperationResult : uint8;
 	enum class EChangeStreamOperationResult : uint8;
-	
-	struct FChangeClientReplicationRequest;
+}
+
+namespace UE::MultiUserClient::Replication
+{
+	class FClientChangeOperation;
+	class FReplicationClientManager;
+	class FSubmissionQueue;
 	
 	/**
 	 * Handles change requests made by external modules via IMultiUserReplication::EnqueueChanges.
diff --git a/Engine/Plugins/Developer/Concert/ConcertApp/MultiUserClient/Source/MultiUserClient/Private/Widgets/ActiveSession/Replication/Client/ClientUtils.cpp b/Engine/Plugins/Developer/Concert/ConcertApp/MultiUserClient/Source/MultiUserClient/Private/Widgets/ActiveSession/Replication/Client/ClientUtils.cpp
--- a/Engine/Plugins/Developer/Concert/ConcertApp/MultiUserClient/Source/MultiUserClient/Private/Widgets/ActiveSession/Replication/Client/ClientUtils.cpp
+++ b/Engine/Plugins/Developer/Concert/ConcertApp/MultiUserClient/Source/MultiUserClient/Private/Widgets/ActiveSession/Replication/Client/ClientUtils.cpp
@@ -7,7 +7,7 @@
 #include "Replication/Client/ReplicationClientManager.h"
 #include "Widgets/Client/SClientName.h"
 
-namespace UE::MultiUserClient::ClientUtils
+namespace UE::MultiUserClient::Replication::ClientUtils
 {
 	FString GetClientDisplayName(const IConcertClient& InLocalClientInstance, const FGuid& InClientEndpointId)
 	{
diff --git a/Engine/Plugins/Developer/Concert/ConcertApp/MultiUserClient/Source/MultiUserClient/Private/Widgets/ActiveSession/Replication/Client/ClientUtils.h b/Engine/Plugins/Developer/Concert/ConcertApp/MultiUserClient/Source/MultiUserClient/Private/Widgets/ActiveSession/Replication/Client/ClientUtils.h
--- a/Engine/Plugins/Developer/Concert/ConcertApp/MultiUserClient/Source/MultiUserClient/Private/Widgets/ActiveSession/Replication/Client/ClientUtils.h
+++ b/Engine/Plugins/Developer/Concert/ConcertApp/MultiUserClient/Source/MultiUserClient/Private/Widgets/ActiveSession/Replication/Client/ClientUtils.h
@@ -11,13 +11,13 @@ class IConcertClientSession;
 struct FConcertClientInfo;
 struct FGuid;
 
-namespace UE::MultiUserClient
+namespace UE::MultiUserClient::Replication
 {
 	class FReplicationClient;
 	class FReplicationClientManager;
 }
 
-namespace UE::MultiUserClient::ClientUtils
+namespace UE::MultiUserClient::Replication::ClientUtils
 {
 	/**
 	 * Gets the display name for a client. Appends (me) if the client is local.
diff --git a/Engine/Plugins/Developer/Concert/ConcertApp/MultiUserClient/Source/MultiUserClient/Private/Widgets/ActiveSession/Replication/Client/Context/ContextMenuUtils.h b/Engine/Plugins/Developer/Concert/ConcertApp/MultiUserClient/Source/MultiUserClient/Private/Widgets/ActiveSession/Replication/Client/Context/ContextMenuUtils.h
--- a/Engine/Plugins/Developer/Concert/ConcertApp/MultiUserClient/Source/MultiUserClient/Private/Widgets/ActiveSession/Replication/Client/Context/ContextMenuUtils.h
+++ b/Engine/Plugins/Developer/Concert/ConcertApp/MultiUserClient/Source/MultiUserClient/Private/Widgets/ActiveSession/Replication/Client/Context/ContextMenuUtils.h
@@ -14,13 +14,13 @@ namespace UE::ConcertSharedSlate
 	class IMultiReplicationStreamEditor;
 	class IObjectHierarchyModel;
 }
-namespace UE::MultiUserClient
+namespace UE::MultiUserClient::Replication
 {
 	class FReassignObjectPropertiesLogic;
 	class FReplicationClientManager;
 }
 
-namespace UE::MultiUserClient::ContextMenuUtils
+namespace UE::MultiUserClient::Replication::ContextMenuUtils
 {
 	/** Adds menu entries for reassigning the object to another client. */
 	void AddReassignmentOptions(
diff --git a/Engine/Plugins/Developer/Concert/ConcertApp/MultiUserClient/Source/MultiUserClient/Private/Widgets/ActiveSession/Replication/Client/Context/ContextMenuUtils_Frequency.cpp b/Engine/Plugins/Developer/Concert/ConcertApp/MultiUserClient/Source/MultiUserClient/Private/Widgets/ActiveSession/Replication/Client/Context/ContextMenuUtils_Frequency.cpp
--- a/Engine/Plugins/Developer/Concert/ConcertApp/MultiUserClient/Source/MultiUserClient/Private/Widgets/ActiveSession/Replication/Client/Context/ContextMenuUtils_Frequency.cpp
+++ b/Engine/Plugins/Developer/Concert/ConcertApp/MultiUserClient/Source/MultiUserClient/Private/Widgets/ActiveSession/Replication/Client/Context/ContextMenuUtils_Frequency.cpp
@@ -19,7 +19,7 @@
 
 #define LOCTEXT_NAMESPACE "FrequencyContextMenuUtils"
 
-namespace UE::MultiUserClient::ContextMenuUtils
+namespace UE::MultiUserClient::Replication::ContextMenuUtils
 {
 	namespace Private
 	{
diff --git a/Engine/Plugins/Developer/Concert/ConcertApp/MultiUserClient/Source/MultiUserClient/Private/Widgets/ActiveSession/Replication/Client/Context/ContextMenuUtils_Reassignment.cpp b/Engine/Plugins/Developer/Concert/ConcertApp/MultiUserClient/Source/MultiUserClient/Private/Widgets/ActiveSession/Replication/Client/Context/ContextMenuUtils_Reassignment.cpp
--- a/Engine/Plugins/Developer/Concert/ConcertApp/MultiUserClient/Source/MultiUserClient/Private/Widgets/ActiveSession/Replication/Client/Context/ContextMenuUtils_Reassignment.cpp
+++ b/Engine/Plugins/Developer/Concert/ConcertApp/MultiUserClient/Source/MultiUserClient/Private/Widgets/ActiveSession/Replication/Client/Context/ContextMenuUtils_Reassignment.cpp
@@ -16,7 +16,7 @@
 
 #define LOCTEXT_NAMESPACE "ReassignmentContextMenuUtils"
 
-namespace UE::MultiUserClient::ContextMenuUtils
+namespace UE::MultiUserClient::Replication::ContextMenuUtils
 {
 	namespace Private
 	{
diff --git a/Engine/Plugins/Developer/Concert/ConcertApp/MultiUserClient/Source/MultiUserClient/Private/Widgets/ActiveSession/Replication/Client/FrequencyContextMenuUtils.h b/Engine/Plugins/Developer/Concert/ConcertApp/MultiUserClient/Source/MultiUserClient/Private/Widgets/ActiveSession/Replication/Client/FrequencyContextMenuUtils.h
--- a/Engine/Plugins/Developer/Concert/ConcertApp/MultiUserClient/Source/MultiUserClient/Private/Widgets/ActiveSession/Replication/Client/FrequencyContextMenuUtils.h
+++ b/Engine/Plugins/Developer/Concert/ConcertApp/MultiUserClient/Source/MultiUserClient/Private/Widgets/ActiveSession/Replication/Client/FrequencyContextMenuUtils.h
@@ -8,12 +8,12 @@
 class FMenuBuilder;
 struct FGuid;
 
-namespace UE::MultiUserClient
+namespace UE::MultiUserClient::Replication
 {
 	class FReplicationClientManager;
 }
 
-namespace UE::MultiUserClient::FrequencyContextMenuUtils
+namespace UE::MultiUserClient::Replication::FrequencyContextMenuUtils
 {
 	/** Adds an edit box for changing the selected object's frequencies. */
 	void AddFrequencyOptionsForSingleClient(FMenuBuilder& MenuBuilder, const FSoftObjectPath& ContextObjects, const FGuid& ClientId, FReplicationClientManager& InClientManager);
diff --git a/Engine/Plugins/Developer/Concert/ConcertApp/MultiUserClient/Source/MultiUserClient/Private/Widgets/ActiveSession/Replication/Client/Multi/Columns/MultiStreamColumns.h b/Engine/Plugins/Developer/Concert/ConcertApp/MultiUserClient/Source/MultiUserClient/Private/Widgets/ActiveSession/Replication/Client/Multi/Columns/MultiStreamColumns.h
--- a/Engine/Plugins/Developer/Concert/ConcertApp/MultiUserClient/Source/MultiUserClient/Private/Widgets/ActiveSession/Replication/Client/Multi/Columns/MultiStreamColumns.h
+++ b/Engine/Plugins/Developer/Concert/ConcertApp/MultiUserClient/Source/MultiUserClient/Private/Widgets/ActiveSession/Replication/Client/Multi/Columns/MultiStreamColumns.h
@@ -11,7 +11,7 @@
 
 class IConcertClient;
 
-namespace UE::MultiUserClient
+namespace UE::MultiUserClient::Replication
 {
 	class FMuteChangeTracker;
 	class FReassignObjectPropertiesLogic;
@@ -26,7 +26,7 @@ namespace UE::ConcertSharedSlate
 	class IReplicationStreamViewer;
 }
 
-namespace UE::MultiUserClient::MultiStreamColumns
+namespace UE::MultiUserClient::Replication::MultiStreamColumns
 {
 	const extern FName MuteToggleColumnId;
 	const extern FName AssignedClientsColumnId;
diff --git a/Engine/Plugins/Developer/Concert/ConcertApp/MultiUserClient/Source/MultiUserClient/Private/Widgets/ActiveSession/Replication/Client/Multi/Columns/MultiStreamColumns_AssignProperties.cpp b/Engine/Plugins/Developer/Concert/ConcertApp/MultiUserClient/Source/MultiUserClient/Private/Widgets/ActiveSession/Replication/Client/Multi/Columns/MultiStreamColumns_AssignProperties.cpp
--- a/Engine/Plugins/Developer/Concert/ConcertApp/MultiUserClient/Source/MultiUserClient/Private/Widgets/ActiveSession/Replication/Client/Multi/Columns/MultiStreamColumns_AssignProperties.cpp
+++ b/Engine/Plugins/Developer/Concert/ConcertApp/MultiUserClient/Source/MultiUserClient/Private/Widgets/ActiveSession/Replication/Client/Multi/Columns/MultiStreamColumns_AssignProperties.cpp
@@ -17,7 +17,7 @@
 
 #define LOCTEXT_NAMESPACE "AssignPropertyColumn"
 
-namespace UE::MultiUserClient::MultiStreamColumns
+namespace UE::MultiUserClient::Replication::MultiStreamColumns
 {
 	const FName AssignPropertyColumnId(TEXT("AssignPropertyColumn"));
 
diff --git a/Engine/Plugins/Developer/Concert/ConcertApp/MultiUserClient/Source/MultiUserClient/Private/Widgets/ActiveSession/Replication/Client/Multi/Columns/MultiStreamColumns_AssignedClients.cpp b/Engine/Plugins/Developer/Concert/ConcertApp/MultiUserClient/Source/MultiUserClient/Private/Widgets/ActiveSession/Replication/Client/Multi/Columns/MultiStreamColumns_AssignedClients.cpp
--- a/Engine/Plugins/Developer/Concert/ConcertApp/MultiUserClient/Source/MultiUserClient/Private/Widgets/ActiveSession/Replication/Client/Multi/Columns/MultiStreamColumns_AssignedClients.cpp
+++ b/Engine/Plugins/Developer/Concert/ConcertApp/MultiUserClient/Source/MultiUserClient/Private/Widgets/ActiveSession/Replication/Client/Multi/Columns/MultiStreamColumns_AssignedClients.cpp
@@ -15,7 +15,7 @@
 
 #define LOCTEXT_NAMESPACE "AssignedClientsColumnId"
 
-namespace UE::MultiUserClient::MultiStreamColumns
+namespace UE::MultiUserClient::Replication::MultiStreamColumns
 {
 	const FName AssignedClientsColumnId(TEXT("AssignedClientsColumn"));
 
diff --git a/Engine/Plugins/Developer/Concert/ConcertApp/MultiUserClient/Source/MultiUserClient/Private/Widgets/ActiveSession/Replication/Client/SPresetComboButton.h b/Engine/Plugins/Developer/Concert/ConcertApp/MultiUserClient/Source/MultiUserClient/Private/Widgets/ActiveSession/Replication/Client/SPresetComboButton.h
--- a/Engine/Plugins/Developer/Concert/ConcertApp/MultiUserClient/Source/MultiUserClient/Private/Widgets/ActiveSession/Replication/Client/SPresetComboButton.h
+++ b/Engine/Plugins/Developer/Concert/ConcertApp/MultiUserClient/Source/MultiUserClient/Private/Widgets/ActiveSession/Replication/Client/SPresetComboButton.h
@@ -6,16 +6,16 @@
 #include "Widgets/SCompoundWidget.h"
 #include "Widgets/Disconnected/ConcertClientSessionBrowserController.h"
 
-namespace UE::MultiUserClient
+namespace UE::MultiUserClient::Replication
 {
 	struct FSavePresetOptions;
 }
 
 class FMenuBuilder;
 class IConcertClient;
-namespace UE::MultiUserClient { enum class EApplyPresetFlags : uint8; }
+namespace UE::MultiUserClient::Replication { enum class EApplyPresetFlags : uint8; }
 
-namespace UE::MultiUserClient
+namespace UE::MultiUserClient::Replication
 {
 	class FPresetManager;
 
diff --git a/Engine/Plugins/Developer/Concert/ConcertApp/MultiUserClient/Source/MultiUserClient/Private/Widgets/ActiveSession/SActiveSessionRoot.cpp b/Engine/Plugins/Developer/Concert/ConcertApp/MultiUserClient/Source/MultiUserClient/Private/Widgets/ActiveSession/SActiveSessionRoot.cpp
--- a/Engine/Plugins/Developer/Concert/ConcertApp/MultiUserClient/Source/MultiUserClient/Private/Widgets/ActiveSession/SActiveSessionRoot.cpp
+++ b/Engine/Plugins/Developer/Concert/ConcertApp/MultiUserClient/Source/MultiUserClient/Private/Widgets/ActiveSession/SActiveSessionRoot.cpp
@@ -24,7 +24,7 @@ namespace UE::MultiUserClient
 	void SActiveSessionRoot::Construct(
 		const FArguments& InArgs,
 		TSharedPtr<IConcertSyncClient> InConcertSyncClient,
-		TSharedRef<FMultiUserReplicationManager> InReplicationManager
+		TSharedRef<Replication::FMultiUserReplicationManager> InReplicationManager
 		)
 	{
 		ChildSlot
@@ -61,7 +61,7 @@ namespace UE::MultiUserClient
 				]
 				+SWidgetSwitcher::Slot()
 				[
-					SNew(SReplicationRootWidget, InReplicationManager, InConcertSyncClient.ToSharedRef())
+					SNew(Replication::SReplicationRootWidget, InReplicationManager, InConcertSyncClient.ToSharedRef())
 				]
 			]
 		];
diff --git a/Engine/Plugins/Developer/Concert/ConcertApp/MultiUserClient/Source/MultiUserClient/Private/Widgets/ActiveSession/SActiveSessionRoot.h b/Engine/Plugins/Developer/Concert/ConcertApp/MultiUserClient/Source/MultiUserClient/Private/Widgets/ActiveSession/SActiveSessionRoot.h
--- a/Engine/Plugins/Developer/Concert/ConcertApp/MultiUserClient/Source/MultiUserClient/Private/Widgets/ActiveSession/SActiveSessionRoot.h
+++ b/Engine/Plugins/Developer/Concert/ConcertApp/MultiUserClient/Source/MultiUserClient/Private/Widgets/ActiveSession/SActiveSessionRoot.h
@@ -13,10 +13,10 @@ class SDockTab;
 class SWidgetSwitcher;
 class SWindow;
 
+namespace UE::MultiUserClient::Replication { class FMultiUserReplicationManager; }
+
 namespace UE::MultiUserClient
 {
-	class FMultiUserReplicationManager;
-	
 	/**
 	 * Displayed when the client is connected to an active session.
 	 * Manages the child content in tabs.
@@ -32,7 +32,7 @@ namespace UE::MultiUserClient
 		{}
 		SLATE_END_ARGS()
 
-		void Construct(const FArguments& InArgs, TSharedPtr<IConcertSyncClient> InConcertSyncClient, TSharedRef<FMultiUserReplicationManager> InReplicationManager);
+		void Construct(const FArguments& InArgs, TSharedPtr<IConcertSyncClient> InConcertSyncClient, TSharedRef<Replication::FMultiUserReplicationManager> InReplicationManager);
 
 	private:
 
diff --git a/Engine/Plugins/Developer/Concert/ConcertApp/MultiUserClient/Source/MultiUserClient/Private/Widgets/SConcertBrowser.cpp b/Engine/Plugins/Developer/Concert/ConcertApp/MultiUserClient/Source/MultiUserClient/Private/Widgets/SConcertBrowser.cpp
--- a/Engine/Plugins/Developer/Concert/ConcertApp/MultiUserClient/Source/MultiUserClient/Private/Widgets/SConcertBrowser.cpp
+++ b/Engine/Plugins/Developer/Concert/ConcertApp/MultiUserClient/Source/MultiUserClient/Private/Widgets/SConcertBrowser.cpp
@@ -20,7 +20,7 @@
 void SConcertBrowser::Construct(
 	const FArguments& InArgs,
 	TSharedRef<IConcertSyncClient> InSyncClient,
-	TSharedRef<UE::MultiUserClient::FMultiUserReplicationManager> InReplicationManager)
+	TSharedRef<UE::MultiUserClient::Replication::FMultiUserReplicationManager> InReplicationManager)
 {
 	if (!MultiUserClientUtils::HasServerCompatibleCommunicationPluginEnabled())
 	{
@@ -62,7 +62,7 @@ void SConcertBrowser::AttachChildWidget(EConcertConnectionStatus ConnectionStatu
 	{
 		if (ConnectionStatus == EConcertConnectionStatus::Connected)
 		{
-			if (const TSharedPtr<UE::MultiUserClient::FMultiUserReplicationManager> ReplicationManager = WeakReplicationManager.Pin()
+			if (const TSharedPtr<UE::MultiUserClient::Replication::FMultiUserReplicationManager> ReplicationManager = WeakReplicationManager.Pin()
 				; ensure(ReplicationManager))
 			{
 				ChildSlot.AttachWidget(
diff --git a/Engine/Plugins/Developer/Concert/ConcertApp/MultiUserClient/Source/MultiUserClient/Private/Widgets/SConcertBrowser.h b/Engine/Plugins/Developer/Concert/ConcertApp/MultiUserClient/Source/MultiUserClient/Private/Widgets/SConcertBrowser.h
--- a/Engine/Plugins/Developer/Concert/ConcertApp/MultiUserClient/Source/MultiUserClient/Private/Widgets/SConcertBrowser.h
+++ b/Engine/Plugins/Developer/Concert/ConcertApp/MultiUserClient/Source/MultiUserClient/Private/Widgets/SConcertBrowser.h
@@ -9,13 +9,9 @@
 #include "IConcertModule.h"  // Change to use Fwd or Ptr.h?
 #include "ConcertMessages.h"
 
-namespace UE::MultiUserClient
-{
-	class FMultiUserReplicationManager;
-}
-
 class IConcertClientSession;
 class IConcertSyncClient;
+namespace UE::MultiUserClient::Replication { class FMultiUserReplicationManager; }
 
 /**
  * Displays the multi-users windows enabling the user to browse active and archived sessions,
@@ -36,14 +32,14 @@ public:
 	* @param InSyncClient The sync client.
 	* @param InReplicationManager Used to create replication UI when in an active session
 	*/
-	void Construct(const FArguments& InArgs, TSharedRef<IConcertSyncClient> InSyncClient, TSharedRef<UE::MultiUserClient::FMultiUserReplicationManager> InReplicationManager);
+	void Construct(const FArguments& InArgs, TSharedRef<IConcertSyncClient> InSyncClient, TSharedRef<UE::MultiUserClient::Replication::FMultiUserReplicationManager> InReplicationManager);
 
 private:
 
 	/** Keeps the sync client interface. */
 	TWeakPtr<IConcertSyncClient> WeakConcertSyncClient;
 	/** Interacts with the replication system on behalf of Multi-User. */
-	TWeakPtr<UE::MultiUserClient::FMultiUserReplicationManager> WeakReplicationManager;
+	TWeakPtr<UE::MultiUserClient::Replication::FMultiUserReplicationManager> WeakReplicationManager;
 
 	/** Keeps the session browser searched text in memory to reapply it when a user leaves a session and goes back to the session browser. */
 	TSharedPtr<FText> SearchedText;
