Renames client related types:
- FReplicationClientManager to FOnlineClientManager,
- FReplicationClient to FOnlineClient,
- FRemoteReplicationClient to FRemoteClient, and
- FLocalReplicationClient to FLocalClient.

Also moves them into the Client/Online folder.
This is in parparation for a change introducing offline clients.

diff --git a/Engine/Plugins/Developer/Concert/ConcertApp/MultiUserClient/Source/MultiUserClient/Private/Replication/Client/LocalReplicationClient.cpp b/Engine/Plugins/Developer/Concert/ConcertApp/MultiUserClient/Source/MultiUserClient/Private/Replication/Client/LocalReplicationClient.cpp
deleted file mode 100644
diff --git a/Engine/Plugins/Developer/Concert/ConcertApp/MultiUserClient/Source/MultiUserClient/Private/Replication/Client/LocalReplicationClient.h b/Engine/Plugins/Developer/Concert/ConcertApp/MultiUserClient/Source/MultiUserClient/Private/Replication/Client/LocalReplicationClient.h
deleted file mode 100644
diff --git a/Engine/Plugins/Developer/Concert/ConcertApp/MultiUserClient/Source/MultiUserClient/Private/Replication/Client/Online/LocalClient.cpp b/Engine/Plugins/Developer/Concert/ConcertApp/MultiUserClient/Source/MultiUserClient/Private/Replication/Client/Online/LocalClient.cpp
new file mode 100644
--- /dev/null
+++ b/Engine/Plugins/Developer/Concert/ConcertApp/MultiUserClient/Source/MultiUserClient/Private/Replication/Client/Online/LocalClient.cpp
@@ -0,0 +1,28 @@
+// Copyright Epic Games, Inc. All Rights Reserved.
+
+#include "LocalClient.h"
+
+#include "IConcertSyncClient.h"
+#include "Replication/Authority/AuthoritySynchronizer_LocalClient.h"
+#include "Replication/Submission/SubmissionWorkflow_LocalClient.h"
+
+namespace UE::MultiUserClient::Replication
+{
+	FLocalClient::FLocalClient(
+		FReplicationDiscoveryContainer& InDiscoveryContainer,
+		FGlobalAuthorityCache& InAuthorityCache,
+		UMultiUserReplicationStream& InClientStreamContent,
+		TUniquePtr<IClientStreamSynchronizer> InStreamSynchronizer,
+		TSharedRef<IConcertSyncClient> InClient
+		)
+		: FOnlineClient(
+			InClient->GetConcertClient()->GetCurrentSession()->GetSessionClientEndpointId(),
+			InDiscoveryContainer,
+			InAuthorityCache,
+			InClientStreamContent,
+			MoveTemp(InStreamSynchronizer),
+			MakeUnique<FAuthoritySynchronizer_LocalClient>(InClient),
+			MakeUnique<FSubmissionWorkflow_LocalClient>(MoveTemp(InClient)))
+		, RemoteSubmissionListener(InClient->GetConcertClient()->GetCurrentSession().ToSharedRef(), GetStreamSynchronizer(), GetSubmissionQueue())
+	{}
+}
diff --git a/Engine/Plugins/Developer/Concert/ConcertApp/MultiUserClient/Source/MultiUserClient/Private/Replication/Client/Online/LocalClient.h b/Engine/Plugins/Developer/Concert/ConcertApp/MultiUserClient/Source/MultiUserClient/Private/Replication/Client/Online/LocalClient.h
new file mode 100644
--- /dev/null
+++ b/Engine/Plugins/Developer/Concert/ConcertApp/MultiUserClient/Source/MultiUserClient/Private/Replication/Client/Online/LocalClient.h
@@ -0,0 +1,38 @@
+// Copyright Epic Games, Inc. All Rights Reserved.
+
+#pragma once
+
+#include "OnlineClient.h"
+#include "Replication/Submission/Remote/RemoteSubmissionListener.h"
+
+namespace UE::MultiUserClient::Replication
+{
+	class FGlobalAuthorityCache;
+
+	/** Holds extra information about a local replication client. */
+	class FLocalClient : public FOnlineClient
+	{
+	public:
+		
+		/**
+		 * @param InDiscoveryContainer Used for auto-discovering properties added to this client's stream. The caller ensures it outlives the constructed instance.
+		 * @param InAuthorityCache Caches authority state of all clients. Passed to subsystems. The caller ensures it outlives the constructed instance.
+		 * @param InClientStreamContent Object that this client's stream changes are to be written into. The caller ensures it outlives the constructed instance.
+		 * @param InStreamSynchronizer Implementation for obtaining stream registered on server. The constructed instance takes ownership.
+		 * @param InClient The local editor's running client. Passed to subsystems so they can send messages to the current session. Keeps strong reference.
+		 */
+		FLocalClient(
+			FReplicationDiscoveryContainer& InDiscoveryContainer,
+			FGlobalAuthorityCache& InAuthorityCache,
+			UMultiUserReplicationStream& InClientStreamContent,
+			TUniquePtr<IClientStreamSynchronizer> InStreamSynchronizer,
+			TSharedRef<IConcertSyncClient> InClient
+			);
+
+	private:
+		
+		/** Listens for and handles for submission request made by remote clients' SubmissionWorkflow.*/
+		FRemoteSubmissionListener RemoteSubmissionListener;
+	};
+}
+
diff --git a/Engine/Plugins/Developer/Concert/ConcertApp/MultiUserClient/Source/MultiUserClient/Private/Replication/Client/Online/OnlineClient.cpp b/Engine/Plugins/Developer/Concert/ConcertApp/MultiUserClient/Source/MultiUserClient/Private/Replication/Client/Online/OnlineClient.cpp
new file mode 100644
--- /dev/null
+++ b/Engine/Plugins/Developer/Concert/ConcertApp/MultiUserClient/Source/MultiUserClient/Private/Replication/Client/Online/OnlineClient.cpp
@@ -0,0 +1,208 @@
+// Copyright Epic Games, Inc. All Rights Reserved.
+
+#include "OnlineClient.h"
+
+#include "Assets/MultiUserReplicationClientContent.h"
+#include "Replication/Editor/Model/IEditableReplicationStreamModel.h"
+#include "Replication/ReplicationWidgetFactories.h"
+#include "Replication/Stream/Discovery/MultiUserStreamExtender.h"
+#include "Replication/Stream/StreamChangeTracker.h"
+#include "Replication/Submission/Data/AuthoritySubmission.h"
+#include "Settings/MultiUserReplicationSettings.h"
+
+#include "Misc/CoreDelegates.h"
+#include "Replication/ClientReplicationWidgetFactories.h"
+
+namespace UE::MultiUserClient::Replication
+{
+	FOnlineClient::FOnlineClient(
+		const FGuid& EndpointId,
+		FReplicationDiscoveryContainer& InDiscoveryContainer,
+		FGlobalAuthorityCache& InAuthorityCache,
+		UMultiUserReplicationStream& InClientStreamContent,
+		TUniquePtr<IClientStreamSynchronizer> InStreamSynchronizer,
+		TUniquePtr<IClientAuthoritySynchronizer> InAuthoritySynchronizer,
+		TUniquePtr<ISubmissionWorkflow> InSubmissionWorkflow
+		)
+		: EndpointId(EndpointId)
+		, ClientStreamContent(&InClientStreamContent)
+		, StreamSynchronizer(MoveTemp(InStreamSynchronizer))
+		, AuthoritySynchronizer(MoveTemp(InAuthoritySynchronizer))
+		, SubmissionWorkflow(MoveTemp(InSubmissionWorkflow))
+		, SubmissionQueue(*SubmissionWorkflow)
+		, ExternalRequestHandler(
+			StreamSynchronizer->GetStreamId(),
+			FGetStreamContent::CreateLambda([this](){ return &StreamSynchronizer->GetServerState(); }),
+			SubmissionQueue
+			)
+		, StreamExtender(MakeShared<FMultiUserStreamExtender>(EndpointId, InDiscoveryContainer))
+		, LocalClientEditModel(
+			// Transact ClientContentStorage
+			ConcertClientSharedSlate::CreateTransactionalStreamModel(
+				// Read & write the stream data in ClientContentStorage
+				CreateBaseStreamModel(GetClientStreamObject()->MakeReplicationMapGetterAttribute(), StreamExtender),
+				*GetClientStreamObject()
+				)
+			)
+		, LocalClientStreamDiffer(
+			GetStreamSynchronizer(),
+			GetClientStreamObject()->MakeReplicationMapGetterAttribute(),
+			FStreamChangeTracker::FOnModifyReplicationMap::CreateLambda([this](){ GetClientStreamObject()->Modify(); })
+			)
+		, LocalAuthorityDiffer(EndpointId, *AuthoritySynchronizer, InAuthorityCache)
+		, LocalFrequencyChangeTracker(*StreamSynchronizer)
+		, ChangeRequestBuilder(EndpointId, InAuthorityCache, *StreamSynchronizer, LocalClientStreamDiffer, LocalAuthorityDiffer, LocalFrequencyChangeTracker)
+		, AutoSubmissionPolicy(SubmissionQueue, ChangeRequestBuilder, LocalClientEditModel.Get(), LocalAuthorityDiffer, LocalFrequencyChangeTracker)
+		, LevelModificationHandler(GetClientEditModel().Get())
+	{
+		LocalClientEditModel->OnObjectsChanged().AddRaw(this, &FOnlineClient::OnObjectsChanged);
+		LocalClientEditModel->OnPropertiesChanged().AddRaw(this, &FOnlineClient::OnPropertiesChanged);
+		LocalAuthorityDiffer.OnChangedOwnedObjects().AddRaw(this, &FOnlineClient::DeferOnModelChanged);
+
+		SubmissionWorkflow->OnAuthorityRequestCompleted_AnyThread().AddRaw(this, &FOnlineClient::OnAuthoritySubmissionCompleted);
+		StreamSynchronizer->OnServerStateChanged().AddRaw(this, &FOnlineClient::OnServerStateChanged);
+
+		LevelModificationHandler.OnHierarchyNeedsRefresh().AddLambda([this]()
+		{
+			OnHierarchyNeedsRefreshDelegate.Broadcast();
+		});
+	}
+
+	FOnlineClient::~FOnlineClient()
+	{
+		FCoreDelegates::OnEndFrame.RemoveAll(this);
+	}
+
+	bool FOnlineClient::AllowsEditing() const
+	{
+		return CanEverSubmit(SubmissionWorkflow->GetUploadability()); 
+	}
+
+	void FOnlineClient::OnObjectsChanged(
+		TConstArrayView<UObject*> AddedObjects,
+		TConstArrayView<FSoftObjectPath> RemovedObjects,
+		ConcertSharedSlate::EReplicatedObjectChangeReason ReplicatedObjectChangeReason
+		)
+	{
+		DeferOnModelChanged(AddedObjects);
+	}
+
+	void FOnlineClient::OnPropertiesChanged()
+	{
+		DeferOnModelChanged();
+	}
+
+	void FOnlineClient::OnServerStateChanged()
+	{
+		// Whenever this client's server state changes, the UI must be refreshed.
+
+		// The UI adds empty actors. However, we never send them to the server...
+		TSet<FSoftObjectPath> StagedObjects;
+		for (const TPair<FSoftObjectPath, FConcertReplicatedObjectInfo>& Pair : GetClientStreamObject()->ReplicationMap.ReplicatedObjects)
+		{
+			if (Pair.Value.PropertySelection.ReplicatedProperties.IsEmpty())
+			{
+				StagedObjects.Add(Pair.Key);
+			}
+		}
+		
+		// ... if the user removes the last property from the entire actor-component hierarchy, we want the hierarchy to continue to displayed...
+		GetClientStreamObject()->ReplicationMap = GetStreamSynchronizer().GetServerState();
+		// ... so add back the staged objects
+		for (const FSoftObjectPath& StagedObject : StagedObjects)
+		{
+			if (const UObject* Object = StagedObject.ResolveObject())
+			{
+				GetClientStreamObject()->ReplicationMap.ReplicatedObjects.Add(StagedObject) = FConcertReplicatedObjectInfo::Make(*Object);
+			}
+		}
+		// To remove the hierarchy, the user must click the actor and delete it explicitly, which will call IEditableReplicationStream::RemoveObjects on the staged objects.
+		
+		DeferOnModelChanged();
+	}
+
+	void FOnlineClient::DeferOnModelChanged(TConstArrayView<UObject*> AddedObjects)
+	{
+		if (!DeferredOnModelChangedData)
+		{
+			DeferredOnModelChangedData.Emplace();
+			FCoreDelegates::OnEndFrame.AddRaw(this, &FOnlineClient::ProcessOnModelChanged);
+		}
+
+		Algo::Transform(AddedObjects, DeferredOnModelChangedData->AccumulatedAddedObjects, [](UObject* Object)
+		{
+			return Object;
+		});
+	}
+
+	void FOnlineClient::ProcessOnModelChanged()
+	{
+		check(DeferredOnModelChangedData);
+		const FDeferredOnModelChangedData ChangeData = MoveTemp(*DeferredOnModelChangedData);
+		DeferredOnModelChangedData.Reset();
+		FCoreDelegates::OnEndFrame.RemoveAll(this);
+		
+		// Could improve performance by just considering what actually changed instead of doing a full rebuild
+		// This must be done before SetAuthorityIfAllowed because it uses the cache for checking whether the object has properties assigned
+		LocalClientStreamDiffer.RefreshChangesCache();
+		
+		// Better UX for user: automatically take authority for newly added objects (but only if it is allowed and causes no conflicts)
+		TakeAuthorityOverNewlyAddedObjects(ChangeData);
+		// Refresh because local authority changes may no longer be valid after modifying the stream
+		LocalAuthorityDiffer.RefreshChanges();
+		// Check whether added objects are supposed to default to some special frequency settings
+		ApplyDefaultFrequencySettings(ChangeData);
+
+		// Finally, let everybody else know.
+		OnModelChangedDelegate.Broadcast();
+		AutoSubmissionPolicy.ProcessAccumulatedChangesAndSubmit();
+	}
+	
+	void FOnlineClient::TakeAuthorityOverNewlyAddedObjects(const FDeferredOnModelChangedData& ChangeData)
+	{
+		TArray<FSoftObjectPath> ObjectPaths;
+		Algo::TransformIf(ChangeData.AccumulatedAddedObjects, ObjectPaths,
+			[](const TWeakObjectPtr<UObject>& Object)
+			{
+				// The object might have been made invalid last frame.
+				return Object.IsValid();
+			},
+			[](const TWeakObjectPtr<UObject>& Object)
+			{
+				return FSoftObjectPath(Object.Get()) ;
+			});
+		LocalAuthorityDiffer.SetAuthorityIfAllowed(ObjectPaths, true);
+	}
+
+	void FOnlineClient::ApplyDefaultFrequencySettings(const FDeferredOnModelChangedData& ChangeData)
+	{
+		UMultiUserReplicationSettings* Settings = UMultiUserReplicationSettings::Get();
+		for (const TWeakObjectPtr<UObject>& AddedObject : ChangeData.AccumulatedAddedObjects)
+		{
+			if (!AddedObject.IsValid())
+			{
+				continue;
+			}
+			
+			const TOptional<FConcertObjectReplicationSettings> FrequencyOverride = Settings->DetermineObjectFrequencySettings(*AddedObject);
+			if (FrequencyOverride)
+			{
+				LocalFrequencyChangeTracker.AddOverride(AddedObject.Get(), *FrequencyOverride);
+			}
+		}
+	}
+
+	void FOnlineClient::OnAuthoritySubmissionCompleted(const FSubmitAuthorityChangesRequest& Request, const FSubmitAuthorityChangesResponse& Response)
+	{
+		if (!Response.Response)
+		{
+			return;
+		}
+
+		// Use case: You and another client submit at the same time. You lose. Revert your local changes so the checkboxes accurately reflect the authority state.
+		for (const TPair<FSoftObjectPath, FConcertStreamArray>& RejectedObjectPair : Response.Response.GetValue().RejectedObjects)
+		{
+			LocalAuthorityDiffer.ClearAuthorityChange({ RejectedObjectPair.Key });
+		}
+	}
+}
diff --git a/Engine/Plugins/Developer/Concert/ConcertApp/MultiUserClient/Source/MultiUserClient/Private/Replication/Client/Online/OnlineClient.h b/Engine/Plugins/Developer/Concert/ConcertApp/MultiUserClient/Source/MultiUserClient/Private/Replication/Client/Online/OnlineClient.h
new file mode 100644
--- /dev/null
+++ b/Engine/Plugins/Developer/Concert/ConcertApp/MultiUserClient/Source/MultiUserClient/Private/Replication/Client/Online/OnlineClient.h
@@ -0,0 +1,200 @@
+// Copyright Epic Games, Inc. All Rights Reserved.
+
+#pragma once
+
+#include "Assets/MultiUserReplicationClientContent.h"
+#include "Replication/Authority/AuthorityChangeTracker.h"
+#include "Replication/Authority/IClientAuthoritySynchronizer.h"
+#include "Replication/Editor/UnrealEditor/ModifyObjectInLevelHandler.h"
+#include "Replication/Misc/Frequency/FrequencyChangeTracker.h"
+#include "Replication/Stream/IClientStreamSynchronizer.h"
+#include "Replication/Stream/StreamChangeTracker.h"
+#include "Replication/Submission/AutoSubmissionPolicy.h"
+#include "Replication/Submission/ChangeRequestBuilder.h"
+#include "Replication/Submission/External/ExternalClientChangeRequestHandler.h"
+#include "Replication/Submission/ISubmissionWorkflow.h"
+#include "Replication/Submission/Queue/SubmissionQueue.h"
+
+#include "Templates/SharedPointer.h"
+#include "Templates/UnrealTemplate.h"
+
+class IConcertClient;
+class UMultiUserReplicationClientContent;
+
+namespace UE::ConcertSharedSlate
+{
+	class IEditableReplicationStreamModel;
+	enum class EReplicatedObjectChangeReason : uint8;
+}
+
+namespace UE::MultiUserClient::Replication
+{
+	class FMultiUserStreamExtender;
+	class FReplicationDiscoveryContainer;
+	class ISubmissionWorkflow;
+	
+	/**
+	 * Initializes systems that are common for local or remote clients.
+	 * This class' responsibility is to initialize all systems that exist for the life time of a client in a session.
+	 */
+	class FOnlineClient : public FNoncopyable
+	{
+	public:
+		
+		/**
+		 * @param EndpointId The endpoint ID this instance corresponds to.
+		 * @param InDiscoveryContainer Used for auto-discovering properties added to this client's stream. The caller ensures it outlives the constructed instance.
+		 * @param InAuthorityCache Caches authority state of all clients. Passed to subsystems. The caller ensures it outlives the constructed instance.
+		 * @param InClientStreamContent Object that this client's stream changes are to be written into. The caller ensures it outlives the constructed instance.
+		 * @param InStreamSynchronizer Implementation for obtaining stream registered on server. The constructed instance takes ownership.
+		 * @param InAuthoritySynchronizer Implementation for obtaining the client's authority state on server. The constructed instance takes ownership.
+		 * @param InSubmissionWorkflow Implementation for for changing client streams and authority on the server. The constructed instance takes ownership.
+		 */
+		FOnlineClient(
+			const FGuid& EndpointId,
+			FReplicationDiscoveryContainer& InDiscoveryContainer UE_LIFETIMEBOUND,
+			FGlobalAuthorityCache& InAuthorityCache UE_LIFETIMEBOUND,
+			UMultiUserReplicationStream& InClientStreamContent UE_LIFETIMEBOUND,
+			TUniquePtr<IClientStreamSynchronizer> InStreamSynchronizer,
+			TUniquePtr<IClientAuthoritySynchronizer> InAuthoritySynchronizer,
+			TUniquePtr<ISubmissionWorkflow> InSubmissionWorkflow
+			);
+		~FOnlineClient();
+
+		UMultiUserReplicationStream* GetClientStreamObject() const { return ClientStreamContent; }
+		FMultiUserStreamExtender& GetStreamExtender() { return *StreamExtender; }
+		/**
+		 * This is used so the UI can construct the IReplicationStreamEditor.
+		 * @see CreateBaseStreamEditor and FCreateEditorParams.
+		 * 
+		 * @note You must make sure to release this object when this FReplicationClient is destroyed.
+		 * Listen for events on the owning FReplicationClientManager::OnPreRemoteClientRemoved and FMultiUserReplicationManager::OnLeaveSession
+		 * @see FReplicationClientManager and FMultiUserReplicationManager
+		 */
+		TSharedRef<ConcertSharedSlate::IEditableReplicationStreamModel> GetClientEditModel() const { return LocalClientEditModel; }
+		IClientStreamSynchronizer& GetStreamSynchronizer() const { return *StreamSynchronizer.Get(); }
+		IClientAuthoritySynchronizer& GetAuthoritySynchronizer() const { return *AuthoritySynchronizer.Get(); }
+		
+		const FStreamChangeTracker& GetStreamDiffer() const { return LocalClientStreamDiffer; }
+		FStreamChangeTracker& GetStreamDiffer() { return LocalClientStreamDiffer; }
+		
+		const FAuthorityChangeTracker& GetAuthorityDiffer() const { return LocalAuthorityDiffer; }
+		FAuthorityChangeTracker& GetAuthorityDiffer() { return LocalAuthorityDiffer; }
+
+		/** @return The object with which you can submit changes. You can use it to listen to general events. If you want to submit changes, you should prefer to use the SubmissionQueue. */
+		const ISubmissionWorkflow& GetSubmissionWorkflow() const { return *SubmissionWorkflow; }
+		ISubmissionWorkflow& GetSubmissionWorkflow() { return *SubmissionWorkflow; }
+		/** @return Implements simple game-thread based queue for SubmissionWorkflow. Only one submission can be in progress at any given time. */
+		FSubmissionQueue& GetSubmissionQueue() { return SubmissionQueue; }
+		/** @return Gets the adapter handling requests received by IMultiUserReplication::EnqueuesChanges */
+		FExternalClientChangeRequestHandler& GetExternalRequestHandler() { return ExternalRequestHandler; }
+
+		/** @return The endpoint ID of this client in the Concert session. */
+		const FGuid& GetEndpointId() const { return EndpointId; }
+		/** @return Whether it is allowed to edit the stream and authority for this client. */
+		bool AllowsEditing() const;
+		
+		bool operator==(const FOnlineClient& Client) const { return GetEndpointId() == Client.GetEndpointId(); }
+
+		DECLARE_MULTICAST_DELEGATE(FOnModelExternallyChanged);
+		/**
+		 * Broadcast when the data underlying the model has changed for any reason:
+		 * - Edited directly by client
+		 * - Changed by transaction
+		 * - Server state changed
+		 */
+		FOnModelExternallyChanged& OnModelChanged() { return OnModelChangedDelegate; }
+		
+		DECLARE_MULTICAST_DELEGATE(FOnHiearchyNeedsRefresh);
+		/** Broadcasts when the hierarchy may have changed. */
+		FOnHiearchyNeedsRefresh& OnHierarchyNeedsRefresh() { return OnHierarchyNeedsRefreshDelegate; }
+		
+	private:
+
+		/** This client's Concert Endpoint ID. */
+		const FGuid EndpointId;
+		
+		/** The state of the server is synched up with this object and displayed in the UI. */
+		TObjectPtr<UMultiUserReplicationStream> ClientStreamContent;
+		
+		/** Keeps the client's stream state on the server in sync. */
+		TUniquePtr<IClientStreamSynchronizer> StreamSynchronizer;
+		/** Keeps the client's authority state on the server in sync.*/
+		TUniquePtr<IClientAuthoritySynchronizer> AuthoritySynchronizer;
+		/** Handles the logic of submitting and reverting for this client. */
+		TUniquePtr<ISubmissionWorkflow> SubmissionWorkflow;
+
+		/** Allows systems to queue pending submissions to SubmissionWorkflow in case a submission is in progress. */
+		FSubmissionQueue SubmissionQueue;
+		/** Enqueues external change requests into the SubmissionQueue. */
+		FExternalClientChangeRequestHandler ExternalRequestHandler;
+
+		/** Injected into LocalClientEditModel during construction. Adds additional objects and properties when a new object is added to the stream. */
+		TSharedRef<FMultiUserStreamExtender> StreamExtender;
+		/**
+		 * Used to detect changes made to the client's config by the local editor.
+		 * Those changes can later be applied to the client.
+		 * 
+		 * The SessionContent UObject must be kept alive as long as this model lives.
+		 * 
+		 * It should be noted that the UI keeps a strong reference to this model.
+		 * The UI is destroyed right before FReplicationClient, which is guaranteed by FMultiUserReplicationManager.
+		 * @see FMultiUserReplicationManager::OnLeaveSession
+		 */
+		TSharedRef<ConcertSharedSlate::IEditableReplicationStreamModel> LocalClientEditModel;
+		
+		/** Tracks changes made to server's state of the client's streams and prepares to upload them using StreamSynchronizer. */
+		FStreamChangeTracker LocalClientStreamDiffer;
+		/** Tracks changes made to the client's authority state. */
+		FAuthorityChangeTracker LocalAuthorityDiffer;
+		/** Tracks enqueued local changes to object replication frequency. */
+		FFrequencyChangeTracker LocalFrequencyChangeTracker;
+		
+		/** Shared logic for building stream and authority change requests based on local change made. */
+		FChangeRequestBuilder ChangeRequestBuilder;
+		/** Automatically submits changes as they are made by the user. */
+		FAutoSubmissionPolicy AutoSubmissionPolicy;
+		
+		/** Handles updating the local client's stream when the level is modified, e.g. when an actor is removed. */
+		ConcertClientSharedSlate::FModifyObjectInLevelHandler LevelModificationHandler;
+
+		/** Broadcasts when the hierarchy may have changed. */
+		FOnHiearchyNeedsRefresh OnHierarchyNeedsRefreshDelegate;
+
+		/**
+		 * Broadcast when the data underlying the model has changed for any reason:
+		 * - Edited directly by client
+		 * - Changed by transaction
+		 * - Server state changed
+		 */
+		FOnModelExternallyChanged OnModelChangedDelegate;
+
+		struct FDeferredOnModelChangedData
+		{
+			TSet<TWeakObjectPtr<UObject>> AccumulatedAddedObjects;
+		};
+		TOptional<FDeferredOnModelChangedData> DeferredOnModelChangedData;
+
+		// Respond to model changing
+		void OnObjectsChanged(TConstArrayView<UObject*> AddedObjects, TConstArrayView<FSoftObjectPath> RemovedObjects, ConcertSharedSlate::EReplicatedObjectChangeReason ReplicatedObjectChangeReason);
+		void OnPropertiesChanged();
+
+		void OnServerStateChanged();
+
+		/** Defers rebuilding operations, such as refreshing state and calling OnModelChanged, in case there are multiple changes in the same frame. */
+		void DeferOnModelChanged() { DeferOnModelChanged({}); }
+		void DeferOnModelChanged(TConstArrayView<UObject*> AddedObjects);
+		/** Processes all changes that have happened to the stream this frame. */
+		void ProcessOnModelChanged();
+		
+		/** Takes authority over newly added objects for better UX */
+		void TakeAuthorityOverNewlyAddedObjects(const FDeferredOnModelChangedData& ChangeData);
+
+		/** Applies default replication frequency settings specified in the Multi User settings */
+		void ApplyDefaultFrequencySettings(const FDeferredOnModelChangedData& ChangeData);
+		
+		/** Removes authority if request fails */
+		void OnAuthoritySubmissionCompleted(const FSubmitAuthorityChangesRequest& Request, const FSubmitAuthorityChangesResponse& Response);
+	};
+}
+
diff --git a/Engine/Plugins/Developer/Concert/ConcertApp/MultiUserClient/Source/MultiUserClient/Private/Replication/Client/Online/OnlineClientManager.cpp b/Engine/Plugins/Developer/Concert/ConcertApp/MultiUserClient/Source/MultiUserClient/Private/Replication/Client/Online/OnlineClientManager.cpp
new file mode 100644
--- /dev/null
+++ b/Engine/Plugins/Developer/Concert/ConcertApp/MultiUserClient/Source/MultiUserClient/Private/Replication/Client/Online/OnlineClientManager.cpp
@@ -0,0 +1,200 @@
+// Copyright Epic Games, Inc. All Rights Reserved.
+
+#include "OnlineClientManager.h"
+
+#include "Assets/MultiUserReplicationStream.h"
+#include "IConcertSyncClient.h"
+#include "LocalClient.h"
+#include "RemoteClient.h"
+#include "Replication/Stream/StreamSynchronizer_LocalClient.h"
+#include "Replication/Stream/MultiUserStreamId.h"
+
+#include "UObject/Package.h"
+#include "UObject/UObjectGlobals.h"
+
+namespace UE::MultiUserClient::Replication
+{
+	namespace Private
+	{
+		static UMultiUserReplicationStream* MakeClientContent()
+		{
+			UMultiUserReplicationStream* Stream = NewObject<UMultiUserReplicationStream>(GetTransientPackage(), NAME_None, RF_Transient | RF_Transactional);
+			Stream->StreamId = MultiUserStreamID;
+			return Stream;
+		}
+	}
+	
+	FOnlineClientManager::FOnlineClientManager(
+		const TSharedRef<IConcertSyncClient>& InClient,
+		const TSharedRef<IConcertClientSession>& InSession,
+		FReplicationDiscoveryContainer& InRegisteredExtenders,
+		FStreamAndAuthorityQueryService& InQueryService
+		)
+		: ConcertClient(InClient)
+		, Session(InSession)
+		, RegisteredExtenders(InRegisteredExtenders)
+		, QueryService(InQueryService)
+		, AuthorityCache(*this)
+		, LocalClient([this, InClient]()
+		{
+			UMultiUserReplicationStream* ClientPreset = Private::MakeClientContent();
+			return FLocalClient(RegisteredExtenders, AuthorityCache, *ClientPreset, MakeUnique<FStreamSynchronizer_LocalClient>(InClient, ClientPreset->StreamId), InClient);
+		}())
+		, ReassignmentLogic(*this)
+	{
+		AuthorityCache.RegisterEvents();
+		InSession->OnSessionClientChanged().AddRaw(this, &FOnlineClientManager::OnSessionClientChanged);
+
+		for (const FGuid& ClientEndpointId : InSession->GetSessionClientEndpointIds())
+		{
+			CreateRemoteClient(ClientEndpointId);
+		}
+	}
+
+	FOnlineClientManager::~FOnlineClientManager()
+	{
+		if (const TSharedPtr<IConcertClientSession> SessionPin = Session.Pin())
+		{
+			SessionPin->OnSessionClientChanged().RemoveAll(this);
+		}
+	}
+
+	TArray<TNonNullPtr<const FRemoteClient>> FOnlineClientManager::GetRemoteClients() const
+	{
+		TArray<TNonNullPtr<const FRemoteClient>> Result;
+		Algo::Transform(RemoteClients, Result, [](const TUniquePtr<FRemoteClient>& Client) -> TNonNullPtr<const FRemoteClient>
+		{
+			return Client.Get();
+		});
+		return Result;
+	}
+
+	TArray<TNonNullPtr<FRemoteClient>> FOnlineClientManager::GetRemoteClients()
+	{
+		TArray<TNonNullPtr<FRemoteClient>> Result;
+		Algo::Transform(RemoteClients, Result, [](const TUniquePtr<FRemoteClient>& Client) -> TNonNullPtr<FRemoteClient>
+		{
+			return Client.Get();
+		});
+		return Result;
+	}
+
+	TArray<TNonNullPtr<const FOnlineClient>> FOnlineClientManager::GetClients(TFunctionRef<bool(const FOnlineClient& Client)> Predicate) const
+	{
+		TArray<TNonNullPtr<const FOnlineClient>> Result;
+		ForEachClient([&Predicate, &Result](const FOnlineClient& Client)
+		{
+			if (Predicate(Client))
+			{
+				Result.Add(&Client);
+			}
+			return EBreakBehavior::Continue;
+		});
+		return Result;
+	}
+
+	const FRemoteClient* FOnlineClientManager::FindRemoteClient(const FGuid& EndpointId) const
+	{
+		const TUniquePtr<FRemoteClient>* Client = RemoteClients.FindByPredicate([&EndpointId](const TUniquePtr<FRemoteClient>& Client)
+		{
+			return Client->GetEndpointId() == EndpointId;
+		});
+		return Client ? Client->Get() : nullptr;
+	}
+
+	void FOnlineClientManager::ForEachClient(TFunctionRef<EBreakBehavior(const FOnlineClient&)> ProcessClient) const
+	{
+		if (ProcessClient(GetLocalClient()) == EBreakBehavior::Break)
+		{
+			return;
+		}
+		for (const TNonNullPtr<const FRemoteClient>& RemoteClient : GetRemoteClients())
+		{
+			if (ProcessClient(*RemoteClient) == EBreakBehavior::Break)
+			{
+				return;
+			}
+		}
+	}
+
+	void FOnlineClientManager::ForEachClient(TFunctionRef<EBreakBehavior(FOnlineClient&)> ProcessClient)
+	{
+		const FOnlineClientManager* ConstThis = this;
+		ConstThis->ForEachClient([&ProcessClient](const FOnlineClient& Client)
+		{
+			// const_cast safe here because remote clients are never const and GetLocalClient() is not const here since this overload is non-const
+			return ProcessClient(const_cast<FOnlineClient&>(Client));
+		});
+	}
+
+	void FOnlineClientManager::AddReferencedObjects(FReferenceCollector& Collector)
+	{
+		ForEachClient([&Collector](const FOnlineClient& Client)
+		{
+			TObjectPtr<UMultiUserReplicationStream> Content = Client.GetClientStreamObject();
+			Collector.AddReferencedObject(Content);
+			ensureMsgf(Content == Client.GetClientStreamObject(), TEXT("Did not expect reference to be obliterated"));
+			return EBreakBehavior::Continue;
+		});
+	}
+
+	void FOnlineClientManager::OnSessionClientChanged(IConcertClientSession&, EConcertClientStatus NewStatus, const FConcertSessionClientInfo& ClientInfo)
+	{
+		const FGuid& ClientEndpointId = ClientInfo.ClientEndpointId;
+		switch (NewStatus)
+		{
+			
+		case EConcertClientStatus::Connected:
+			CreateRemoteClient(ClientEndpointId);
+			break;
+			
+		case EConcertClientStatus::Disconnected:
+			{
+				const int32 Index = RemoteClients.IndexOfByPredicate(
+					[&ClientEndpointId](const TUniquePtr<FRemoteClient>& Client)
+					{
+						return Client->GetEndpointId() == ClientEndpointId;
+					});
+				if (!ensure(RemoteClients.IsValidIndex(Index)))
+				{
+					return;
+				}
+
+				{
+					const TUniquePtr<FRemoteClient> Client = MoveTemp(RemoteClients[Index]);
+					OnPreRemoteClientRemovedDelegate.Broadcast(*Client.Get());
+					RemoteClients.RemoveAtSwap(Index);
+				}
+				// We want to broadcast after the client has been fully cleaned up
+				OnRemoteClientsChangedDelegate.Broadcast();
+			}
+			break;
+			
+		case EConcertClientStatus::Updated:
+			break;
+		default: checkNoEntry();
+		}
+	}
+	
+	void FOnlineClientManager::CreateRemoteClient(const FGuid& ClientEndpointId, bool bBroadcastDelegate)
+	{
+		TUniquePtr<FRemoteClient> RemoteClientPtr = MakeUnique<FRemoteClient>(
+			ClientEndpointId,
+			RegisteredExtenders,
+			ConcertClient->GetConcertClient(),
+			AuthorityCache,
+			*Private::MakeClientContent(),
+			QueryService
+			);
+		FRemoteClient& RemoteClient = *RemoteClientPtr;
+		RemoteClients.Emplace(
+			MoveTemp(RemoteClientPtr)
+			);
+
+		if (bBroadcastDelegate)
+		{
+			OnPostRemoteClientAddedDelegate.Broadcast(RemoteClient);
+			OnRemoteClientsChangedDelegate.Broadcast();
+		}
+	}
+}
diff --git a/Engine/Plugins/Developer/Concert/ConcertApp/MultiUserClient/Source/MultiUserClient/Private/Replication/Client/Online/OnlineClientManager.h b/Engine/Plugins/Developer/Concert/ConcertApp/MultiUserClient/Source/MultiUserClient/Private/Replication/Client/Online/OnlineClientManager.h
new file mode 100644
--- /dev/null
+++ b/Engine/Plugins/Developer/Concert/ConcertApp/MultiUserClient/Source/MultiUserClient/Private/Replication/Client/Online/OnlineClientManager.h
@@ -0,0 +1,162 @@
+// Copyright Epic Games, Inc. All Rights Reserved.
+
+#pragma once
+
+#include "ConcertMessageData.h"
+#include "LocalClient.h"
+#include "OnlineClient.h"
+#include "Replication/Misc/GlobalAuthorityCache.h"
+#include "Replication/Submission/MultiEdit/ReassignObjectPropertiesLogic.h"
+
+#include "UObject/GCObject.h"
+#include "Templates/UnrealTemplate.h"
+#include "Templates/NonNullPointer.h"
+
+class IConcertClientSession;
+class IConcertSyncClient;
+namespace UE::ConcertSharedSlate { class IEditableReplicationStreamModel; }
+enum class EConcertClientStatus : uint8;
+struct FConcertSessionClientInfo;
+
+namespace UE::MultiUserClient::Replication
+{
+	class FRemoteClient;
+	class FReplicationDiscoveryContainer;
+	class FOnlineClient;
+	class FStreamAndAuthorityQueryService;
+	class FStreamChangeTracker;
+
+	/**
+	 * Keeps track of clients connected to replication via a MU session.
+	 * It synchronizes their stream data in a UMultiUserReplicationSessionPreset.
+	 * 
+	 * This class is instantiated only for as long as the local client is in a Concert session and is owned by FMultiUserReplicationManager,
+	 * which drives the lifetime.
+	 */
+	class FOnlineClientManager
+		: public FGCObject
+		, public FNoncopyable
+	{
+	public:
+		
+		/**
+		 * @param InClient The local client. The owning FMultiUserReplicationManager ensures it outlives the constructed instance.
+		 * @param InSession The session to observe. The owning FMultiUserReplicationManager ensures it outlives the constructed instance.
+		 * @param InRegisteredExtenders Used for auto-discovering properties added to this client's stream. Passed to the clients. The owning FMultiUserReplicationManager ensures it outlives the constructed instance.
+		 * @param InQueryService Used for querying info about clients' streams. Passed to the clients. The owning FMultiUserReplicationManager ensures it outlives the constructed instance.
+		 */
+		FOnlineClientManager(
+			const TSharedRef<IConcertSyncClient>& InClient,
+			const TSharedRef<IConcertClientSession>& InSession,
+			FReplicationDiscoveryContainer& InRegisteredExtenders UE_LIFETIMEBOUND,
+			FStreamAndAuthorityQueryService& InQueryService UE_LIFETIMEBOUND
+			);
+		virtual ~FOnlineClientManager() override;
+
+		const FLocalClient& GetLocalClient() const { return LocalClient; }
+		FLocalClient& GetLocalClient() { return LocalClient; }
+		TArray<TNonNullPtr<const FRemoteClient>> GetRemoteClients() const;
+		TArray<TNonNullPtr<FRemoteClient>> GetRemoteClients();
+
+		/** @return Gets all clients given a predicate */
+		TArray<TNonNullPtr<const FOnlineClient>> GetClients(TFunctionRef<bool(const FOnlineClient& Client)> Predicate = [](const auto&){ return true; }) const;
+		
+		const FGlobalAuthorityCache& GetAuthorityCache() const { return AuthorityCache; }
+		FGlobalAuthorityCache& GetAuthorityCache() { return AuthorityCache; }
+		
+		const FReassignObjectPropertiesLogic& GetReassignmentLogic() const { return ReassignmentLogic; }
+		FReassignObjectPropertiesLogic& GetReassignmentLogic() { return ReassignmentLogic; }
+
+		/** Util for finding a remote client by its EndpointId. */
+		const FRemoteClient* FindRemoteClient(const FGuid& EndpointId) const;
+		FRemoteClient* FindRemoteClient(const FGuid& EndpointId)
+		{
+			const FOnlineClientManager* ConstThis = this;
+			return const_cast<FRemoteClient*>(ConstThis->FindRemoteClient(EndpointId));
+		}
+
+		/** Util for finding a local or remote client by its EndpointId. */
+		const FOnlineClient* FindClient(const FGuid& EndpointId) const
+		{
+			return GetLocalClient().GetEndpointId() == EndpointId
+				? &GetLocalClient()
+				: static_cast<const FOnlineClient*>(FindRemoteClient(EndpointId));
+		}
+		FOnlineClient* FindClient(const FGuid& EndpointId)
+		{
+			const FOnlineClientManager* ConstThis = this;
+			return const_cast<FOnlineClient*>(ConstThis->FindClient(EndpointId));
+		}
+
+		/** Iterates through every client */
+		void ForEachClient(TFunctionRef<EBreakBehavior(const FOnlineClient&)> ProcessClient) const;
+		/** Iterates through every client */
+		void ForEachClient(TFunctionRef<EBreakBehavior(FOnlineClient&)> ProcessClient);
+		
+		DECLARE_MULTICAST_DELEGATE(FRemoteClientsChanged);
+		/** Called when RemoteClients changes. Called after OnPostRemoteClientAdded. */
+		FRemoteClientsChanged& OnRemoteClientsChanged() { return OnRemoteClientsChangedDelegate; }
+		
+		DECLARE_MULTICAST_DELEGATE_OneParam(FRemoteClientDelegate, FRemoteClient&);
+		/** Called just after a remote client is has been added to RemoteClients. Called before OnRemoteClientsChanged. */
+		FRemoteClientDelegate& OnPostRemoteClientAdded() { return OnPostRemoteClientAddedDelegate; }
+		/** Called just before a remote client is about to be removed from RemoteClients. */
+		FRemoteClientDelegate& OnPreRemoteClientRemoved() { return OnPreRemoteClientRemovedDelegate; }
+		
+		//~ Begin FGCObject Interface
+		virtual void AddReferencedObjects(FReferenceCollector& Collector) override;
+		virtual FString GetReferencerName() const override { return TEXT("FReplicationStreamSynchronizer"); }
+		//~ End FGCObject Interface
+		
+	private:
+
+		/** The local Concert client */
+		TSharedRef<IConcertSyncClient> ConcertClient;
+		/**
+		 * The session the local client is in.
+		 * 
+		 * This FReplicationClientManager's owner is supposed to make sure this FReplicationClientManager is destroyed
+		 * when the session shuts down.
+		 */
+		const TWeakPtr<IConcertClientSession> Session;
+
+		/**
+		 * Passed to clients for constructing for the purposes of auto-discovering properties and additional objects when adding properties to the stream.
+		 * The owning FMultiUserReplicationManager ensures it outlives this FReplicationClientManager instance.
+		 */
+		FReplicationDiscoveryContainer& RegisteredExtenders;
+		
+		/**
+		 * Sends FConcertReplication_QueryReplicationInfo_Request in regular intervals.
+		 * Shared by all remote clients so all requests are bundled reducing the number of network requests. 
+		 */
+		FStreamAndAuthorityQueryService& QueryService;
+		/** Keeps a cache of object to owning clients. */
+		FGlobalAuthorityCache AuthorityCache;
+		
+		/** Manages the local client */
+		FLocalClient LocalClient;
+		/**
+		 * Manages remote clients. Updated when client connects or disconnects to the active session.
+		 * UI keeps references to systems inside the client so it is TSharedRef in case TArray is reallocated.
+		 */
+		TArray<TUniquePtr<FRemoteClient>> RemoteClients;
+		
+		/** Called when RemoteClients changes. */
+		FRemoteClientsChanged OnRemoteClientsChangedDelegate;
+		/** Called when RemoteClients changes. Called after OnPostRemoteClientAdded. */
+		FRemoteClientDelegate OnPostRemoteClientAddedDelegate;
+		/** Called just before a remote client is about to be removed from RemoteClients. */
+		FRemoteClientDelegate OnPreRemoteClientRemovedDelegate; 
+
+		/** Used for transferring ownership from multiple clients to one. Used by multi client view. */
+		FReassignObjectPropertiesLogic ReassignmentLogic;
+		
+		/** Updates RemoteClients depending on the change. */
+		void OnSessionClientChanged(IConcertClientSession&, EConcertClientStatus NewStatus, const FConcertSessionClientInfo& ClientInfo);
+		
+		/** Shared logic for creating a remote client. */
+		void CreateRemoteClient(const FGuid& ClientEndpointId, bool bBroadcastDelegate = true);
+	};
+}
+
diff --git a/Engine/Plugins/Developer/Concert/ConcertApp/MultiUserClient/Source/MultiUserClient/Private/Replication/Client/Online/RemoteClient.cpp b/Engine/Plugins/Developer/Concert/ConcertApp/MultiUserClient/Source/MultiUserClient/Private/Replication/Client/Online/RemoteClient.cpp
new file mode 100644
--- /dev/null
+++ b/Engine/Plugins/Developer/Concert/ConcertApp/MultiUserClient/Source/MultiUserClient/Private/Replication/Client/Online/RemoteClient.cpp
@@ -0,0 +1,30 @@
+// Copyright Epic Games, Inc. All Rights Reserved.
+
+#include "RemoteClient.h"
+
+#include "IConcertClient.h"
+#include "Assets/MultiUserReplicationClientContent.h"
+#include "Replication/Authority/AuthoritySynchronizer_RemoteClient.h"
+#include "Replication/Stream/StreamSynchronizer_RemoteClient.h"
+#include "Replication/Submission/Remote/SubmissionWorkflow_RemoteClient.h"
+
+namespace UE::MultiUserClient::Replication
+{
+	FRemoteClient::FRemoteClient(
+		const FGuid& InConcertClientId,
+		FReplicationDiscoveryContainer& InDiscoveryContainer,
+		TSharedRef<IConcertClient> InClient,
+		FGlobalAuthorityCache& InAuthorityCache,
+		UMultiUserReplicationStream& InClientStreamContent,
+		FStreamAndAuthorityQueryService& QueryService
+		)
+		: FOnlineClient(
+			InConcertClientId,
+			InDiscoveryContainer,
+			InAuthorityCache,
+			InClientStreamContent,
+			MakeUnique<FStreamSynchronizer_RemoteClient>(InConcertClientId, QueryService),
+			MakeUnique<FAuthoritySynchronizer_RemoteClient>(InConcertClientId,QueryService),
+			MakeUnique<FSubmissionWorkflow_RemoteClient>(InClient->GetCurrentSession().ToSharedRef(), InConcertClientId))
+	{}
+}
diff --git a/Engine/Plugins/Developer/Concert/ConcertApp/MultiUserClient/Source/MultiUserClient/Private/Replication/Client/Online/RemoteClient.h b/Engine/Plugins/Developer/Concert/ConcertApp/MultiUserClient/Source/MultiUserClient/Private/Replication/Client/Online/RemoteClient.h
new file mode 100644
--- /dev/null
+++ b/Engine/Plugins/Developer/Concert/ConcertApp/MultiUserClient/Source/MultiUserClient/Private/Replication/Client/Online/RemoteClient.h
@@ -0,0 +1,38 @@
+// Copyright Epic Games, Inc. All Rights Reserved.
+
+#pragma once
+
+#include "OnlineClient.h"
+
+#include "HAL/Platform.h"
+
+class IConcertClient;
+
+namespace UE::MultiUserClient::Replication
+{
+	class FStreamAndAuthorityQueryService;
+	
+	/** Holds extra information about a remote replication client. */
+	class FRemoteClient : public FOnlineClient
+	{
+	public:
+		
+		/**
+		 * @param InConcertClientId The endpoint ID this instance corresponds to.
+		 * @param InDiscoveryContainer Used for auto-discovering properties added to this client's stream. The caller ensures it outlives the constructed instance.
+		 * @param InClient The local editor's running client. Passed to subsystems so they can send messages to the current session. Keeps strong reference.
+		 * @param InAuthorityCache Caches authority state of all clients. Passed to subsystems. The caller ensures it outlives the constructed instance.
+		 * @param InClientStreamContent Object that this client's stream changes are to be written into. The caller ensures it outlives the constructed instance.
+		 * @param QueryService Service shared across remote clients which periodically asks the server about this client's state. The caller ensures it outlives the constructed instance.
+		 */
+		FRemoteClient(
+			const FGuid& InConcertClientId,
+			FReplicationDiscoveryContainer& InDiscoveryContainer UE_LIFETIMEBOUND,
+			TSharedRef<IConcertClient> InClient,
+			FGlobalAuthorityCache& InAuthorityCache UE_LIFETIMEBOUND,
+			UMultiUserReplicationStream& InClientStreamContent UE_LIFETIMEBOUND,
+			FStreamAndAuthorityQueryService& QueryService UE_LIFETIMEBOUND
+			);
+	};
+}
+
diff --git a/Engine/Plugins/Developer/Concert/ConcertApp/MultiUserClient/Source/MultiUserClient/Private/Replication/Client/RemoteReplicationClient.cpp b/Engine/Plugins/Developer/Concert/ConcertApp/MultiUserClient/Source/MultiUserClient/Private/Replication/Client/RemoteReplicationClient.cpp
deleted file mode 100644
diff --git a/Engine/Plugins/Developer/Concert/ConcertApp/MultiUserClient/Source/MultiUserClient/Private/Replication/Client/RemoteReplicationClient.h b/Engine/Plugins/Developer/Concert/ConcertApp/MultiUserClient/Source/MultiUserClient/Private/Replication/Client/RemoteReplicationClient.h
deleted file mode 100644
diff --git a/Engine/Plugins/Developer/Concert/ConcertApp/MultiUserClient/Source/MultiUserClient/Private/Replication/Client/ReplicationClient.cpp b/Engine/Plugins/Developer/Concert/ConcertApp/MultiUserClient/Source/MultiUserClient/Private/Replication/Client/ReplicationClient.cpp
deleted file mode 100644
diff --git a/Engine/Plugins/Developer/Concert/ConcertApp/MultiUserClient/Source/MultiUserClient/Private/Replication/Client/ReplicationClient.h b/Engine/Plugins/Developer/Concert/ConcertApp/MultiUserClient/Source/MultiUserClient/Private/Replication/Client/ReplicationClient.h
deleted file mode 100644
diff --git a/Engine/Plugins/Developer/Concert/ConcertApp/MultiUserClient/Source/MultiUserClient/Private/Replication/Client/ReplicationClientManager.cpp b/Engine/Plugins/Developer/Concert/ConcertApp/MultiUserClient/Source/MultiUserClient/Private/Replication/Client/ReplicationClientManager.cpp
deleted file mode 100644
diff --git a/Engine/Plugins/Developer/Concert/ConcertApp/MultiUserClient/Source/MultiUserClient/Private/Replication/Client/ReplicationClientManager.h b/Engine/Plugins/Developer/Concert/ConcertApp/MultiUserClient/Source/MultiUserClient/Private/Replication/Client/ReplicationClientManager.h
deleted file mode 100644
diff --git a/Engine/Plugins/Developer/Concert/ConcertApp/MultiUserClient/Source/MultiUserClient/Private/Replication/Misc/Frequency/FrequencyUtils.cpp b/Engine/Plugins/Developer/Concert/ConcertApp/MultiUserClient/Source/MultiUserClient/Private/Replication/Misc/Frequency/FrequencyUtils.cpp
--- a/Engine/Plugins/Developer/Concert/ConcertApp/MultiUserClient/Source/MultiUserClient/Private/Replication/Misc/Frequency/FrequencyUtils.cpp
+++ b/Engine/Plugins/Developer/Concert/ConcertApp/MultiUserClient/Source/MultiUserClient/Private/Replication/Misc/Frequency/FrequencyUtils.cpp
@@ -7,7 +7,7 @@ namespace UE::MultiUserClient::Replication::FrequencyUtils
 	TSharedPtr<IParallelSubmissionOperation> SetFrequencySettingForClients(
 		const FSoftObjectPath& ContextObject,
 		const FInlineClientArray& Clients,
-		FReplicationClientManager& InClientManager,
+		FOnlineClientManager& InClientManager,
 		TFunctionRef<void(FConcertObjectReplicationSettings&)> ApplySettingChange
 		)
 	{
@@ -15,7 +15,7 @@ namespace UE::MultiUserClient::Replication::FrequencyUtils
 		
 		for (const FGuid& ClientId : Clients)
 		{
-			FReplicationClient* Client = InClientManager.FindClient(ClientId);
+			FOnlineClient* Client = InClientManager.FindClient(ClientId);
 			if (Client && CanChangeFrequencySettings(ContextObject, *Client))
 			{
 				const FGuid StreamId = Client->GetStreamSynchronizer().GetStreamId();
diff --git a/Engine/Plugins/Developer/Concert/ConcertApp/MultiUserClient/Source/MultiUserClient/Private/Replication/Misc/Frequency/FrequencyUtils.h b/Engine/Plugins/Developer/Concert/ConcertApp/MultiUserClient/Source/MultiUserClient/Private/Replication/Misc/Frequency/FrequencyUtils.h
--- a/Engine/Plugins/Developer/Concert/ConcertApp/MultiUserClient/Source/MultiUserClient/Private/Replication/Misc/Frequency/FrequencyUtils.h
+++ b/Engine/Plugins/Developer/Concert/ConcertApp/MultiUserClient/Source/MultiUserClient/Private/Replication/Misc/Frequency/FrequencyUtils.h
@@ -2,7 +2,7 @@
 
 #pragma once
 
-#include "Replication/Client/ReplicationClientManager.h"
+#include "Replication/Client/Online/OnlineClientManager.h"
 #include "Replication/Data/ReplicationFrequencySettings.h"
 
 #include "Algo/AllOf.h"
@@ -28,7 +28,7 @@ namespace UE::MultiUserClient::Replication::FrequencyUtils
 	TOptional<TSubSetting> FindSharedFrequencySetting(
 		const FSoftObjectPath& ContextObject,
 		const FInlineClientArray& Clients,
-		const FReplicationClientManager& InClientManager,
+		const FOnlineClientManager& InClientManager,
 		TSelectSubSetting&& Transform
 		)
 	requires std::is_invocable_r_v<TSubSetting, TSelectSubSetting, const FConcertObjectReplicationSettings&>;
@@ -41,12 +41,12 @@ namespace UE::MultiUserClient::Replication::FrequencyUtils
 	TSharedPtr<IParallelSubmissionOperation> SetFrequencySettingForClients(
 		const FSoftObjectPath& ContextObject,
 		const FInlineClientArray& Clients,
-		FReplicationClientManager& InClientManager,
+		FOnlineClientManager& InClientManager,
 		TFunctionRef<void(FConcertObjectReplicationSettings&)> ApplySettingChange
 		);
 
 	/** Gets the replication rate of all clients, if they share the same value. */
-	inline TOptional<uint8> FindSharedFrequencyRate(const FSoftObjectPath& ContextObject, const FInlineClientArray& Clients, const FReplicationClientManager& InClientManager)
+	inline TOptional<uint8> FindSharedFrequencyRate(const FSoftObjectPath& ContextObject, const FInlineClientArray& Clients, const FOnlineClientManager& InClientManager)
 	{
 		bool bHasRightMode = false;
 		const TOptional<uint8> SharedMode = FindSharedFrequencySetting<uint8>(ContextObject, Clients, InClientManager,
@@ -59,7 +59,7 @@ namespace UE::MultiUserClient::Replication::FrequencyUtils
 	}
 
 	/** @return The replication mode that all clients are using, if they share the same value. */
-	inline TOptional<EConcertObjectReplicationMode> FindSharedReplicationMode(const FSoftObjectPath& ContextObject, const FInlineClientArray& Clients, const FReplicationClientManager& InClientManager)
+	inline TOptional<EConcertObjectReplicationMode> FindSharedReplicationMode(const FSoftObjectPath& ContextObject, const FInlineClientArray& Clients, const FOnlineClientManager& InClientManager)
 	{
 		return FindSharedFrequencySetting<EConcertObjectReplicationMode>(ContextObject, Clients, InClientManager,
 			[](const FConcertObjectReplicationSettings& Setting)
@@ -69,13 +69,13 @@ namespace UE::MultiUserClient::Replication::FrequencyUtils
 	}
 	
 	/** @return Whether all clients have the given replication mode */
-	inline bool AllClientsHaveMode(const EConcertObjectReplicationMode SearchedMode, const FSoftObjectPath& ContextObject, const FInlineClientArray& Clients, const FReplicationClientManager& InClientManager) 
+	inline bool AllClientsHaveMode(const EConcertObjectReplicationMode SearchedMode, const FSoftObjectPath& ContextObject, const FInlineClientArray& Clients, const FOnlineClientManager& InClientManager) 
 	{
 		return SearchedMode == FindSharedReplicationMode(ContextObject, Clients, InClientManager);
 	}
 
 	/** @return Whether ContextObject's frequency settings can be changed for Client. */
-	inline bool CanChangeFrequencySettings(const FSoftObjectPath& ContextObject, const FReplicationClient& Client)
+	inline bool CanChangeFrequencySettings(const FSoftObjectPath& ContextObject, const FOnlineClient& Client)
 	{
 		const bool bHasProperties = Client.GetStreamSynchronizer().GetServerState().HasProperties(ContextObject);
 		const bool bAllowsEditing = Client.AllowsEditing();
@@ -84,11 +84,11 @@ namespace UE::MultiUserClient::Replication::FrequencyUtils
 	}
 
 	/** @return Whether ContextObject's frequency settings can be changed for all given Clients. */
-	inline bool CanChangeFrequencySettings(const FSoftObjectPath& ContextObject, const FInlineClientArray& Clients, const FReplicationClientManager& InClientManager)
+	inline bool CanChangeFrequencySettings(const FSoftObjectPath& ContextObject, const FInlineClientArray& Clients, const FOnlineClientManager& InClientManager)
 	{
 		return Algo::AllOf(Clients, [ContextObject, &InClientManager](const FGuid& ClientId)
 		{
-			const FReplicationClient* Client = InClientManager.FindClient(ClientId);
+			const FOnlineClient* Client = InClientManager.FindClient(ClientId);
 			return Client && CanChangeFrequencySettings(ContextObject, *Client);
 		});
 	}
@@ -100,7 +100,7 @@ namespace UE::MultiUserClient::Replication::FrequencyUtils
 	TOptional<TSubSetting> FindSharedFrequencySetting(
 		const FSoftObjectPath& ContextObject,
 		const FInlineClientArray& Clients,
-		const FReplicationClientManager& InClientManager,
+		const FOnlineClientManager& InClientManager,
 		TSelectSubSetting&& Transform
 		)
 	requires std::is_invocable_r_v<TSubSetting, TSelectSubSetting, const FConcertObjectReplicationSettings&>
@@ -108,7 +108,7 @@ namespace UE::MultiUserClient::Replication::FrequencyUtils
 		TOptional<TSubSetting> SharedSetting;
 		for (const FGuid& ClientId : Clients)
 		{
-			const FReplicationClient* Client = InClientManager.FindClient(ClientId);
+			const FOnlineClient* Client = InClientManager.FindClient(ClientId);
 			if (!Client)
 			{
 				continue;
diff --git a/Engine/Plugins/Developer/Concert/ConcertApp/MultiUserClient/Source/MultiUserClient/Private/Replication/Misc/GlobalAuthorityCache.cpp b/Engine/Plugins/Developer/Concert/ConcertApp/MultiUserClient/Source/MultiUserClient/Private/Replication/Misc/GlobalAuthorityCache.cpp
--- a/Engine/Plugins/Developer/Concert/ConcertApp/MultiUserClient/Source/MultiUserClient/Private/Replication/Misc/GlobalAuthorityCache.cpp
+++ b/Engine/Plugins/Developer/Concert/ConcertApp/MultiUserClient/Source/MultiUserClient/Private/Replication/Misc/GlobalAuthorityCache.cpp
@@ -4,12 +4,12 @@
 
 #include "Replication/AuthorityConflictSharedUtils.h"
 #include "Replication/Authority/EAuthorityMutability.h"
-#include "Replication/Client/RemoteReplicationClient.h"
-#include "Replication/Client/ReplicationClientManager.h"
+#include "Replication/Client/Online/RemoteClient.h"
+#include "Replication/Client/Online/OnlineClientManager.h"
 
 namespace UE::MultiUserClient::Replication
 {
-	FGlobalAuthorityCache::FGlobalAuthorityCache(FReplicationClientManager& InClientManager)
+	FGlobalAuthorityCache::FGlobalAuthorityCache(FOnlineClientManager& InClientManager)
 		: ClientManager(InClientManager)
 	{}
 
@@ -92,7 +92,7 @@ namespace UE::MultiUserClient::Replication
 
 	EAuthorityMutability FGlobalAuthorityCache::CanClientTakeAuthorityAfterSubmission(const FSoftObjectPath& Object, const FGuid& ClientId, FProcessPropertyConflict ProcessConflict) const
 	{
-		const FReplicationClient* Client = ClientManager.FindClient(ClientId);
+		const FOnlineClient* Client = ClientManager.FindClient(ClientId);
 		if (!ensure(Client))
 		{
 			return EAuthorityMutability::NotApplicable;
@@ -121,7 +121,7 @@ namespace UE::MultiUserClient::Replication
 
 	bool FGlobalAuthorityCache::CanClientAddProperty(const FSoftObjectPath& Object, const FGuid& ClientId, const FConcertPropertyChain& Chain) const
 	{
-		const FReplicationClient* Client = ClientManager.FindClient(ClientId);
+		const FOnlineClient* Client = ClientManager.FindClient(ClientId);
 		if (!ensure(Client))
 		{
 			return false;
@@ -147,7 +147,7 @@ namespace UE::MultiUserClient::Replication
 
 		for (const FGuid& ClientId : *ClientsWithStreams)
 		{
-			const FReplicationClient* Client = ClientManager.FindClient(ClientId);
+			const FOnlineClient* Client = ClientManager.FindClient(ClientId);
 			if (!ensure(Client))
 			{
 				continue;
@@ -181,7 +181,7 @@ namespace UE::MultiUserClient::Replication
 		
 		ForEachClientWithAuthorityOverObject(Object, [this, &Object, &PropertyChain, &Result, Hash](const FGuid& ClientId)
 		{
-			const FReplicationClient* Client = ClientManager.FindClient(ClientId);
+			const FOnlineClient* Client = ClientManager.FindClient(ClientId);
 			if (!ensureMsgf(Client, TEXT("OnPreRemoteClientRemoved should have updated OwnedObjectsToClients")))
 			{
 				return EBreakBehavior::Continue;
@@ -211,14 +211,14 @@ namespace UE::MultiUserClient::Replication
 		ConcertSyncCore::Replication::AuthorityConflictUtils::CleanseConflictsFromStreamRequest(Request, SendingClient, *this);
 	}
 
-	void FGlobalAuthorityCache::RegisterForClientEvents(const FReplicationClient& Client)
+	void FGlobalAuthorityCache::RegisterForClientEvents(const FOnlineClient& Client)
 	{
 		const FGuid& ClientEndpointId = Client.GetEndpointId();
 		Client.GetAuthoritySynchronizer().OnServerStateChanged().AddRaw(this, &FGlobalAuthorityCache::OnPostAuthorityChanged, ClientEndpointId);
 		Client.GetStreamSynchronizer().OnServerStateChanged().AddRaw(this, &FGlobalAuthorityCache::OnStreamChanged, ClientEndpointId);
 	}
 
-	void FGlobalAuthorityCache::UnregisterFromClientEvents(const FReplicationClient& Client) const
+	void FGlobalAuthorityCache::UnregisterFromClientEvents(const FOnlineClient& Client) const
 	{
 		Client.GetAuthoritySynchronizer().OnServerStateChanged().RemoveAll(this);
 		Client.GetStreamSynchronizer().OnServerStateChanged().RemoveAll(this);
@@ -226,7 +226,7 @@ namespace UE::MultiUserClient::Replication
 
 	void FGlobalAuthorityCache::AddClient(const FGuid& ClientId)
 	{
-		const FReplicationClient* Client = ClientManager.FindClient(ClientId);
+		const FOnlineClient* Client = ClientManager.FindClient(ClientId);
 		if (!ensure(Client))
 		{
 			return;
@@ -271,7 +271,7 @@ namespace UE::MultiUserClient::Replication
 	
 	void FGlobalAuthorityCache::ForEachStream(const FGuid& ClientEndpointId, TFunctionRef<EBreakBehavior(const FGuid& StreamId, const FConcertObjectReplicationMap& ReplicationMap)> Callback) const
 	{
-		const FReplicationClient* Client = ClientManager.FindClient(ClientEndpointId);
+		const FOnlineClient* Client = ClientManager.FindClient(ClientEndpointId);
 		if (!ensure(Client))
 		{
 			return;
@@ -283,7 +283,7 @@ namespace UE::MultiUserClient::Replication
 
 	void FGlobalAuthorityCache::ForEachClient(TFunctionRef<EBreakBehavior(const FGuid& ClientEndpointId)> Callback) const
 	{
-		auto ProcessClient = [&Callback](const FReplicationClient& Client)
+		auto ProcessClient = [&Callback](const FOnlineClient& Client)
 		{
 			return Client.GetAuthoritySynchronizer().HasAnyAuthority()
 				? Callback(Client.GetEndpointId())
@@ -294,7 +294,7 @@ namespace UE::MultiUserClient::Replication
 		{
 			return;
 		}
-		for (const TNonNullPtr<FRemoteReplicationClient>& RemoteClient : ClientManager.GetRemoteClients())
+		for (const TNonNullPtr<FRemoteClient>& RemoteClient : ClientManager.GetRemoteClients())
 		{
 			if (ProcessClient(*RemoteClient) == EBreakBehavior::Break)
 			{
diff --git a/Engine/Plugins/Developer/Concert/ConcertApp/MultiUserClient/Source/MultiUserClient/Private/Replication/Misc/GlobalAuthorityCache.h b/Engine/Plugins/Developer/Concert/ConcertApp/MultiUserClient/Source/MultiUserClient/Private/Replication/Misc/GlobalAuthorityCache.h
--- a/Engine/Plugins/Developer/Concert/ConcertApp/MultiUserClient/Source/MultiUserClient/Private/Replication/Misc/GlobalAuthorityCache.h
+++ b/Engine/Plugins/Developer/Concert/ConcertApp/MultiUserClient/Source/MultiUserClient/Private/Replication/Misc/GlobalAuthorityCache.h
@@ -2,7 +2,7 @@
 
 #pragma once
 
-#include "Replication/Client/RemoteReplicationClient.h"
+#include "Replication/Client/Online/RemoteClient.h"
 
 #include "Replication/AuthorityConflictSharedUtils.h"
 
@@ -14,9 +14,9 @@ struct FConcertPropertyChain;
 
 namespace UE::MultiUserClient::Replication
 {
-	class FRemoteReplicationClient;
-	class FReplicationClient;
-	class FReplicationClientManager;
+	class FRemoteClient;
+	class FOnlineClient;
+	class FOnlineClientManager;
 }
 
 namespace UE::MultiUserClient::Replication
@@ -33,7 +33,7 @@ namespace UE::MultiUserClient::Replication
 
 		using FProcessPropertyConflict = TFunctionRef<EBreakBehavior(const FGuid& ConflictingClientId, const FConcertPropertyChain& Property)>;
 		
-		FGlobalAuthorityCache(FReplicationClientManager& InClientManager);
+		FGlobalAuthorityCache(FOnlineClientManager& InClientManager);
 		/** Called when the local client has been created and it is safe to register client events with FReplicationClientManager. */
 		void RegisterEvents();
 		
@@ -90,7 +90,7 @@ namespace UE::MultiUserClient::Replication
 	private:
 
 		/** Used to obtain the clients and their states */
-		FReplicationClientManager& ClientManager;
+		FOnlineClientManager& ClientManager;
 		
 		/** Maps objects that are owned to the clients that have them in the stream */
 		TMap<FSoftObjectPath, TSet<FGuid>> StreamObjectsToClients;
@@ -103,20 +103,20 @@ namespace UE::MultiUserClient::Replication
 		FOnCacheChanged OnCacheChangedDelegate;
 
 		/** Registers for authority and stream changes */
-		void RegisterForClientEvents(const FReplicationClient& Client);
-		void UnregisterFromClientEvents(const FReplicationClient& Client) const;
+		void RegisterForClientEvents(const FOnlineClient& Client);
+		void UnregisterFromClientEvents(const FOnlineClient& Client) const;
 
 		/** Adds the client to OwnedObjectsToClients */
 		void AddClient(const FGuid& ClientId);
 		void RemoveClient(const FGuid& ClientId);
 
 		// Respond to remote client registration
-		void OnPostRemoteClientAdded(FRemoteReplicationClient& RemoteClient)
+		void OnPostRemoteClientAdded(FRemoteClient& RemoteClient)
 		{
 			RebuildClient(RemoteClient.GetEndpointId());
 			RegisterForClientEvents(RemoteClient);
 		}
-		void OnPreRemoteClientRemoved(FRemoteReplicationClient& RemoteClient)
+		void OnPreRemoteClientRemoved(FRemoteClient& RemoteClient)
 		{
 			RemoveClient(RemoteClient.GetEndpointId());
 			UnregisterFromClientEvents(RemoteClient);
diff --git a/Engine/Plugins/Developer/Concert/ConcertApp/MultiUserClient/Source/MultiUserClient/Private/Replication/Misc/Notification/ReplicationUserNotifier.cpp b/Engine/Plugins/Developer/Concert/ConcertApp/MultiUserClient/Source/MultiUserClient/Private/Replication/Misc/Notification/ReplicationUserNotifier.cpp
--- a/Engine/Plugins/Developer/Concert/ConcertApp/MultiUserClient/Source/MultiUserClient/Private/Replication/Misc/Notification/ReplicationUserNotifier.cpp
+++ b/Engine/Plugins/Developer/Concert/ConcertApp/MultiUserClient/Source/MultiUserClient/Private/Replication/Misc/Notification/ReplicationUserNotifier.cpp
@@ -4,7 +4,7 @@
 
 namespace UE::MultiUserClient::Replication
 {
-	FReplicationUserNotifier::FReplicationUserNotifier(FReplicationClientManager& InReplicationClientManager, FMuteStateManager& InMuteManager)
+	FReplicationUserNotifier::FReplicationUserNotifier(FOnlineClientManager& InReplicationClientManager, FMuteStateManager& InMuteManager)
 		: SubmissionNotifier(InReplicationClientManager)
 		, MutingNotifier(InMuteManager)
 	{}
diff --git a/Engine/Plugins/Developer/Concert/ConcertApp/MultiUserClient/Source/MultiUserClient/Private/Replication/Misc/Notification/ReplicationUserNotifier.h b/Engine/Plugins/Developer/Concert/ConcertApp/MultiUserClient/Source/MultiUserClient/Private/Replication/Misc/Notification/ReplicationUserNotifier.h
--- a/Engine/Plugins/Developer/Concert/ConcertApp/MultiUserClient/Source/MultiUserClient/Private/Replication/Misc/Notification/ReplicationUserNotifier.h
+++ b/Engine/Plugins/Developer/Concert/ConcertApp/MultiUserClient/Source/MultiUserClient/Private/Replication/Misc/Notification/ReplicationUserNotifier.h
@@ -17,7 +17,7 @@ namespace UE::MultiUserClient::Replication
 	{
 	public:
 
-		FReplicationUserNotifier(FReplicationClientManager& InReplicationClientManager UE_LIFETIMEBOUND, FMuteStateManager& InMuteManager UE_LIFETIMEBOUND);
+		FReplicationUserNotifier(FOnlineClientManager& InReplicationClientManager UE_LIFETIMEBOUND, FMuteStateManager& InMuteManager UE_LIFETIMEBOUND);
 
 	private:
 		
diff --git a/Engine/Plugins/Developer/Concert/ConcertApp/MultiUserClient/Source/MultiUserClient/Private/Replication/Misc/PreventReplicatedPropertyTransaction.cpp b/Engine/Plugins/Developer/Concert/ConcertApp/MultiUserClient/Source/MultiUserClient/Private/Replication/Misc/PreventReplicatedPropertyTransaction.cpp
--- a/Engine/Plugins/Developer/Concert/ConcertApp/MultiUserClient/Source/MultiUserClient/Private/Replication/Misc/PreventReplicatedPropertyTransaction.cpp
+++ b/Engine/Plugins/Developer/Concert/ConcertApp/MultiUserClient/Source/MultiUserClient/Private/Replication/Misc/PreventReplicatedPropertyTransaction.cpp
@@ -3,7 +3,7 @@
 #include "PreventReplicatedPropertyTransaction.h"
 
 #include "IConcertSyncClient.h"
-#include "Replication/Client/ReplicationClientManager.h"
+#include "Replication/Client/Online/OnlineClientManager.h"
 #include "Replication/Misc/GlobalAuthorityCache.h"
 
 #include "Algo/AnyOf.h"
@@ -19,7 +19,7 @@ namespace UE::MultiUserClient::Replication
 	
 	FPreventReplicatedPropertyTransaction::FPreventReplicatedPropertyTransaction(
 		IConcertSyncClient& InSyncClient,
-		FReplicationClientManager& InClientManager,
+		FOnlineClientManager& InClientManager,
 		FMuteStateManager& InMuteManager
 		)
 		: SyncClient(InSyncClient)
diff --git a/Engine/Plugins/Developer/Concert/ConcertApp/MultiUserClient/Source/MultiUserClient/Private/Replication/Misc/PreventReplicatedPropertyTransaction.h b/Engine/Plugins/Developer/Concert/ConcertApp/MultiUserClient/Source/MultiUserClient/Private/Replication/Misc/PreventReplicatedPropertyTransaction.h
--- a/Engine/Plugins/Developer/Concert/ConcertApp/MultiUserClient/Source/MultiUserClient/Private/Replication/Misc/PreventReplicatedPropertyTransaction.h
+++ b/Engine/Plugins/Developer/Concert/ConcertApp/MultiUserClient/Source/MultiUserClient/Private/Replication/Misc/PreventReplicatedPropertyTransaction.h
@@ -14,7 +14,7 @@ struct FConcertTransactionFilterArgs;
 namespace UE::MultiUserClient::Replication
 {
 	class FMuteStateManager;
-	class FReplicationClientManager;
+	class FOnlineClientManager;
 	
 	/** Prevents generating transactions with properties that are being replicated by any other client. */
 	class FPreventReplicatedPropertyTransaction : public FNoncopyable
@@ -23,7 +23,7 @@ namespace UE::MultiUserClient::Replication
 		
 		FPreventReplicatedPropertyTransaction(
 			IConcertSyncClient& InSyncClient UE_LIFETIMEBOUND,
-			FReplicationClientManager& InClientManager UE_LIFETIMEBOUND,
+			FOnlineClientManager& InClientManager UE_LIFETIMEBOUND,
 			FMuteStateManager& InMuteManager UE_LIFETIMEBOUND
 			);
 		~FPreventReplicatedPropertyTransaction();
@@ -34,7 +34,7 @@ namespace UE::MultiUserClient::Replication
 		IConcertSyncClient& SyncClient;
 
 		/** Used to obtain up-to-date registered properties from all clients. */
-		FReplicationClientManager& ClientManager;
+		FOnlineClientManager& ClientManager;
 		/** Used to approximate whether remote clients may have sync control */
 		FMuteStateManager& MuteManager;
 		
diff --git a/Engine/Plugins/Developer/Concert/ConcertApp/MultiUserClient/Source/MultiUserClient/Private/Replication/Misc/PropertySelection/UserPropertySelectionSource.cpp b/Engine/Plugins/Developer/Concert/ConcertApp/MultiUserClient/Source/MultiUserClient/Private/Replication/Misc/PropertySelection/UserPropertySelectionSource.cpp
--- a/Engine/Plugins/Developer/Concert/ConcertApp/MultiUserClient/Source/MultiUserClient/Private/Replication/Misc/PropertySelection/UserPropertySelectionSource.cpp
+++ b/Engine/Plugins/Developer/Concert/ConcertApp/MultiUserClient/Source/MultiUserClient/Private/Replication/Misc/PropertySelection/UserPropertySelectionSource.cpp
@@ -2,7 +2,7 @@
 
 #include "UserPropertySelectionSource.h"
 
-#include "Replication/Client/ReplicationClientManager.h"
+#include "Replication/Client/Online/OnlineClientManager.h"
 #include "Replication/Editor/Model/IEditableReplicationStreamModel.h"
 #include "Replication/Editor/Model/Property/IPropertySource.h"
 
@@ -19,7 +19,7 @@ namespace UE::MultiUserClient::Replication
 			FUserPropertySource(
 				FSoftObjectPath Object,
 				const ConcertSharedSlate::IEditableReplicationStreamModel& InUserSelection UE_LIFETIMEBOUND,
-				const FReplicationClientManager& InClientManager UE_LIFETIMEBOUND
+				const FOnlineClientManager& InClientManager UE_LIFETIMEBOUND
 				)
 				: Object(MoveTemp(Object))
 				, UserSelection(InUserSelection)
@@ -48,7 +48,7 @@ namespace UE::MultiUserClient::Replication
 				using FPropertyPtrArray = TArray<const FConcertPropertyChain*, TInlineAllocator<4>>;
 				TMap<FHash, FPropertyPtrArray, TInlineSetAllocator<512>> VisitedProperties;
 				
-				ClientManager.ForEachClient([this, &Delegate, &VisitedProperties](const FReplicationClient& Client)
+				ClientManager.ForEachClient([this, &Delegate, &VisitedProperties](const FOnlineClient& Client)
 				{
 					const FConcertReplicatedObjectInfo* ObjectInfo = Client.GetStreamSynchronizer()
 						.GetServerState()
@@ -91,13 +91,13 @@ namespace UE::MultiUserClient::Replication
 			/** Used to get the properties the user has selected. */
 			const ConcertSharedSlate::IEditableReplicationStreamModel& UserSelection;
 			/** Used to get client stream content. */
-			const FReplicationClientManager& ClientManager;
+			const FOnlineClientManager& ClientManager;
 		};
 	}
 	
 	FUserPropertySelectionSource::FUserPropertySelectionSource(
 		const ConcertSharedSlate::IEditableReplicationStreamModel& InUserSelection,
-		const FReplicationClientManager& InClientManager
+		const FOnlineClientManager& InClientManager
 		)
 		: UserSelection(InUserSelection)
 		, ClientManager(InClientManager)
diff --git a/Engine/Plugins/Developer/Concert/ConcertApp/MultiUserClient/Source/MultiUserClient/Private/Replication/Misc/PropertySelection/UserPropertySelectionSource.h b/Engine/Plugins/Developer/Concert/ConcertApp/MultiUserClient/Source/MultiUserClient/Private/Replication/Misc/PropertySelection/UserPropertySelectionSource.h
--- a/Engine/Plugins/Developer/Concert/ConcertApp/MultiUserClient/Source/MultiUserClient/Private/Replication/Misc/PropertySelection/UserPropertySelectionSource.h
+++ b/Engine/Plugins/Developer/Concert/ConcertApp/MultiUserClient/Source/MultiUserClient/Private/Replication/Misc/PropertySelection/UserPropertySelectionSource.h
@@ -11,7 +11,7 @@ namespace UE::ConcertSharedSlate { class IEditableReplicationStreamModel; }
 
 namespace UE::MultiUserClient::Replication
 {
-	class FReplicationClientManager;
+	class FOnlineClientManager;
 	
 	/**
 	 * Injected into UI causing it to only display the properties that
@@ -26,7 +26,7 @@ namespace UE::MultiUserClient::Replication
 
 		FUserPropertySelectionSource(
 			const ConcertSharedSlate::IEditableReplicationStreamModel& InUserSelection UE_LIFETIMEBOUND,
-			const FReplicationClientManager& InClientManager UE_LIFETIMEBOUND
+			const FOnlineClientManager& InClientManager UE_LIFETIMEBOUND
 			);
 
 		//~ Begin IPropertySourceProcessor Interface
@@ -41,6 +41,6 @@ namespace UE::MultiUserClient::Replication
 		/** Used to get the properties the user has selected. */
 		const ConcertSharedSlate::IEditableReplicationStreamModel& UserSelection;
 		/** Used to get client stream content and subscribe to changes. */
-		const FReplicationClientManager& ClientManager;
+		const FOnlineClientManager& ClientManager;
 	};
 }
diff --git a/Engine/Plugins/Developer/Concert/ConcertApp/MultiUserClient/Source/MultiUserClient/Private/Replication/Misc/PropertySelection/UserPropertySelector.cpp b/Engine/Plugins/Developer/Concert/ConcertApp/MultiUserClient/Source/MultiUserClient/Private/Replication/Misc/PropertySelection/UserPropertySelector.cpp
--- a/Engine/Plugins/Developer/Concert/ConcertApp/MultiUserClient/Source/MultiUserClient/Private/Replication/Misc/PropertySelection/UserPropertySelector.cpp
+++ b/Engine/Plugins/Developer/Concert/ConcertApp/MultiUserClient/Source/MultiUserClient/Private/Replication/Misc/PropertySelection/UserPropertySelector.cpp
@@ -4,7 +4,7 @@
 
 #include "UserPropertySelectionSource.h"
 #include "Replication/ClientReplicationWidgetFactories.h"
-#include "Replication/Client/ReplicationClientManager.h"
+#include "Replication/Client/Online/OnlineClientManager.h"
 #include "Replication/Editor/Model/IEditableReplicationStreamModel.h"
 
 #include "ScopedTransaction.h"
@@ -15,7 +15,7 @@
 
 namespace UE::MultiUserClient::Replication
 {
-	FUserPropertySelector::FUserPropertySelector(FReplicationClientManager& InClientManager)
+	FUserPropertySelector::FUserPropertySelector(FOnlineClientManager& InClientManager)
 		: ClientManager(InClientManager)
 		, PropertySelection(NewObject<UMultiUserReplicationStream>(GetTransientPackage(), NAME_None, RF_Transient | RF_Transactional))
 		, SelectionEditModel(ConcertSharedSlate::CreateBaseStreamModel(PropertySelection->MakeReplicationMapGetterAttribute()))
@@ -29,7 +29,7 @@ namespace UE::MultiUserClient::Replication
 	FUserPropertySelector::~FUserPropertySelector()
 	{
 		ClientManager.OnPostRemoteClientAdded().RemoveAll(this);
-		ClientManager.ForEachClient([this](FReplicationClient& Client)
+		ClientManager.ForEachClient([this](FOnlineClient& Client)
 		{
 			Client.GetStreamSynchronizer().OnServerStateChanged().RemoveAll(this);
 			return EBreakBehavior::Continue;
@@ -69,7 +69,7 @@ namespace UE::MultiUserClient::Replication
 		Collector.AddReferencedObject(PropertySelection);
 	}
 
-	void FUserPropertySelector::RegisterClient(FReplicationClient& Client)
+	void FUserPropertySelector::RegisterClient(FOnlineClient& Client)
 	{
 		IClientStreamSynchronizer& StreamSynchronizer = Client.GetStreamSynchronizer();
 		TrackProperties(StreamSynchronizer.GetServerState());
@@ -78,7 +78,7 @@ namespace UE::MultiUserClient::Replication
 
 	void FUserPropertySelector::OnServerStateChanged(const FGuid ClientId)
 	{
-		const FReplicationClient* Client = ClientManager.FindClient(ClientId);
+		const FOnlineClient* Client = ClientManager.FindClient(ClientId);
 		if (ensure(Client))
 		{
 			TrackProperties(Client->GetStreamSynchronizer().GetServerState());
diff --git a/Engine/Plugins/Developer/Concert/ConcertApp/MultiUserClient/Source/MultiUserClient/Private/Replication/Misc/PropertySelection/UserPropertySelector.h b/Engine/Plugins/Developer/Concert/ConcertApp/MultiUserClient/Source/MultiUserClient/Private/Replication/Misc/PropertySelection/UserPropertySelector.h
--- a/Engine/Plugins/Developer/Concert/ConcertApp/MultiUserClient/Source/MultiUserClient/Private/Replication/Misc/PropertySelection/UserPropertySelector.h
+++ b/Engine/Plugins/Developer/Concert/ConcertApp/MultiUserClient/Source/MultiUserClient/Private/Replication/Misc/PropertySelection/UserPropertySelector.h
@@ -4,8 +4,9 @@
 
 #include "Assets/MultiUserReplicationStream.h"
 
+#include "Replication/Client/Online/RemoteClient.h"
+
 #include "HAL/Platform.h"
-#include "Replication/Client/RemoteReplicationClient.h"
 #include "Templates/UnrealTemplate.h"
 #include "UObject/GCObject.h"
 #include "UObject/ObjectPtr.h"
@@ -23,7 +24,7 @@ namespace UE::ConcertSharedSlate
 namespace UE::MultiUserClient::Replication
 {
 	class FUserPropertySelectionSource;
-	class FReplicationClientManager;
+	class FOnlineClientManager;
 	
 	/**
 	 * Manages the properties the user is iterating on in the replication session.
@@ -38,7 +39,7 @@ namespace UE::MultiUserClient::Replication
 	{
 	public:
 		
-		FUserPropertySelector(FReplicationClientManager& InClientManager UE_LIFETIMEBOUND);
+		FUserPropertySelector(FOnlineClientManager& InClientManager UE_LIFETIMEBOUND);
 		~FUserPropertySelector();
 
 		/** Add Properties from the user's selection for Object. */
@@ -62,7 +63,7 @@ namespace UE::MultiUserClient::Replication
 	private:
 
 		/** Used to remove deselected properties from local client's stream. */
-		FReplicationClientManager& ClientManager;
+		FOnlineClientManager& ClientManager;
 
 		/** This underlying object saves the properties that user has selected. It allows for transactions. */
 		TObjectPtr<UMultiUserReplicationStream> PropertySelection;
@@ -76,9 +77,9 @@ namespace UE::MultiUserClient::Replication
 		FOnPropertySelectionChanged OnPropertySelectionChangedDelegate;
 
 		/** Called when a remote client joins. */
-		void OnClientAdded(FRemoteReplicationClient& Client) { RegisterClient(Client); }
+		void OnClientAdded(FRemoteClient& Client) { RegisterClient(Client); }
 		/** Ensures that whenever the client's server state changes, its properties are tracked as user selected. */
-		void RegisterClient(FReplicationClient& Client);
+		void RegisterClient(FOnlineClient& Client);
 
 		/** Tracks all properties of the client as user selected. */
 		void OnServerStateChanged(const FGuid ClientId);
diff --git a/Engine/Plugins/Developer/Concert/ConcertApp/MultiUserClient/Source/MultiUserClient/Private/Replication/Misc/Util/SynchronizedRequestUtils.cpp b/Engine/Plugins/Developer/Concert/ConcertApp/MultiUserClient/Source/MultiUserClient/Private/Replication/Misc/Util/SynchronizedRequestUtils.cpp
--- a/Engine/Plugins/Developer/Concert/ConcertApp/MultiUserClient/Source/MultiUserClient/Private/Replication/Misc/Util/SynchronizedRequestUtils.cpp
+++ b/Engine/Plugins/Developer/Concert/ConcertApp/MultiUserClient/Source/MultiUserClient/Private/Replication/Misc/Util/SynchronizedRequestUtils.cpp
@@ -3,7 +3,7 @@
 #include "SynchronizedRequestUtils.h"
 
 #include "ConcertLogGlobal.h"
-#include "Replication/Client/ReplicationClientManager.h"
+#include "Replication/Client/Online/OnlineClientManager.h"
 #include "Replication/Submission/Data/AuthoritySubmission.h"
 #include "Replication/Submission/Data/StreamSubmission.h"
 #include "Replication/Submission/ISubmissionOperation.h"
@@ -16,7 +16,7 @@ namespace UE::MultiUserClient::Replication
 {
 	namespace ParallelSubmission::Private
 	{
-		static void RemoveInvalidClientEntries(const FReplicationClientManager& ClientManager, TMap<FGuid, FSubmissionParams> ParallelOperations)
+		static void RemoveInvalidClientEntries(const FOnlineClientManager& ClientManager, TMap<FGuid, FSubmissionParams> ParallelOperations)
 		{
 			for (auto It = ParallelOperations.CreateIterator(); It; ++It)
 			{
@@ -36,7 +36,7 @@ namespace UE::MultiUserClient::Replication
 			, public IDeferredSubmitter
 		{
 			FSyncOperation& Owner;
-			FReplicationClientManager& ClientManager;
+			FOnlineClientManager& ClientManager;
 			FSubmissionQueue& SubmissionQueue;
 			const FGuid ClientId;
 			const FSubmissionParams SubmissionParams;
@@ -45,7 +45,7 @@ namespace UE::MultiUserClient::Replication
 
 			FDeferredSubmitter(
 				FSyncOperation& InOwner,
-				FReplicationClientManager& InClientManager,
+				FOnlineClientManager& InClientManager,
 				FSubmissionQueue& InSubmissionQueue,
 				const FGuid& InClientId,
 				FSubmissionParams&& InSubmissionParams
@@ -60,7 +60,7 @@ namespace UE::MultiUserClient::Replication
 			virtual ~FDeferredSubmitter() override
 			{
 				checkf(IsInGameThread(), TEXT("You must destroy IParallelSubmissionOperation on the game thread"));
-				if (FReplicationClient* Client = ClientManager.FindClient(ClientId))
+				if (FOnlineClient* Client = ClientManager.FindClient(ClientId))
 				{
 					Client->GetSubmissionQueue().Dequeue_GameThread(*this);
 				}
@@ -81,7 +81,7 @@ namespace UE::MultiUserClient::Replication
 		{
 		public:
 			
-			FSyncOperation(FReplicationClientManager& ClientManager, TMap<FGuid, FSubmissionParams>&& Operations)
+			FSyncOperation(FOnlineClientManager& ClientManager, TMap<FGuid, FSubmissionParams>&& Operations)
 			{
 				const int32 NumOperations = Operations.Num();
 				IntermediateResults.StreamResponses.Reserve(NumOperations);
@@ -96,7 +96,7 @@ namespace UE::MultiUserClient::Replication
 					IntermediateResults.StreamResponses.Add(ClientId);
 					IntermediateResults.AuthorityResponses.Add(ClientId);
 
-					FReplicationClient* Client = ClientManager.FindClient(ClientId);
+					FOnlineClient* Client = ClientManager.FindClient(ClientId);
 					checkf(Client, TEXT("Was supposed to have been validated ExecuteParallelStreamChanges."))
 					SubOperations.Emplace(*this, ClientManager, ClientId, MoveTemp(Pair.Value), Client->GetSubmissionQueue());
 				}
@@ -180,7 +180,7 @@ namespace UE::MultiUserClient::Replication
 				
 				FPendingOperation(
 					FSyncOperation& InOwner,
-					FReplicationClientManager& InClientManager,
+					FOnlineClientManager& InClientManager,
 					const FGuid& InClientId,
 					FSubmissionParams&& InSubmissionParams,
 					FSubmissionQueue& InSubmissionQueue
@@ -257,7 +257,7 @@ namespace UE::MultiUserClient::Replication
 		}
 	}
 	
-	TSharedPtr<IParallelSubmissionOperation> ExecuteParallelStreamChanges(FReplicationClientManager& ClientManager, TMap<FGuid, FSubmissionParams> ParallelOperations)
+	TSharedPtr<IParallelSubmissionOperation> ExecuteParallelStreamChanges(FOnlineClientManager& ClientManager, TMap<FGuid, FSubmissionParams> ParallelOperations)
 	{
 		ParallelSubmission::Private::RemoveInvalidClientEntries(ClientManager, ParallelOperations);
 		if (!ensure(!ParallelOperations.IsEmpty()))
@@ -271,7 +271,7 @@ namespace UE::MultiUserClient::Replication
 		return Operation;
 	}
 
-	TSharedPtr<IParallelSubmissionOperation> ExecuteParallelStreamChanges(FReplicationClientManager& ClientManager, TMap<FGuid, FConcertReplication_ChangeStream_Request> ParallelOperations)
+	TSharedPtr<IParallelSubmissionOperation> ExecuteParallelStreamChanges(FOnlineClientManager& ClientManager, TMap<FGuid, FConcertReplication_ChangeStream_Request> ParallelOperations)
 	{
 		if (!ensure(!ParallelOperations.IsEmpty()))
 		{
diff --git a/Engine/Plugins/Developer/Concert/ConcertApp/MultiUserClient/Source/MultiUserClient/Private/Replication/Misc/Util/SynchronizedRequestUtils.h b/Engine/Plugins/Developer/Concert/ConcertApp/MultiUserClient/Source/MultiUserClient/Private/Replication/Misc/Util/SynchronizedRequestUtils.h
--- a/Engine/Plugins/Developer/Concert/ConcertApp/MultiUserClient/Source/MultiUserClient/Private/Replication/Misc/Util/SynchronizedRequestUtils.h
+++ b/Engine/Plugins/Developer/Concert/ConcertApp/MultiUserClient/Source/MultiUserClient/Private/Replication/Misc/Util/SynchronizedRequestUtils.h
@@ -12,7 +12,7 @@ struct FConcertReplication_ChangeStream_Request;
 
 namespace UE::MultiUserClient::Replication
 {
-	class FReplicationClientManager;
+	class FOnlineClientManager;
 	struct FSubmissionParams;
 
 	struct FParallelExecutionResult
@@ -51,8 +51,8 @@ namespace UE::MultiUserClient::Replication
 	 * @return Object you can get future through. This object also handles correct unregistering from other systems when you destroy it.
 	 * @note The returned object MUST be destroyed on the game thread.
 	 */
-	TSharedPtr<IParallelSubmissionOperation> ExecuteParallelStreamChanges(FReplicationClientManager& ClientManager, TMap<FGuid, FSubmissionParams> ParallelOperations);
+	TSharedPtr<IParallelSubmissionOperation> ExecuteParallelStreamChanges(FOnlineClientManager& ClientManager, TMap<FGuid, FSubmissionParams> ParallelOperations);
 	/** Util for transforming into FSubmissionParams */
-	TSharedPtr<IParallelSubmissionOperation> ExecuteParallelStreamChanges(FReplicationClientManager& ClientManager, TMap<FGuid, FConcertReplication_ChangeStream_Request> ParallelOperations);
+	TSharedPtr<IParallelSubmissionOperation> ExecuteParallelStreamChanges(FOnlineClientManager& ClientManager, TMap<FGuid, FConcertReplication_ChangeStream_Request> ParallelOperations);
 }
 
diff --git a/Engine/Plugins/Developer/Concert/ConcertApp/MultiUserClient/Source/MultiUserClient/Private/Replication/MultiUserReplicationManager.cpp b/Engine/Plugins/Developer/Concert/ConcertApp/MultiUserClient/Source/MultiUserClient/Private/Replication/MultiUserReplicationManager.cpp
--- a/Engine/Plugins/Developer/Concert/ConcertApp/MultiUserClient/Source/MultiUserClient/Private/Replication/MultiUserReplicationManager.cpp
+++ b/Engine/Plugins/Developer/Concert/ConcertApp/MultiUserClient/Source/MultiUserClient/Private/Replication/MultiUserReplicationManager.cpp
@@ -142,8 +142,8 @@ namespace UE::MultiUserClient::Replication
 
 	void FMultiUserReplicationManager::SetupClientConnectionEvents()
 	{
-		FReplicationClientManager& ClientManager = ConnectedState->ClientManager;
-		ClientManager.ForEachClient([this](FReplicationClient& InClient){ SetupClientDelegates(InClient); return EBreakBehavior::Continue; });
+		FOnlineClientManager& ClientManager = ConnectedState->ClientManager;
+		ClientManager.ForEachClient([this](FOnlineClient& InClient){ SetupClientDelegates(InClient); return EBreakBehavior::Continue; });
 		ClientManager.OnPostRemoteClientAdded().AddRaw(this, &FMultiUserReplicationManager::OnReplicationClientConnected);
 	}
 
@@ -159,7 +159,7 @@ namespace UE::MultiUserClient::Replication
 		OnAuthorityServerStateChangedDelegate.Broadcast(EndpointId);
 	}
 
-	void FMultiUserReplicationManager::SetupClientDelegates(FReplicationClient& InClient) const
+	void FMultiUserReplicationManager::SetupClientDelegates(FOnlineClient& InClient) const
 	{
 		InClient.GetStreamSynchronizer().OnServerStateChanged().AddRaw(this, &FMultiUserReplicationManager::OnClientStreamServerStateChanged, InClient.GetEndpointId());
 		InClient.GetAuthoritySynchronizer().OnServerStateChanged().AddRaw(this, &FMultiUserReplicationManager::OnClientAuthorityServerStateChanged, InClient.GetEndpointId());
@@ -169,7 +169,7 @@ namespace UE::MultiUserClient::Replication
 	{
 		if (ConnectedState && ensureMsgf(IsInGameThread(), TEXT("To simplify implementation, only calls from game thread are allowed.")))
 		{
-			const FReplicationClient* ReplicationClient = ConnectedState->ClientManager.FindClient(ClientId);
+			const FOnlineClient* ReplicationClient = ConnectedState->ClientManager.FindClient(ClientId);
 			return ReplicationClient
 				? &ReplicationClient->GetStreamSynchronizer().GetServerState()
 				: nullptr;
@@ -181,7 +181,7 @@ namespace UE::MultiUserClient::Replication
 	{
 		if (ConnectedState && ensureMsgf(IsInGameThread(), TEXT("To simplify implementation, only calls from game thread are allowed.")))
 		{
-			const FReplicationClient* ReplicationClient = ConnectedState->ClientManager.FindClient(ClientId);
+			const FOnlineClient* ReplicationClient = ConnectedState->ClientManager.FindClient(ClientId);
 			return ReplicationClient
 				? &ReplicationClient->GetStreamSynchronizer().GetFrequencySettings()
 				: nullptr;
@@ -193,7 +193,7 @@ namespace UE::MultiUserClient::Replication
 	{
 		if (ConnectedState && ensureMsgf(IsInGameThread(), TEXT("To simplify implementation, only calls from game thread are allowed.")))
 		{
-			const FReplicationClient* ReplicationClient = ConnectedState->ClientManager.FindClient(ClientId);
+			const FOnlineClient* ReplicationClient = ConnectedState->ClientManager.FindClient(ClientId);
 			return ReplicationClient && ReplicationClient->GetAuthoritySynchronizer().HasAuthorityOver(ObjectPath);
 		}
 		return false;
@@ -224,7 +224,7 @@ namespace UE::MultiUserClient::Replication
 		
 		if (ConnectedState)
 		{
-			FReplicationClient* ReplicationClient = ConnectedState->ClientManager.FindClient(ClientId);
+			FOnlineClient* ReplicationClient = ConnectedState->ClientManager.FindClient(ClientId);
 			return ReplicationClient
 				? ReplicationClient->GetExternalRequestHandler().HandleRequest(MoveTemp(SubmissionParams))
 				: FExternalClientChangeRequestHandler::MakeFailedOperation(EChangeStreamOperationResult::UnknownClient, EChangeAuthorityOperationResult::UnknownClient);
diff --git a/Engine/Plugins/Developer/Concert/ConcertApp/MultiUserClient/Source/MultiUserClient/Private/Replication/MultiUserReplicationManager.h b/Engine/Plugins/Developer/Concert/ConcertApp/MultiUserClient/Source/MultiUserClient/Private/Replication/MultiUserReplicationManager.h
--- a/Engine/Plugins/Developer/Concert/ConcertApp/MultiUserClient/Source/MultiUserClient/Private/Replication/MultiUserReplicationManager.h
+++ b/Engine/Plugins/Developer/Concert/ConcertApp/MultiUserClient/Source/MultiUserClient/Private/Replication/MultiUserReplicationManager.h
@@ -2,7 +2,7 @@
 
 #pragma once
 
-#include "Client/ReplicationClientManager.h"
+#include "Client/Online/OnlineClientManager.h"
 #include "IConcertSession.h"
 #include "Misc/ChangeLevelHandler.h"
 #include "Misc/Notification/ReplicationUserNotifier.h"
@@ -27,7 +27,7 @@ namespace UE::ConcertSyncClient { struct FJoinReplicatedSessionResult; }
 
 namespace UE::MultiUserClient::Replication
 {
-	class FReplicationClient;
+	class FOnlineClient;
 	
 	enum class EMultiUserReplicationConnectionState : uint8
 	{
@@ -62,8 +62,8 @@ namespace UE::MultiUserClient::Replication
 		void JoinReplicationSession();
 
 		/** @note You're not supposed to keep any reference to the ClientManager since it can become invalid depending on connection state. */
-		FReplicationClientManager* GetClientManager() { return ConnectedState ? &ConnectedState->ClientManager : nullptr; }
-		const FReplicationClientManager* GetClientManager() const { return ConnectedState ? &ConnectedState->ClientManager : nullptr; }
+		FOnlineClientManager* GetClientManager() { return ConnectedState ? &ConnectedState->ClientManager : nullptr; }
+		const FOnlineClientManager* GetClientManager() const { return ConnectedState ? &ConnectedState->ClientManager : nullptr; }
 		
 		/** @note You're not supposed to keep any reference to the MuteManager since it can become invalid depending on connection state. */
 		FMuteStateManager* GetMuteManager() { return ConnectedState ? &ConnectedState->MuteManager : nullptr; }
@@ -121,7 +121,7 @@ namespace UE::MultiUserClient::Replication
 			 *
 			 * Only valid when ConnectionState == EMultiUserReplicationConnectionState::Connected.
 			 */
-			FReplicationClientManager ClientManager;
+			FOnlineClientManager ClientManager;
 			/** Interacts with the mute global server mute system. */
 			FMuteStateManager MuteManager;
 			/** Saves and loads presets for the session. Accessed by UI. */
@@ -171,8 +171,8 @@ namespace UE::MultiUserClient::Replication
 		void SetupClientConnectionEvents();
 		void OnClientStreamServerStateChanged(const FGuid EndpointId) const;
 		void OnClientAuthorityServerStateChanged(const FGuid EndpointId) const;
-		void OnReplicationClientConnected(FRemoteReplicationClient& RemoteClient) const { SetupClientDelegates(RemoteClient); }
-		void SetupClientDelegates(FReplicationClient& InClient) const;
+		void OnReplicationClientConnected(FRemoteClient& RemoteClient) const { SetupClientDelegates(RemoteClient); }
+		void SetupClientDelegates(FOnlineClient& InClient) const;
 	};
 }
 
diff --git a/Engine/Plugins/Developer/Concert/ConcertApp/MultiUserClient/Source/MultiUserClient/Private/Replication/Preset/PresetManager.cpp b/Engine/Plugins/Developer/Concert/ConcertApp/MultiUserClient/Source/MultiUserClient/Private/Replication/Preset/PresetManager.cpp
--- a/Engine/Plugins/Developer/Concert/ConcertApp/MultiUserClient/Source/MultiUserClient/Private/Replication/Preset/PresetManager.cpp
+++ b/Engine/Plugins/Developer/Concert/ConcertApp/MultiUserClient/Source/MultiUserClient/Private/Replication/Preset/PresetManager.cpp
@@ -5,7 +5,7 @@
 #include "Assets/MultiUserReplicationSessionPreset.h"
 #include "ConcertLogGlobal.h"
 #include "IConcertSyncClient.h"
-#include "Replication/Client/ReplicationClientManager.h"
+#include "Replication/Client/Online/OnlineClientManager.h"
 #include "Replication/Misc/ReplicationStreamUtils.h"
 #include "Replication/Muting/MuteStateManager.h"
 #include "Replication/Stream/MultiUserStreamId.h"
@@ -168,14 +168,14 @@ namespace UE::MultiUserClient::Replication
 			EditModel.RemoveObjects(EmptyObjects);
 		}
 
-		static TArray<TPair<const FReplicationClient*, FConcertClientInfo>> DetermineSavedClients(
-			const FReplicationClientManager& ClientManager,
+		static TArray<TPair<const FOnlineClient*, FConcertClientInfo>> DetermineSavedClients(
+			const FOnlineClientManager& ClientManager,
 			const IConcertClientSession& Session,
 			const FSavePresetOptions& Options
 			)
 		{
-			TArray<TPair<const FReplicationClient*, FConcertClientInfo>> IncludedClients;
-			ClientManager.ForEachClient([&Session, &IncludedClients, &Options](const FReplicationClient& Client)
+			TArray<TPair<const FOnlineClient*, FConcertClientInfo>> IncludedClients;
+			ClientManager.ForEachClient([&Session, &IncludedClients, &Options](const FOnlineClient& Client)
 			{
 				FConcertClientInfo ClientInfo;
 				const bool bGotClientInfo = ClientUtils::GetClientDisplayInfo(Session, Client.GetEndpointId(), ClientInfo);
@@ -199,7 +199,7 @@ namespace UE::MultiUserClient::Replication
 	
 	FPresetManager::FPresetManager(
 		const IConcertSyncClient& SyncClient,
-		const FReplicationClientManager& ClientManager,
+		const FOnlineClientManager& ClientManager,
 		const FMuteStateSynchronizer& MuteStateSynchronizer
 		)
 		: SyncClient(SyncClient)
@@ -266,7 +266,7 @@ namespace UE::MultiUserClient::Replication
 		const TSharedPtr<IConcertClientSession> Session = SyncClient.GetConcertClient()->GetCurrentSession();
 		checkf(Session, TEXT("FPresetManager is only supposed to exist while in a session"));
 		
-		const TArray<TPair<const FReplicationClient*, FConcertClientInfo>> IncludedClients = Private::DetermineSavedClients(ClientManager, *Session, Options);
+		const TArray<TPair<const FOnlineClient*, FConcertClientInfo>> IncludedClients = Private::DetermineSavedClients(ClientManager, *Session, Options);
 		return IncludedClients.IsEmpty()
 			? ECanSaveResult::NoClients
 			: ECanSaveResult::Yes;
@@ -288,7 +288,7 @@ namespace UE::MultiUserClient::Replication
 		const TSharedPtr<IConcertClientSession> Session = SyncClient.GetConcertClient()->GetCurrentSession();
 		checkf(Session, TEXT("FPresetManager is only supposed to exist while in a session"));
 
-		const TArray<TPair<const FReplicationClient*, FConcertClientInfo>> IncludedClients = Private::DetermineSavedClients(ClientManager, *Session, Options);
+		const TArray<TPair<const FOnlineClient*, FConcertClientInfo>> IncludedClients = Private::DetermineSavedClients(ClientManager, *Session, Options);
 		if (IncludedClients.IsEmpty())
 		{
 			return nullptr;
@@ -300,7 +300,7 @@ namespace UE::MultiUserClient::Replication
 			// Mark as transient so FEditorFileUtils::SaveAssetsAs creates a new package for the object.
 			RF_Transient
 			);
-		for (const TPair<const FReplicationClient*, FConcertClientInfo>& ClientData : IncludedClients)
+		for (const TPair<const FOnlineClient*, FConcertClientInfo>& ClientData : IncludedClients)
 		{
 			const auto[Client, ClientInfo] = ClientData;
 			UMultiUserReplicationStream* CopiedClientStream = Client->GetClientStreamObject();
diff --git a/Engine/Plugins/Developer/Concert/ConcertApp/MultiUserClient/Source/MultiUserClient/Private/Replication/Preset/PresetManager.h b/Engine/Plugins/Developer/Concert/ConcertApp/MultiUserClient/Source/MultiUserClient/Private/Replication/Preset/PresetManager.h
--- a/Engine/Plugins/Developer/Concert/ConcertApp/MultiUserClient/Source/MultiUserClient/Private/Replication/Preset/PresetManager.h
+++ b/Engine/Plugins/Developer/Concert/ConcertApp/MultiUserClient/Source/MultiUserClient/Private/Replication/Preset/PresetManager.h
@@ -21,7 +21,7 @@ namespace UE::ConcertSyncClient::Replication { struct FRemoteEditEvent; }
 namespace UE::MultiUserClient::Replication
 {
 	class FMuteStateSynchronizer;
-	class FReplicationClientManager;
+	class FOnlineClientManager;
 	
 	/**
 	 * Implements all logic for managing presets in the MU session: saving and loading presets.
@@ -33,7 +33,7 @@ namespace UE::MultiUserClient::Replication
 		
 		FPresetManager(
 			const IConcertSyncClient& SyncClient UE_LIFETIMEBOUND,
-			const FReplicationClientManager& ClientManager UE_LIFETIMEBOUND,
+			const FOnlineClientManager& ClientManager UE_LIFETIMEBOUND,
 			const FMuteStateSynchronizer& MuteStateSynchronizer UE_LIFETIMEBOUND
 			);
 		~FPresetManager();
@@ -57,7 +57,7 @@ namespace UE::MultiUserClient::Replication
 		/** Used to get display information of clients in the session. */
 		const IConcertSyncClient& SyncClient;
 		/** Used to get the clients' replication content. */
-		const FReplicationClientManager& ClientManager;
+		const FOnlineClientManager& ClientManager;
 		/** Used to get the mute state when saving. */
 		const FMuteStateSynchronizer& MuteStateSynchronizer;
 
diff --git a/Engine/Plugins/Developer/Concert/ConcertApp/MultiUserClient/Source/MultiUserClient/Private/Replication/Submission/External/ExternalClientChangeRequestHandler.h b/Engine/Plugins/Developer/Concert/ConcertApp/MultiUserClient/Source/MultiUserClient/Private/Replication/Submission/External/ExternalClientChangeRequestHandler.h
--- a/Engine/Plugins/Developer/Concert/ConcertApp/MultiUserClient/Source/MultiUserClient/Private/Replication/Submission/External/ExternalClientChangeRequestHandler.h
+++ b/Engine/Plugins/Developer/Concert/ConcertApp/MultiUserClient/Source/MultiUserClient/Private/Replication/Submission/External/ExternalClientChangeRequestHandler.h
@@ -21,7 +21,7 @@ namespace UE::MultiUserClient
 namespace UE::MultiUserClient::Replication
 {
 	class FClientChangeOperation;
-	class FReplicationClientManager;
+	class FOnlineClientManager;
 	class FSubmissionQueue;
 	
 	/**
diff --git a/Engine/Plugins/Developer/Concert/ConcertApp/MultiUserClient/Source/MultiUserClient/Private/Replication/Submission/MultiEdit/ReassignObjectPropertiesLogic.cpp b/Engine/Plugins/Developer/Concert/ConcertApp/MultiUserClient/Source/MultiUserClient/Private/Replication/Submission/MultiEdit/ReassignObjectPropertiesLogic.cpp
--- a/Engine/Plugins/Developer/Concert/ConcertApp/MultiUserClient/Source/MultiUserClient/Private/Replication/Submission/MultiEdit/ReassignObjectPropertiesLogic.cpp
+++ b/Engine/Plugins/Developer/Concert/ConcertApp/MultiUserClient/Source/MultiUserClient/Private/Replication/Submission/MultiEdit/ReassignObjectPropertiesLogic.cpp
@@ -3,7 +3,7 @@
 #include "ReassignObjectPropertiesLogic.h"
 
 #include "ConcertLogGlobal.h"
-#include "Replication/Client/ReplicationClientManager.h"
+#include "Replication/Client/Online/OnlineClientManager.h"
 #include "Replication/Submission/ISubmissionOperation.h"
 #include "Replication/Misc/Util/SynchronizedRequestUtils.h"
 
@@ -32,7 +32,7 @@ namespace UE::MultiUserClient::Replication
 		};
 		
 		/** Builds all the changes that must be made to the clients we are transferring from.  */
-		static FObjectReassignment BuildChangesForTransferal(const FReplicationClientManager& ClientManager, TConstArrayView<FSoftObjectPath> ObjectsToReassign, const FGuid& StreamId)
+		static FObjectReassignment BuildChangesForTransferal(const FOnlineClientManager& ClientManager, TConstArrayView<FSoftObjectPath> ObjectsToReassign, const FGuid& StreamId)
 		{
 			FObjectReassignment Result;
 			TMap<FGuid, FConcertObjectReplicationMap>& OldRegisteredObjects = Result.OldRegisteredObjects;
@@ -49,7 +49,7 @@ namespace UE::MultiUserClient::Replication
 					[&ClientManager, &OldRegisteredObjects, &OldFrequencies, &ReassignedClientRequests, &ReassignedAuthority, &ObjectPath, &ObjectId, &AuthorityCache](const FGuid& ClientId)
 					{
 						// Get the property assignments of the client being reassigned from
-						const FReplicationClient* ClientToReassignFrom = ClientManager.FindClient(ClientId);
+						const FOnlineClient* ClientToReassignFrom = ClientManager.FindClient(ClientId);
 						if (!ensureMsgf(ClientToReassignFrom, TEXT("Authority cache out of sync"))
 							|| !ClientToReassignFrom->AllowsEditing())
 						{
@@ -87,7 +87,7 @@ namespace UE::MultiUserClient::Replication
 			// Need to set Defaults of the OldFrequencies because that's not been done, yet.
 			for (TPair<FGuid, FConcertStreamFrequencySettings>& OldFrequency : OldFrequencies)
 			{
-				const FReplicationClient* ClientToReassignFrom = ClientManager.FindClient(OldFrequency.Key);
+				const FOnlineClient* ClientToReassignFrom = ClientManager.FindClient(OldFrequency.Key);
 				check(ClientToReassignFrom);
 				OldFrequency.Value.Defaults = ClientToReassignFrom->GetStreamSynchronizer().GetFrequencySettings().Defaults;
 			}
@@ -204,12 +204,12 @@ namespace UE::MultiUserClient::Replication
 		/** @return Whether TargetClientObjects includes all properties ClientsToConsider have registered for ObjectToCheck. */
 		static bool DoesTargetIncludeOthers_SingleObject(
 			const FConcertObjectReplicationMap& TargetClientObjects,
-			const TArray<TNonNullPtr<const FReplicationClient>> ClientsToConsider,
+			const TArray<TNonNullPtr<const FOnlineClient>> ClientsToConsider,
 			const FSoftObjectPath& ObjectToCheck)
 		{
 			const FConcertReplicatedObjectInfo* ObjectInfo = TargetClientObjects.ReplicatedObjects.Find(ObjectToCheck);
 				
-			for (const FReplicationClient* OtherClient : ClientsToConsider)
+			for (const FOnlineClient* OtherClient : ClientsToConsider)
 			{
 				const FConcertObjectReplicationMap& OtherClientObjects = OtherClient->GetStreamSynchronizer().GetServerState();
 				const FConcertReplicatedObjectInfo* OtherObjectInfo = OtherClientObjects.ReplicatedObjects.Find(ObjectToCheck);
@@ -233,18 +233,18 @@ namespace UE::MultiUserClient::Replication
 		static bool DoesTargetIncludeOthers(
 			const FGuid& TargetClientId,
 			TConstArrayView<FSoftObjectPath> ObjectsToCheck,
-			const FReplicationClientManager& ClientManager,
-			TFunctionRef<bool(const FReplicationClient& Client)> ShouldConsiderClientPredicate
+			const FOnlineClientManager& ClientManager,
+			TFunctionRef<bool(const FOnlineClient& Client)> ShouldConsiderClientPredicate
 			)
 		{
-			const FReplicationClient* TargetClient = ClientManager.FindClient(TargetClientId);
+			const FOnlineClient* TargetClient = ClientManager.FindClient(TargetClientId);
 			if (!ensure(TargetClient))
 			{
 				return false;
 			}
 
 			const FConcertObjectReplicationMap& TargetClientObjects = TargetClient->GetStreamSynchronizer().GetServerState();
-			const TArray<TNonNullPtr<const FReplicationClient>> ClientsToConsider = ClientManager.GetClients(ShouldConsiderClientPredicate);
+			const TArray<TNonNullPtr<const FOnlineClient>> ClientsToConsider = ClientManager.GetClients(ShouldConsiderClientPredicate);
 			return Algo::AllOf(ObjectsToCheck, [&TargetClientObjects, &ClientsToConsider](const FSoftObjectPath& ObjectToCheck)
 			{
 				return DoesTargetIncludeOthers_SingleObject(TargetClientObjects, ClientsToConsider, ObjectToCheck);
@@ -252,9 +252,9 @@ namespace UE::MultiUserClient::Replication
 		}
 
 		/** @return Whether ClientId has any properties assigned to any of the objects in ObjectsToCheck. */
-		static bool HasAnyProperties(const FGuid& ClientId, TConstArrayView<FSoftObjectPath> ObjectsToCheck, const FReplicationClientManager& ClientManager)
+		static bool HasAnyProperties(const FGuid& ClientId, TConstArrayView<FSoftObjectPath> ObjectsToCheck, const FOnlineClientManager& ClientManager)
 		{
-			const FReplicationClient* TargetClient = ClientManager.FindClient(ClientId);
+			const FOnlineClient* TargetClient = ClientManager.FindClient(ClientId);
 			if (!ensure(TargetClient))
 			{
 				return false;
@@ -269,7 +269,7 @@ namespace UE::MultiUserClient::Replication
 		}
 	}
 	
-	FReassignObjectPropertiesLogic::FReassignObjectPropertiesLogic(FReplicationClientManager& InClientManager)
+	FReassignObjectPropertiesLogic::FReassignObjectPropertiesLogic(FOnlineClientManager& InClientManager)
 		: ClientManager(InClientManager)
 	{
 		ClientManager.OnRemoteClientsChanged().AddRaw(this, &FReassignObjectPropertiesLogic::OnRemoteClientsChanged);
@@ -291,7 +291,7 @@ namespace UE::MultiUserClient::Replication
 
 	void FReassignObjectPropertiesLogic::EnumerateClientOwnershipState(const FSoftObjectPath& ObjectPath, FProcessClientOwnership Callback) const
 	{
-		ClientManager.ForEachClient([&ObjectPath, &Callback](const FReplicationClient& ReplicationClient)
+		ClientManager.ForEachClient([&ObjectPath, &Callback](const FOnlineClient& ReplicationClient)
 		{
 			const bool bHasProperties = ReplicationClient.GetStreamSynchronizer().GetServerState().HasProperties(ObjectPath);
 			const EOwnershipState Ownership = bHasProperties ? EOwnershipState::HasObjectRegistered : EOwnershipState::NoOwnership;
@@ -337,7 +337,7 @@ namespace UE::MultiUserClient::Replication
 			return false;
 		}
 		
-		const FReplicationClient* Client = ClientManager.FindClient(ClientId);
+		const FOnlineClient* Client = ClientManager.FindClient(ClientId);
 		if (!Client)
 		{
 			SET_REASON(LOCTEXT("Reason.ClientDisconnected", "Client disconnected"));
@@ -351,7 +351,7 @@ namespace UE::MultiUserClient::Replication
 		}
 
 		// Relevant after user has re-assigned all properties to a given client
-		auto ShouldConsiderClient = [this](const FReplicationClient& Client) { return Client.AllowsEditing(); };
+		auto ShouldConsiderClient = [this](const FOnlineClient& Client) { return Client.AllowsEditing(); };
 		if (ReassignObjectProperties::Private::DoesTargetIncludeOthers(ClientId, ObjectsToReassign, ClientManager, ShouldConsiderClient))
 		{
 			const bool bHasAnyProperties = ReassignObjectProperties::Private::HasAnyProperties(ClientId, ObjectsToReassign, ClientManager);
@@ -372,7 +372,7 @@ namespace UE::MultiUserClient::Replication
 	{
 		using namespace ConcertSyncClient::Replication;
 		
-		const FReplicationClient* ClientToAssignTo = ClientManager.FindClient(ClientId);
+		const FOnlineClient* ClientToAssignTo = ClientManager.FindClient(ClientId);
 		if (!ensure(ClientToAssignTo && ClientToAssignTo->AllowsEditing()))
 		{
 			UE_LOG(LogConcert, Error, TEXT("Property Reassignment: The target client is not editable."));
@@ -462,7 +462,7 @@ namespace UE::MultiUserClient::Replication
 
 	FSubmissionQueue* FReassignObjectPropertiesLogic::FindSubmissionQueueForTargetClient() const
 	{
-		FReplicationClient* Client = InProgressOperation.IsSet()
+		FOnlineClient* Client = InProgressOperation.IsSet()
 			? ClientManager.FindClient(InProgressOperation->AssignedToClient)
 			: nullptr;
 		return Client ? &Client->GetSubmissionQueue() : nullptr;
@@ -476,7 +476,7 @@ namespace UE::MultiUserClient::Replication
 			return;
 		}
 		
-		const FReplicationClient* ClientToAssignTo = ClientManager.FindClient(InProgressOperation->AssignedToClient);
+		const FOnlineClient* ClientToAssignTo = ClientManager.FindClient(InProgressOperation->AssignedToClient);
 		if (!ensureMsgf(ClientToAssignTo, TEXT("FReassignObjectPropertiesLogic should have cancelled the change upon disconnect"))
 			|| !ClientToAssignTo->AllowsEditing())
 		{
diff --git a/Engine/Plugins/Developer/Concert/ConcertApp/MultiUserClient/Source/MultiUserClient/Private/Replication/Submission/MultiEdit/ReassignObjectPropertiesLogic.h b/Engine/Plugins/Developer/Concert/ConcertApp/MultiUserClient/Source/MultiUserClient/Private/Replication/Submission/MultiEdit/ReassignObjectPropertiesLogic.h
--- a/Engine/Plugins/Developer/Concert/ConcertApp/MultiUserClient/Source/MultiUserClient/Private/Replication/Submission/MultiEdit/ReassignObjectPropertiesLogic.h
+++ b/Engine/Plugins/Developer/Concert/ConcertApp/MultiUserClient/Source/MultiUserClient/Private/Replication/Submission/MultiEdit/ReassignObjectPropertiesLogic.h
@@ -15,7 +15,7 @@ struct FConcertPropertyChain;
 
 namespace UE::MultiUserClient::Replication
 {
-	class FReplicationClientManager;
+	class FOnlineClientManager;
 	class IParallelSubmissionOperation;
 	
 	/**
@@ -41,7 +41,7 @@ namespace UE::MultiUserClient::Replication
 		};
 		using FProcessClientOwnership = TFunctionRef<EBreakBehavior(const FGuid& ClientId, EOwnershipState Ownership)>;
 		
-		FReassignObjectPropertiesLogic(FReplicationClientManager& InClientManager);
+		FReassignObjectPropertiesLogic(FOnlineClientManager& InClientManager);
 		virtual ~FReassignObjectPropertiesLogic() override;
 
 		/** Enumerates every client that has registered properties to this object. */
@@ -70,7 +70,7 @@ namespace UE::MultiUserClient::Replication
 		const TSharedRef<FToken> DestructionToken = FToken::Make();
 
 		/** Used to get info about clients and register to client list changing. */
-		FReplicationClientManager& ClientManager;
+		FOnlineClientManager& ClientManager;
 
 		struct FOperationData
 		{
diff --git a/Engine/Plugins/Developer/Concert/ConcertApp/MultiUserClient/Source/MultiUserClient/Private/Replication/Submission/Notification/SubmissionNotifier.cpp b/Engine/Plugins/Developer/Concert/ConcertApp/MultiUserClient/Source/MultiUserClient/Private/Replication/Submission/Notification/SubmissionNotifier.cpp
--- a/Engine/Plugins/Developer/Concert/ConcertApp/MultiUserClient/Source/MultiUserClient/Private/Replication/Submission/Notification/SubmissionNotifier.cpp
+++ b/Engine/Plugins/Developer/Concert/ConcertApp/MultiUserClient/Source/MultiUserClient/Private/Replication/Submission/Notification/SubmissionNotifier.cpp
@@ -4,8 +4,8 @@
 
 #include "SAuthorityRejectedNotification.h"
 #include "SStreamRejectedNotification.h"
-#include "Replication/Client/RemoteReplicationClient.h"
-#include "Replication/Client/ReplicationClientManager.h"
+#include "Replication/Client/Online/RemoteClient.h"
+#include "Replication/Client/Online/OnlineClientManager.h"
 #include "Replication/Submission/Data/AuthoritySubmission.h"
 #include "Replication/Submission/Data/StreamSubmission.h"
 
@@ -16,7 +16,7 @@
 
 namespace UE::MultiUserClient::Replication
 {
-	FSubmissionNotifier::FSubmissionNotifier(FReplicationClientManager& InReplicationClientManager)
+	FSubmissionNotifier::FSubmissionNotifier(FOnlineClientManager& InReplicationClientManager)
 		: ReplicationClientManager(InReplicationClientManager)
 	{
 		ReplicationClientManager.OnPostRemoteClientAdded().AddRaw(this, &FSubmissionNotifier::OnPostRemoteClientAdded);
@@ -34,7 +34,7 @@ namespace UE::MultiUserClient::Replication
 		ReplicationClientManager.OnPreRemoteClientRemoved().RemoveAll(this);
 		
 		UnregisterClient(ReplicationClientManager.GetLocalClient());
-		for (const TNonNullPtr<FRemoteReplicationClient>& RemoteClient : ReplicationClientManager.GetRemoteClients())
+		for (const TNonNullPtr<FRemoteClient>& RemoteClient : ReplicationClientManager.GetRemoteClients())
 		{
 			UnregisterClient(*RemoteClient);
 		}
@@ -51,18 +51,18 @@ namespace UE::MultiUserClient::Replication
 		RETURN_QUICK_DECLARE_CYCLE_STAT(FSubmissionNotifier, STATGROUP_Tickables);
 	}
 
-	void FSubmissionNotifier::OnPostRemoteClientAdded(FRemoteReplicationClient& RemoteReplicationClient)
+	void FSubmissionNotifier::OnPostRemoteClientAdded(FRemoteClient& RemoteReplicationClient)
 	{
 		RegisterClient(RemoteReplicationClient);
 	}
 
-	void FSubmissionNotifier::RegisterClient(FReplicationClient& Client)
+	void FSubmissionNotifier::RegisterClient(FOnlineClient& Client)
 	{
 		Client.GetSubmissionWorkflow().OnStreamRequestCompleted_AnyThread().AddRaw(this, &FSubmissionNotifier::OnStreamRequestCompleted_AnyThread);
 		Client.GetSubmissionWorkflow().OnAuthorityRequestCompleted_AnyThread().AddRaw(this, &FSubmissionNotifier::OnAuthorityRequestCompleted_AnyThread);
 	}
 
-	void FSubmissionNotifier::UnregisterClient(FReplicationClient& Client)
+	void FSubmissionNotifier::UnregisterClient(FOnlineClient& Client)
 	{
 		Client.GetSubmissionWorkflow().OnStreamRequestCompleted_AnyThread().RemoveAll(this);
 		Client.GetSubmissionWorkflow().OnAuthorityRequestCompleted_AnyThread().RemoveAll(this);
diff --git a/Engine/Plugins/Developer/Concert/ConcertApp/MultiUserClient/Source/MultiUserClient/Private/Replication/Submission/Notification/SubmissionNotifier.h b/Engine/Plugins/Developer/Concert/ConcertApp/MultiUserClient/Source/MultiUserClient/Private/Replication/Submission/Notification/SubmissionNotifier.h
--- a/Engine/Plugins/Developer/Concert/ConcertApp/MultiUserClient/Source/MultiUserClient/Private/Replication/Submission/Notification/SubmissionNotifier.h
+++ b/Engine/Plugins/Developer/Concert/ConcertApp/MultiUserClient/Source/MultiUserClient/Private/Replication/Submission/Notification/SubmissionNotifier.h
@@ -13,9 +13,9 @@ class FReply;
 
 namespace UE::MultiUserClient::Replication
 {
-	class FRemoteReplicationClient;
-	class FReplicationClient;
-	class FReplicationClientManager;
+	class FRemoteClient;
+	class FOnlineClient;
+	class FOnlineClientManager;
 	
 	class SAuthorityRejectedNotification;
 	class SStreamRejectedNotification;
@@ -32,7 +32,7 @@ namespace UE::MultiUserClient::Replication
 	{
 	public:
 		
-		FSubmissionNotifier(FReplicationClientManager& InReplicationClientManager);
+		FSubmissionNotifier(FOnlineClientManager& InReplicationClientManager);
 		~FSubmissionNotifier();
 
 		//~ Begin FTickableEditorObject Interface
@@ -43,7 +43,7 @@ namespace UE::MultiUserClient::Replication
 	private:
 
 		/** Emits events when remote clients are discovered. */
-		FReplicationClientManager& ReplicationClientManager;
+		FOnlineClientManager& ReplicationClientManager;
 
 		/** Only valid if changing streams has failed and the user has not yet dismissed the notification. */
 		TSharedPtr<SNotificationItem> StreamNotificationItem;
@@ -67,9 +67,9 @@ namespace UE::MultiUserClient::Replication
 		
 		
 		// Client events
-		void OnPostRemoteClientAdded(FRemoteReplicationClient& RemoteReplicationClient);
-		void RegisterClient(FReplicationClient& Client);
-		void UnregisterClient(FReplicationClient& Client);
+		void OnPostRemoteClientAdded(FRemoteClient& RemoteReplicationClient);
+		void RegisterClient(FOnlineClient& Client);
+		void UnregisterClient(FOnlineClient& Client);
 
 		// Per-client submission events
 		void OnStreamRequestCompleted_AnyThread(const FSubmitStreamChangesResponse& Response);
diff --git a/Engine/Plugins/Developer/Concert/ConcertApp/MultiUserClient/Source/MultiUserClient/Private/Widgets/ActiveSession/Replication/Client/ClientUtils.cpp b/Engine/Plugins/Developer/Concert/ConcertApp/MultiUserClient/Source/MultiUserClient/Private/Widgets/ActiveSession/Replication/Client/ClientUtils.cpp
--- a/Engine/Plugins/Developer/Concert/ConcertApp/MultiUserClient/Source/MultiUserClient/Private/Widgets/ActiveSession/Replication/Client/ClientUtils.cpp
+++ b/Engine/Plugins/Developer/Concert/ConcertApp/MultiUserClient/Source/MultiUserClient/Private/Widgets/ActiveSession/Replication/Client/ClientUtils.cpp
@@ -4,7 +4,7 @@
 
 #include "ConcertMessageData.h"
 #include "IConcertClient.h"
-#include "Replication/Client/ReplicationClientManager.h"
+#include "Replication/Client/Online/OnlineClientManager.h"
 #include "Widgets/Client/SClientName.h"
 
 namespace UE::MultiUserClient::Replication::ClientUtils
@@ -58,16 +58,16 @@ namespace UE::MultiUserClient::Replication::ClientUtils
 		return false;
 	}
 	
-	TArray<const FReplicationClient*> GetSortedClientList(const IConcertClient& InLocalClientInstance, const FReplicationClientManager& InReplicationManager)
+	TArray<const FOnlineClient*> GetSortedClientList(const IConcertClient& InLocalClientInstance, const FOnlineClientManager& InReplicationManager)
 	{
 		return GetSortedClientList(*InLocalClientInstance.GetCurrentSession(), InReplicationManager);
 	}
 
-	TArray<const FReplicationClient*> GetSortedClientList(const IConcertClientSession& InSession, const FReplicationClientManager& InReplicationManager)
+	TArray<const FOnlineClient*> GetSortedClientList(const IConcertClientSession& InSession, const FOnlineClientManager& InReplicationManager)
 	{
-		TArray<const FReplicationClient*> Result;
-		TMap<const FReplicationClient*, FConcertClientInfo> ClientToDisplayInfo;
-		for (const TNonNullPtr<const FRemoteReplicationClient> RemoteClient : InReplicationManager.GetRemoteClients())
+		TArray<const FOnlineClient*> Result;
+		TMap<const FOnlineClient*, FConcertClientInfo> ClientToDisplayInfo;
+		for (const TNonNullPtr<const FRemoteClient> RemoteClient : InReplicationManager.GetRemoteClients())
 		{
 			FConcertClientInfo Info;
 			if (GetClientDisplayInfo(InSession, RemoteClient->GetEndpointId(), Info))
@@ -77,7 +77,7 @@ namespace UE::MultiUserClient::Replication::ClientUtils
 			}
 		}
 
-		Result.Sort([&ClientToDisplayInfo](const FReplicationClient& Left, const FReplicationClient& Right)
+		Result.Sort([&ClientToDisplayInfo](const FOnlineClient& Left, const FOnlineClient& Right)
 		{
 			const FConcertClientInfo& LeftInfo = ClientToDisplayInfo[&Left];
 			const FConcertClientInfo& RightInfo = ClientToDisplayInfo[&Right];
diff --git a/Engine/Plugins/Developer/Concert/ConcertApp/MultiUserClient/Source/MultiUserClient/Private/Widgets/ActiveSession/Replication/Client/ClientUtils.h b/Engine/Plugins/Developer/Concert/ConcertApp/MultiUserClient/Source/MultiUserClient/Private/Widgets/ActiveSession/Replication/Client/ClientUtils.h
--- a/Engine/Plugins/Developer/Concert/ConcertApp/MultiUserClient/Source/MultiUserClient/Private/Widgets/ActiveSession/Replication/Client/ClientUtils.h
+++ b/Engine/Plugins/Developer/Concert/ConcertApp/MultiUserClient/Source/MultiUserClient/Private/Widgets/ActiveSession/Replication/Client/ClientUtils.h
@@ -13,8 +13,8 @@ struct FGuid;
 
 namespace UE::MultiUserClient::Replication
 {
-	class FReplicationClient;
-	class FReplicationClientManager;
+	class FOnlineClient;
+	class FOnlineClientManager;
 }
 
 namespace UE::MultiUserClient::Replication::ClientUtils
@@ -46,6 +46,6 @@ namespace UE::MultiUserClient::Replication::ClientUtils
 	 * @param InReplicationManager Used to obtain all clients
 	 * @return Sorted client array
 	 */
-	TArray<const FReplicationClient*> GetSortedClientList(const IConcertClient& InLocalClientInstance, const FReplicationClientManager& InReplicationManager);
-	TArray<const FReplicationClient*> GetSortedClientList(const IConcertClientSession& InSession, const FReplicationClientManager& InReplicationManager);
+	TArray<const FOnlineClient*> GetSortedClientList(const IConcertClient& InLocalClientInstance, const FOnlineClientManager& InReplicationManager);
+	TArray<const FOnlineClient*> GetSortedClientList(const IConcertClientSession& InSession, const FOnlineClientManager& InReplicationManager);
 }
diff --git a/Engine/Plugins/Developer/Concert/ConcertApp/MultiUserClient/Source/MultiUserClient/Private/Widgets/ActiveSession/Replication/Client/Context/ContextMenuUtils.h b/Engine/Plugins/Developer/Concert/ConcertApp/MultiUserClient/Source/MultiUserClient/Private/Widgets/ActiveSession/Replication/Client/Context/ContextMenuUtils.h
--- a/Engine/Plugins/Developer/Concert/ConcertApp/MultiUserClient/Source/MultiUserClient/Private/Widgets/ActiveSession/Replication/Client/Context/ContextMenuUtils.h
+++ b/Engine/Plugins/Developer/Concert/ConcertApp/MultiUserClient/Source/MultiUserClient/Private/Widgets/ActiveSession/Replication/Client/Context/ContextMenuUtils.h
@@ -17,7 +17,7 @@ namespace UE::ConcertSharedSlate
 namespace UE::MultiUserClient::Replication
 {
 	class FReassignObjectPropertiesLogic;
-	class FReplicationClientManager;
+	class FOnlineClientManager;
 }
 
 namespace UE::MultiUserClient::Replication::ContextMenuUtils
@@ -27,16 +27,16 @@ namespace UE::MultiUserClient::Replication::ContextMenuUtils
 		FMenuBuilder& MenuBuilder,
 		const TSoftObjectPtr<>& ContextObject,
 		const IConcertClient& ConcertClient,
-		const FReplicationClientManager& ReplicationManager,
+		const FOnlineClientManager& ReplicationManager,
 		ConcertSharedSlate::IObjectHierarchyModel& ObjectHierarchy,
 		FReassignObjectPropertiesLogic& ReassignmentLogic,
 		ConcertSharedSlate::IMultiReplicationStreamEditor& MultiStreamEditor
 		);
 	
 	/** Adds an edit box for batch reassigning the select object's frequencies for all replicating clients. */
-	void AddFrequencyOptionsForMultipleClients(FMenuBuilder& MenuBuilder, const FSoftObjectPath& ContextObjects, FReplicationClientManager& InClientManager);
+	void AddFrequencyOptionsForMultipleClients(FMenuBuilder& MenuBuilder, const FSoftObjectPath& ContextObjects, FOnlineClientManager& InClientManager);
 	/** Adds an edit box for batch reassigning the select object's frequencies for all replicating clients. */
-	inline void AddFrequencyOptionsIfOneContextObject_MultiClient(FMenuBuilder& MenuBuilder, TConstArrayView<TSoftObjectPtr<>> ContextObjects, FReplicationClientManager& InClientManager)
+	inline void AddFrequencyOptionsIfOneContextObject_MultiClient(FMenuBuilder& MenuBuilder, TConstArrayView<TSoftObjectPtr<>> ContextObjects, FOnlineClientManager& InClientManager)
 	{
 		if (ContextObjects.Num() == 1)
 		{
diff --git a/Engine/Plugins/Developer/Concert/ConcertApp/MultiUserClient/Source/MultiUserClient/Private/Widgets/ActiveSession/Replication/Client/Context/ContextMenuUtils_Frequency.cpp b/Engine/Plugins/Developer/Concert/ConcertApp/MultiUserClient/Source/MultiUserClient/Private/Widgets/ActiveSession/Replication/Client/Context/ContextMenuUtils_Frequency.cpp
--- a/Engine/Plugins/Developer/Concert/ConcertApp/MultiUserClient/Source/MultiUserClient/Private/Widgets/ActiveSession/Replication/Client/Context/ContextMenuUtils_Frequency.cpp
+++ b/Engine/Plugins/Developer/Concert/ConcertApp/MultiUserClient/Source/MultiUserClient/Private/Widgets/ActiveSession/Replication/Client/Context/ContextMenuUtils_Frequency.cpp
@@ -2,7 +2,7 @@
 
 #include "ContextMenuUtils.h"
 
-#include "Replication/Client/ReplicationClientManager.h"
+#include "Replication/Client/Online/OnlineClientManager.h"
 #include "Replication/Misc/Frequency/FrequencyUtils.h"
 #include "Replication/Misc/Util/SynchronizedRequestUtils.h"
 
@@ -40,7 +40,7 @@ namespace UE::MultiUserClient::Replication::ContextMenuUtils
 				SLATE_ATTRIBUTE(FInlineClientArray, Clients)
 			SLATE_END_ARGS()
 			
-			void Construct(const FArguments& InArgs, FReplicationClientManager& InClientManager)
+			void Construct(const FArguments& InArgs, FOnlineClientManager& InClientManager)
 			{
 				ContextObject = InArgs._SelectedObject;
 				ClientsAttribute = InArgs._Clients;
@@ -82,7 +82,7 @@ namespace UE::MultiUserClient::Replication::ContextMenuUtils
 
 			FSoftObjectPath ContextObject;
 			TAttribute<FInlineClientArray> ClientsAttribute;
-			FReplicationClientManager* ClientManager = nullptr;
+			FOnlineClientManager* ClientManager = nullptr;
 
 			struct FChangeOperation
 			{
@@ -182,7 +182,7 @@ namespace UE::MultiUserClient::Replication::ContextMenuUtils
 			}
 		};
 
-		static void AddReplicationModeSubMenu(FMenuBuilder& MenuBuilder, const FSoftObjectPath& ContextObject, TAttribute<FInlineClientArray> Clients, FReplicationClientManager& InClientManager)
+		static void AddReplicationModeSubMenu(FMenuBuilder& MenuBuilder, const FSoftObjectPath& ContextObject, TAttribute<FInlineClientArray> Clients, FOnlineClientManager& InClientManager)
 		{
 			const auto SetReplicationMode = [ContextObject, Clients, &InClientManager](const EConcertObjectReplicationMode ModeToSet)
 			{
@@ -242,7 +242,7 @@ namespace UE::MultiUserClient::Replication::ContextMenuUtils
 			FMenuBuilder& MenuBuilder,
 			const FSoftObjectPath& ContextObject,
 			TAttribute<FInlineClientArray> GetClientsAttribute,
-			FReplicationClientManager& InClientManager
+			FOnlineClientManager& InClientManager
 			)
 		{
 			// Change to Realtime / Specified Rate
@@ -270,7 +270,7 @@ namespace UE::MultiUserClient::Replication::ContextMenuUtils
 			FMenuBuilder& MenuBuilder,
 			const FSoftObjectPath& ContextObject,
 			TFunction<FInlineClientArray(const FSoftObjectPath& Object)> GetClientsForObjectFunc,
-			FReplicationClientManager& InClientManager
+			FOnlineClientManager& InClientManager
 			)
 		{
 			bool bAddedSubmenu = false;
@@ -318,13 +318,13 @@ namespace UE::MultiUserClient::Replication::ContextMenuUtils
 		}
 	}
 
-	void AddFrequencyOptionsForMultipleClients(FMenuBuilder& MenuBuilder, const FSoftObjectPath& ContextObject, FReplicationClientManager& InClientManager)
+	void AddFrequencyOptionsForMultipleClients(FMenuBuilder& MenuBuilder, const FSoftObjectPath& ContextObject, FOnlineClientManager& InClientManager)
 	{
 		using namespace Private;
 		const auto GetClientsForObject = [&InClientManager](const FSoftObjectPath& Object)
 		{
 			FInlineClientArray Result;
-			InClientManager.ForEachClient([&Object, &Result](const FReplicationClient& Client)
+			InClientManager.ForEachClient([&Object, &Result](const FOnlineClient& Client)
 			{
 				// Only clients that have the object registered should be considered
 				if (Client.GetStreamSynchronizer().GetServerState().HasProperties(Object))
diff --git a/Engine/Plugins/Developer/Concert/ConcertApp/MultiUserClient/Source/MultiUserClient/Private/Widgets/ActiveSession/Replication/Client/Context/ContextMenuUtils_Reassignment.cpp b/Engine/Plugins/Developer/Concert/ConcertApp/MultiUserClient/Source/MultiUserClient/Private/Widgets/ActiveSession/Replication/Client/Context/ContextMenuUtils_Reassignment.cpp
--- a/Engine/Plugins/Developer/Concert/ConcertApp/MultiUserClient/Source/MultiUserClient/Private/Widgets/ActiveSession/Replication/Client/Context/ContextMenuUtils_Reassignment.cpp
+++ b/Engine/Plugins/Developer/Concert/ConcertApp/MultiUserClient/Source/MultiUserClient/Private/Widgets/ActiveSession/Replication/Client/Context/ContextMenuUtils_Reassignment.cpp
@@ -2,7 +2,7 @@
 
 #include "ContextMenuUtils.h"
 
-#include "Replication/Client/ReplicationClient.h"
+#include "Replication/Client/Online/OnlineClient.h"
 #include "Replication/Editor/Model/Object/IObjectHierarchyModel.h"
 #include "Replication/Editor/View/IMultiReplicationStreamEditor.h"
 #include "Replication/Editor/View/IReplicationStreamEditor.h"
@@ -35,14 +35,14 @@ namespace UE::MultiUserClient::Replication::ContextMenuUtils
 		
 		static void AddReassignSection(
 			FMenuBuilder& MenuBuilder,
-			const TArray<const FReplicationClient*>& SortedClients,
+			const TArray<const FOnlineClient*>& SortedClients,
 			TAttribute<FInlineObjectPathArray> ObjectsToAssign,
 			const IConcertClient& ConcertClient,
 			FReassignObjectPropertiesLogic& ReassignmentLogic,
 			ConcertSharedSlate::IMultiReplicationStreamEditor& MultiStreamEditor
 			)
 		{
-			for (const FReplicationClient* Client : SortedClients)
+			for (const FOnlineClient* Client : SortedClients)
 			{
 				const FGuid& ClientId = Client->GetEndpointId();
 				MenuBuilder.AddMenuEntry(
@@ -76,13 +76,13 @@ namespace UE::MultiUserClient::Replication::ContextMenuUtils
 		FMenuBuilder& MenuBuilder,
 		const TSoftObjectPtr<>& ContextObject,
 		const IConcertClient& ConcertClient,
-		const FReplicationClientManager& ReplicationManager,
+		const FOnlineClientManager& ReplicationManager,
 		ConcertSharedSlate::IObjectHierarchyModel& ObjectHierarchy,
 		FReassignObjectPropertiesLogic& ReassignmentLogic,
 		ConcertSharedSlate::IMultiReplicationStreamEditor& MultiStreamEditor
 		)
 	{
-		const TArray<const FReplicationClient*> SortedClients = ClientUtils::GetSortedClientList(ConcertClient, ReplicationManager);
+		const TArray<const FOnlineClient*> SortedClients = ClientUtils::GetSortedClientList(ConcertClient, ReplicationManager);
 		
 		MenuBuilder.BeginSection(TEXT("Reassign.This"), LOCTEXT("Reassign.This", "Reassign this to"));
 		Private::AddReassignSection(MenuBuilder, SortedClients, Private::FInlineObjectPathArray{ ContextObject.GetUniqueID() }, ConcertClient, ReassignmentLogic, MultiStreamEditor);
diff --git a/Engine/Plugins/Developer/Concert/ConcertApp/MultiUserClient/Source/MultiUserClient/Private/Widgets/ActiveSession/Replication/Client/FrequencyContextMenuUtils.h b/Engine/Plugins/Developer/Concert/ConcertApp/MultiUserClient/Source/MultiUserClient/Private/Widgets/ActiveSession/Replication/Client/FrequencyContextMenuUtils.h
--- a/Engine/Plugins/Developer/Concert/ConcertApp/MultiUserClient/Source/MultiUserClient/Private/Widgets/ActiveSession/Replication/Client/FrequencyContextMenuUtils.h
+++ b/Engine/Plugins/Developer/Concert/ConcertApp/MultiUserClient/Source/MultiUserClient/Private/Widgets/ActiveSession/Replication/Client/FrequencyContextMenuUtils.h
@@ -10,18 +10,18 @@ struct FGuid;
 
 namespace UE::MultiUserClient::Replication
 {
-	class FReplicationClientManager;
+	class FOnlineClientManager;
 }
 
 namespace UE::MultiUserClient::Replication::FrequencyContextMenuUtils
 {
 	/** Adds an edit box for changing the selected object's frequencies. */
-	void AddFrequencyOptionsForSingleClient(FMenuBuilder& MenuBuilder, const FSoftObjectPath& ContextObjects, const FGuid& ClientId, FReplicationClientManager& InClientManager);
+	void AddFrequencyOptionsForSingleClient(FMenuBuilder& MenuBuilder, const FSoftObjectPath& ContextObjects, const FGuid& ClientId, FOnlineClientManager& InClientManager);
 	/** Adds an edit box for batch reassigning the select object's frequencies for all replicating clients. */
-	void AddFrequencyOptionsForMultipleClients(FMenuBuilder& MenuBuilder, const FSoftObjectPath& ContextObjects, FReplicationClientManager& InClientManager);
+	void AddFrequencyOptionsForMultipleClients(FMenuBuilder& MenuBuilder, const FSoftObjectPath& ContextObjects, FOnlineClientManager& InClientManager);
 	
 	/** Adds an edit box for changing the selected object's frequencies. */
-	inline void AddFrequencyOptionsIfOneContextObject_SingleClient(FMenuBuilder& MenuBuilder, TConstArrayView<FSoftObjectPath> ContextObjects, const FGuid& ClientId, FReplicationClientManager& InClientManager)
+	inline void AddFrequencyOptionsIfOneContextObject_SingleClient(FMenuBuilder& MenuBuilder, TConstArrayView<FSoftObjectPath> ContextObjects, const FGuid& ClientId, FOnlineClientManager& InClientManager)
 	{
 		if (ContextObjects.Num() == 1)
 		{
@@ -29,7 +29,7 @@ namespace UE::MultiUserClient::Replication::FrequencyContextMenuUtils
 		}
 	}
 	/** Adds an edit box for batch reassigning the select object's frequencies for all replicating clients. */
-	inline void AddFrequencyOptionsIfOneContextObject_MultiClient(FMenuBuilder& MenuBuilder, TConstArrayView<FSoftObjectPath> ContextObjects, FReplicationClientManager& InClientManager)
+	inline void AddFrequencyOptionsIfOneContextObject_MultiClient(FMenuBuilder& MenuBuilder, TConstArrayView<FSoftObjectPath> ContextObjects, FOnlineClientManager& InClientManager)
 	{
 		if (ContextObjects.Num() == 1)
 		{
diff --git a/Engine/Plugins/Developer/Concert/ConcertApp/MultiUserClient/Source/MultiUserClient/Private/Widgets/ActiveSession/Replication/Client/Multi/AllClientsSelectionModel.cpp b/Engine/Plugins/Developer/Concert/ConcertApp/MultiUserClient/Source/MultiUserClient/Private/Widgets/ActiveSession/Replication/Client/Multi/AllClientsSelectionModel.cpp
--- a/Engine/Plugins/Developer/Concert/ConcertApp/MultiUserClient/Source/MultiUserClient/Private/Widgets/ActiveSession/Replication/Client/Multi/AllClientsSelectionModel.cpp
+++ b/Engine/Plugins/Developer/Concert/ConcertApp/MultiUserClient/Source/MultiUserClient/Private/Widgets/ActiveSession/Replication/Client/Multi/AllClientsSelectionModel.cpp
@@ -2,11 +2,11 @@
 
 #include "AllClientsSelectionModel.h"
 
-#include "Replication/Client/ReplicationClientManager.h"
+#include "Replication/Client/Online/OnlineClientManager.h"
 
 namespace UE::MultiUserClient::Replication
 {
-	FAllClientsSelectionModel::FAllClientsSelectionModel(FReplicationClientManager& InClientManager)
+	FAllClientsSelectionModel::FAllClientsSelectionModel(FOnlineClientManager& InClientManager)
 		: ClientManager(InClientManager)
 	{
 		ClientManager.OnRemoteClientsChanged().AddRaw(this, &FAllClientsSelectionModel::OnRemoteClientsChanged);
@@ -17,7 +17,7 @@ namespace UE::MultiUserClient::Replication
 		ClientManager.OnRemoteClientsChanged().RemoveAll(this);
 	}
 
-	void FAllClientsSelectionModel::ForEachSelectedClient(TFunctionRef<EBreakBehavior(FReplicationClient&)> ProcessClient) const
+	void FAllClientsSelectionModel::ForEachSelectedClient(TFunctionRef<EBreakBehavior(FOnlineClient&)> ProcessClient) const
 	{
 		ClientManager.ForEachClient(ProcessClient);
 	}
diff --git a/Engine/Plugins/Developer/Concert/ConcertApp/MultiUserClient/Source/MultiUserClient/Private/Widgets/ActiveSession/Replication/Client/Multi/AllClientsSelectionModel.h b/Engine/Plugins/Developer/Concert/ConcertApp/MultiUserClient/Source/MultiUserClient/Private/Widgets/ActiveSession/Replication/Client/Multi/AllClientsSelectionModel.h
--- a/Engine/Plugins/Developer/Concert/ConcertApp/MultiUserClient/Source/MultiUserClient/Private/Widgets/ActiveSession/Replication/Client/Multi/AllClientsSelectionModel.h
+++ b/Engine/Plugins/Developer/Concert/ConcertApp/MultiUserClient/Source/MultiUserClient/Private/Widgets/ActiveSession/Replication/Client/Multi/AllClientsSelectionModel.h
@@ -6,25 +6,25 @@
 
 namespace UE::MultiUserClient::Replication
 {
-	class FReplicationClientManager;
+	class FOnlineClientManager;
 	
 	/** Exposes all clients and detects when clients disconnect. */
 	class FAllClientsSelectionModel : public IClientSelectionModel
 	{
 	public:
 		
-		FAllClientsSelectionModel(FReplicationClientManager& InClientManager);
+		FAllClientsSelectionModel(FOnlineClientManager& InClientManager);
 		virtual ~FAllClientsSelectionModel() override;
 
 		//~ Begin IClientSelectionModel Interface
-		virtual void ForEachSelectedClient(TFunctionRef<EBreakBehavior(FReplicationClient&)> ProcessClient) const override;
+		virtual void ForEachSelectedClient(TFunctionRef<EBreakBehavior(FOnlineClient&)> ProcessClient) const override;
 		virtual FOnSelectionChanged& OnSelectionChanged() override { return OnSelectionChangedDelegate; }
 		//~ End IClientSelectionModel Interface
 
 	private:
 
 		/** Informs us when the list of clients changes */
-		FReplicationClientManager& ClientManager;
+		FOnlineClientManager& ClientManager;
 		
 		/** Called when the clients ForEachSelectedClient enumerates has changed. */
 		FOnSelectionChanged OnSelectionChangedDelegate;
diff --git a/Engine/Plugins/Developer/Concert/ConcertApp/MultiUserClient/Source/MultiUserClient/Private/Widgets/ActiveSession/Replication/Client/Multi/Columns/MultiStreamColumns.h b/Engine/Plugins/Developer/Concert/ConcertApp/MultiUserClient/Source/MultiUserClient/Private/Widgets/ActiveSession/Replication/Client/Multi/Columns/MultiStreamColumns.h
--- a/Engine/Plugins/Developer/Concert/ConcertApp/MultiUserClient/Source/MultiUserClient/Private/Widgets/ActiveSession/Replication/Client/Multi/Columns/MultiStreamColumns.h
+++ b/Engine/Plugins/Developer/Concert/ConcertApp/MultiUserClient/Source/MultiUserClient/Private/Widgets/ActiveSession/Replication/Client/Multi/Columns/MultiStreamColumns.h
@@ -15,7 +15,7 @@ namespace UE::MultiUserClient::Replication
 {
 	class FMuteChangeTracker;
 	class FReassignObjectPropertiesLogic;
-	class FReplicationClientManager;
+	class FOnlineClientManager;
 }
 
 namespace UE::ConcertSharedSlate
@@ -66,7 +66,7 @@ namespace UE::MultiUserClient::Replication::MultiStreamColumns
 		TAttribute<TSharedPtr<ConcertSharedSlate::IMultiReplicationStreamEditor>> MultiStreamModelAttribute,
 		const ConcertSharedSlate::IObjectHierarchyModel& ObjectHierarchy UE_LIFETIMEBOUND,
 		FReassignObjectPropertiesLogic& ReassignmentLogic UE_LIFETIMEBOUND,
-		const FReplicationClientManager& ClientManager UE_LIFETIMEBOUND,
+		const FOnlineClientManager& ClientManager UE_LIFETIMEBOUND,
 		int32 ColumnsSortPriority = static_cast<int32>(EColumnSortOrder::ReassignOwnership)
 		);
 	
@@ -83,7 +83,7 @@ namespace UE::MultiUserClient::Replication::MultiStreamColumns
 	ConcertSharedSlate::FPropertyColumnEntry AssignPropertyColumn(
 		TAttribute<TSharedPtr<ConcertSharedSlate::IMultiReplicationStreamEditor>> MultiStreamEditor,
 		TSharedRef<IConcertClient> ConcertClient,
-		FReplicationClientManager& ClientManager UE_LIFETIMEBOUND,
+		FOnlineClientManager& ClientManager UE_LIFETIMEBOUND,
 		const int32 ColumnsSortPriority = static_cast<int32>(EColumnSortOrder::AssignPropertyColumn)
 		);
 }
diff --git a/Engine/Plugins/Developer/Concert/ConcertApp/MultiUserClient/Source/MultiUserClient/Private/Widgets/ActiveSession/Replication/Client/Multi/Columns/MultiStreamColumns_AssignProperties.cpp b/Engine/Plugins/Developer/Concert/ConcertApp/MultiUserClient/Source/MultiUserClient/Private/Widgets/ActiveSession/Replication/Client/Multi/Columns/MultiStreamColumns_AssignProperties.cpp
--- a/Engine/Plugins/Developer/Concert/ConcertApp/MultiUserClient/Source/MultiUserClient/Private/Widgets/ActiveSession/Replication/Client/Multi/Columns/MultiStreamColumns_AssignProperties.cpp
+++ b/Engine/Plugins/Developer/Concert/ConcertApp/MultiUserClient/Source/MultiUserClient/Private/Widgets/ActiveSession/Replication/Client/Multi/Columns/MultiStreamColumns_AssignProperties.cpp
@@ -4,8 +4,8 @@
 
 #include "IConcertClient.h"
 #include "MultiUserReplicationStyle.h"
-#include "Replication/Client/ReplicationClient.h"
-#include "Replication/Client/ReplicationClientManager.h"
+#include "Replication/Client/Online/OnlineClient.h"
+#include "Replication/Client/Online/OnlineClientManager.h"
 #include "Replication/Editor/Model/IEditableMultiReplicationStreamModel.h"
 #include "Replication/Editor/Model/IEditableReplicationStreamModel.h"
 #include "Replication/Editor/View/IMultiReplicationStreamEditor.h"
@@ -46,14 +46,14 @@ namespace UE::MultiUserClient::Replication::MultiStreamColumns
 			}
 		}
 
-		const FReplicationClient* FindClientByStream(const FReplicationClientManager& ClientManager, const ConcertSharedSlate::IReplicationStreamModel& StreamModel)
+		const FOnlineClient* FindClientByStream(const FOnlineClientManager& ClientManager, const ConcertSharedSlate::IReplicationStreamModel& StreamModel)
 		{
 			if (&ClientManager.GetLocalClient().GetClientEditModel().Get() == &StreamModel)
 			{
 				return &ClientManager.GetLocalClient();
 			}
 
-			for (const TNonNullPtr<const FRemoteReplicationClient> Client : ClientManager.GetRemoteClients())
+			for (const TNonNullPtr<const FRemoteClient> Client : ClientManager.GetRemoteClients())
 			{
 				if (&Client->GetClientEditModel().Get() == &StreamModel)
 				{
@@ -64,9 +64,9 @@ namespace UE::MultiUserClient::Replication::MultiStreamColumns
 			return nullptr;
 		}
 		
-		static FString GetClientDisplayText(const IConcertClient& InConcertClient, const FReplicationClientManager& ClientManager, const ConcertSharedSlate::IReplicationStreamModel& StreamModel)
+		static FString GetClientDisplayText(const IConcertClient& InConcertClient, const FOnlineClientManager& ClientManager, const ConcertSharedSlate::IReplicationStreamModel& StreamModel)
 		{
-			if (const FReplicationClient* Client = FindClientByStream(ClientManager, StreamModel))
+			if (const FOnlineClient* Client = FindClientByStream(ClientManager, StreamModel))
 			{
 				return ClientUtils::GetClientDisplayName(InConcertClient, Client->GetEndpointId());
 			}
@@ -79,7 +79,7 @@ namespace UE::MultiUserClient::Replication::MultiStreamColumns
 	ConcertSharedSlate::FPropertyColumnEntry AssignPropertyColumn(
 		TAttribute<TSharedPtr<ConcertSharedSlate::IMultiReplicationStreamEditor>> MultiStreamEditor,
 		TSharedRef<IConcertClient> ConcertClient,
-		FReplicationClientManager& ClientManager,
+		FOnlineClientManager& ClientManager,
 		const int32 ColumnsSortPriority
 		)
 	{
@@ -91,7 +91,7 @@ namespace UE::MultiUserClient::Replication::MultiStreamColumns
 			FPropertyColumn_AssignPropertyColumn(
 				TAttribute<TSharedPtr<IMultiReplicationStreamEditor>> MultiStreamEditor,
 				TSharedRef<IConcertClient> ConcertClient,
-				FReplicationClientManager& ClientManager
+				FOnlineClientManager& ClientManager
 				)
 				: MultiStreamEditor(MoveTemp(MultiStreamEditor))
 				, ConcertClient(MoveTemp(ConcertClient))
@@ -152,7 +152,7 @@ namespace UE::MultiUserClient::Replication::MultiStreamColumns
 
 			const TAttribute<TSharedPtr<IMultiReplicationStreamEditor>> MultiStreamEditor;
 			const TSharedRef<IConcertClient> ConcertClient;
-			FReplicationClientManager& ClientManager;
+			FOnlineClientManager& ClientManager;
 		};
 		
 		check(MultiStreamEditor.IsBound() || MultiStreamEditor.IsSet());
diff --git a/Engine/Plugins/Developer/Concert/ConcertApp/MultiUserClient/Source/MultiUserClient/Private/Widgets/ActiveSession/Replication/Client/Multi/Columns/MultiStreamColumns_AssignedClients.cpp b/Engine/Plugins/Developer/Concert/ConcertApp/MultiUserClient/Source/MultiUserClient/Private/Widgets/ActiveSession/Replication/Client/Multi/Columns/MultiStreamColumns_AssignedClients.cpp
--- a/Engine/Plugins/Developer/Concert/ConcertApp/MultiUserClient/Source/MultiUserClient/Private/Widgets/ActiveSession/Replication/Client/Multi/Columns/MultiStreamColumns_AssignedClients.cpp
+++ b/Engine/Plugins/Developer/Concert/ConcertApp/MultiUserClient/Source/MultiUserClient/Private/Widgets/ActiveSession/Replication/Client/Multi/Columns/MultiStreamColumns_AssignedClients.cpp
@@ -117,7 +117,7 @@ namespace UE::MultiUserClient::Replication::MultiStreamColumns
 		TAttribute<TSharedPtr<ConcertSharedSlate::IMultiReplicationStreamEditor>> MultiStreamModelAttribute,
 		const ConcertSharedSlate::IObjectHierarchyModel& ObjectHierarchy,
 		FReassignObjectPropertiesLogic& ReassignmentLogic,
-		const FReplicationClientManager& ClientManager,
+		const FOnlineClientManager& ClientManager,
 		const int32 ColumnsSortPriority
 		)
 	{
@@ -130,7 +130,7 @@ namespace UE::MultiUserClient::Replication::MultiStreamColumns
 				TAttribute<TSharedPtr<ConcertSharedSlate::IMultiReplicationStreamEditor>> MultiStreamModelAttribute,
 				const ConcertSharedSlate::IObjectHierarchyModel& ObjectHierarchy UE_LIFETIMEBOUND,
 				FReassignObjectPropertiesLogic& ReassignmentLogic UE_LIFETIMEBOUND,
-				const FReplicationClientManager& ClientManager UE_LIFETIMEBOUND
+				const FOnlineClientManager& ClientManager UE_LIFETIMEBOUND
 				)
 				: ConcertClient(MoveTemp(ConcertClient))
 				, MultiStreamModelAttribute(MoveTemp(MultiStreamModelAttribute))
@@ -181,7 +181,7 @@ namespace UE::MultiUserClient::Replication::MultiStreamColumns
 			const TAttribute<TSharedPtr<ConcertSharedSlate::IMultiReplicationStreamEditor>> MultiStreamModelAttribute;
 			const ConcertSharedSlate::IObjectHierarchyModel& ObjectHierarchy;
 			FReassignObjectPropertiesLogic& ReassignmentLogic;
-			const FReplicationClientManager& ClientManager;
+			const FOnlineClientManager& ClientManager;
 			
 			TOptional<FString> GetDisplayString(const FSoftObjectPath& ManagedObject) const
 			{
diff --git a/Engine/Plugins/Developer/Concert/ConcertApp/MultiUserClient/Source/MultiUserClient/Private/Widgets/ActiveSession/Replication/Client/Multi/Columns/SAssignPropertyComboBox.cpp b/Engine/Plugins/Developer/Concert/ConcertApp/MultiUserClient/Source/MultiUserClient/Private/Widgets/ActiveSession/Replication/Client/Multi/Columns/SAssignPropertyComboBox.cpp
--- a/Engine/Plugins/Developer/Concert/ConcertApp/MultiUserClient/Source/MultiUserClient/Private/Widgets/ActiveSession/Replication/Client/Multi/Columns/SAssignPropertyComboBox.cpp
+++ b/Engine/Plugins/Developer/Concert/ConcertApp/MultiUserClient/Source/MultiUserClient/Private/Widgets/ActiveSession/Replication/Client/Multi/Columns/SAssignPropertyComboBox.cpp
@@ -4,8 +4,8 @@
 
 #include "ConcertLogGlobal.h"
 #include "IConcertClient.h"
-#include "Replication/Client/ReplicationClient.h"
-#include "Replication/Client/ReplicationClientManager.h"
+#include "Replication/Client/Online/OnlineClient.h"
+#include "Replication/Client/Online/OnlineClientManager.h"
 #include "Replication/Editor/Model/IEditableReplicationStreamModel.h"
 #include "Replication/Editor/Model/PropertyUtils.h"
 #include "Widgets/ActiveSession/Replication/Misc/SNoClients.h"
@@ -30,10 +30,10 @@ namespace UE::MultiUserClient::Replication
 {
 	namespace AssignPropertyComboBox
 	{
-		TArray<FGuid> GetDisplayedClients(const FReplicationClientManager& ClientManager, const FConcertPropertyChain& DisplayedProperty, const TArray<TSoftObjectPtr<>>& EditedObjects)
+		TArray<FGuid> GetDisplayedClients(const FOnlineClientManager& ClientManager, const FConcertPropertyChain& DisplayedProperty, const TArray<TSoftObjectPtr<>>& EditedObjects)
 		{
 			TArray<FGuid> Clients;
-			ClientManager.ForEachClient([&DisplayedProperty, &EditedObjects, &Clients](const FReplicationClient& Client)
+			ClientManager.ForEachClient([&DisplayedProperty, &EditedObjects, &Clients](const FOnlineClient& Client)
 			{
 				const TMap<FSoftObjectPath, FConcertReplicatedObjectInfo>& ObjectInfoMap = Client.GetStreamSynchronizer().GetServerState().ReplicatedObjects;
 				for (const TSoftObjectPtr<>& ObjectPath : EditedObjects)
@@ -53,7 +53,7 @@ namespace UE::MultiUserClient::Replication
 	
 	TOptional<FString> SAssignPropertyComboBox::GetDisplayString(
 		const TSharedRef<IConcertClient>& LocalConcertClient,
-		const FReplicationClientManager& ClientManager,
+		const FOnlineClientManager& ClientManager,
 		const FConcertPropertyChain& DisplayedProperty,
 		const TArray<TSoftObjectPtr<>>& EditedObjects)
 	{
@@ -71,7 +71,7 @@ namespace UE::MultiUserClient::Replication
 	void SAssignPropertyComboBox::Construct(const FArguments& InArgs,
 	    TSharedRef<ConcertSharedSlate::IMultiReplicationStreamEditor> InEditor,
 	    TSharedRef<IConcertClient> InConcertClient,
-	    FReplicationClientManager& InClientManager
+	    FOnlineClientManager& InClientManager
 	)
 	{
 		Editor = MoveTemp(InEditor);
@@ -147,7 +147,7 @@ namespace UE::MultiUserClient::Replication
 		);
 		
 		MenuBuilder.BeginSection(TEXT("AssignTo"), LOCTEXT("AssignTo", "Assign to"));
-		for (const FReplicationClient* Client : ClientUtils::GetSortedClientList(*ConcertClient, *ClientManager))
+		for (const FOnlineClient* Client : ClientUtils::GetSortedClientList(*ConcertClient, *ClientManager))
 		{
 			TAttribute<FText> Tooltip = TAttribute<FText>::CreateLambda([this, EndpointId = Client->GetEndpointId()]()
 			{
@@ -187,7 +187,7 @@ namespace UE::MultiUserClient::Replication
 	void SAssignPropertyComboBox::OnClickOption(const FGuid EndpointId) const
 	{
 		// Remote clients can disconnect after the combo-box is opened.
-		const FReplicationClient* Client = ClientManager->FindClient(EndpointId);
+		const FOnlineClient* Client = ClientManager->FindClient(EndpointId);
 		if (!Client)
 		{
 			return;
@@ -203,12 +203,12 @@ namespace UE::MultiUserClient::Replication
 		if (bRemovePropertyFromEditedClient)
 		{
 			// ... remove property from all clients
-			UnassignPropertyFromClients([](const FReplicationClient& ClientToRemoveFrom){ return true; });
+			UnassignPropertyFromClients([](const FOnlineClient& ClientToRemoveFrom){ return true; });
 		}
 		else
 		{
 			// ... remove the property from all clients but the one we'll assign to ...
-			UnassignPropertyFromClients([Client](const FReplicationClient& ClientToRemoveFrom){ return *Client != ClientToRemoveFrom; });
+			UnassignPropertyFromClients([Client](const FOnlineClient& ClientToRemoveFrom){ return *Client != ClientToRemoveFrom; });
 
 			// ... and then assign the property
 			const TSharedRef<ConcertSharedSlate::IEditableReplicationStreamModel> EditModel = Client->GetClientEditModel();
@@ -233,7 +233,7 @@ namespace UE::MultiUserClient::Replication
 #define SET_REASON(Text) if (Reason) { *Reason = Text; }
 	bool SAssignPropertyComboBox::CanClickOptionWithReason(const FGuid& EndpointId, FText* Reason) const
 	{
-		const FReplicationClient* Client = ClientManager->FindClient(EndpointId);
+		const FOnlineClient* Client = ClientManager->FindClient(EndpointId);
 		// Remote clients can disconnect after the combo-box is opened.
 		if (!Client)
 		{
@@ -243,7 +243,7 @@ namespace UE::MultiUserClient::Replication
 
 		// The combo box assigns the property to the clicked client and removes from the others... check that the currently assigned clients allow it.
 		bool bCanRemoveFromOwners = true;
-		ClientManager->ForEachClient([this, &EndpointId, &Reason, Client, &bCanRemoveFromOwners](const FReplicationClient& ClientToRemoveFrom)
+		ClientManager->ForEachClient([this, &EndpointId, &Reason, Client, &bCanRemoveFromOwners](const FOnlineClient& ClientToRemoveFrom)
 		{
 			if (*Client != ClientToRemoveFrom && !ClientToRemoveFrom.AllowsEditing())
 			{
@@ -281,7 +281,7 @@ namespace UE::MultiUserClient::Replication
 	
 	ECheckBoxState SAssignPropertyComboBox::GetOptionCheckState(const FGuid EndpointId) const
 	{
-		const FReplicationClient* Client = ClientManager->FindClient(EndpointId);
+		const FOnlineClient* Client = ClientManager->FindClient(EndpointId);
 		// Remote clients can disconnect after the combo-box is opened.
 		if (!Client)
 		{
@@ -322,14 +322,14 @@ namespace UE::MultiUserClient::Replication
 		const FText TransactionText = FText::Format(LOCTEXT("ClearAllClientsFmt", "Clear {0} property"), FText::FromString(Property.ToString(FConcertPropertyChain::EToStringMethod::LeafProperty)));
 		FScopedTransaction Transaction(TransactionText);
 
-		UnassignPropertyFromClients([](const FReplicationClient& ClientToRemoveFrom){ return true; });
+		UnassignPropertyFromClients([](const FOnlineClient& ClientToRemoveFrom){ return true; });
 		OnOptionClickedDelegate.ExecuteIfBound();
 	}
 
 	bool SAssignPropertyComboBox::CanClickClear() const
 	{
 		bool bIsAssignedToAnyClient = false;
-		ClientManager->ForEachClient([this, &bIsAssignedToAnyClient](const FReplicationClient& Client)
+		ClientManager->ForEachClient([this, &bIsAssignedToAnyClient](const FOnlineClient& Client)
 		{
 			for (const TSoftObjectPtr<>& EditedObject : EditedObjects)
 			{
@@ -348,9 +348,9 @@ namespace UE::MultiUserClient::Replication
 		return bIsAssignedToAnyClient;
 	}
 
-	void SAssignPropertyComboBox::UnassignPropertyFromClients(TFunctionRef<bool(const FReplicationClient& Client)> ShouldRemoveFromClient) const
+	void SAssignPropertyComboBox::UnassignPropertyFromClients(TFunctionRef<bool(const FOnlineClient& Client)> ShouldRemoveFromClient) const
 	{
-		ClientManager->ForEachClient([this, &ShouldRemoveFromClient](const FReplicationClient& ClientToRemoveFrom)
+		ClientManager->ForEachClient([this, &ShouldRemoveFromClient](const FOnlineClient& ClientToRemoveFrom)
 		{
 			const TSharedRef<ConcertSharedSlate::IEditableReplicationStreamModel> EditModel = ClientToRemoveFrom.GetClientEditModel();
 			if (ClientToRemoveFrom.AllowsEditing() && ShouldRemoveFromClient(ClientToRemoveFrom))
@@ -386,7 +386,7 @@ namespace UE::MultiUserClient::Replication
 
 	void SAssignPropertyComboBox::RebuildSubscriptions()
 	{
-		ClientManager->ForEachClient([this](FReplicationClient& Client)
+		ClientManager->ForEachClient([this](FOnlineClient& Client)
 		{
 			Client.OnModelChanged().RemoveAll(this);
 			Client.OnModelChanged().AddSP(this, &SAssignPropertyComboBox::RefreshContentBoxContent);
diff --git a/Engine/Plugins/Developer/Concert/ConcertApp/MultiUserClient/Source/MultiUserClient/Private/Widgets/ActiveSession/Replication/Client/Multi/Columns/SAssignPropertyComboBox.h b/Engine/Plugins/Developer/Concert/ConcertApp/MultiUserClient/Source/MultiUserClient/Private/Widgets/ActiveSession/Replication/Client/Multi/Columns/SAssignPropertyComboBox.h
--- a/Engine/Plugins/Developer/Concert/ConcertApp/MultiUserClient/Source/MultiUserClient/Private/Widgets/ActiveSession/Replication/Client/Multi/Columns/SAssignPropertyComboBox.h
+++ b/Engine/Plugins/Developer/Concert/ConcertApp/MultiUserClient/Source/MultiUserClient/Private/Widgets/ActiveSession/Replication/Client/Multi/Columns/SAssignPropertyComboBox.h
@@ -18,8 +18,8 @@ namespace UE::ConcertSharedSlate { class IMultiReplicationStreamEditor; }
 
 namespace UE::MultiUserClient::Replication
 {
-	class FReplicationClient;
-	class FReplicationClientManager;
+	class FOnlineClient;
+	class FOnlineClientManager;
 	
 	/**
 	 * Placed in every property row column to assign properties to clients.
@@ -36,7 +36,7 @@ namespace UE::MultiUserClient::Replication
 		/** @return The display string this widget would have with the given state. If unset, no clients are displayed in the combobox.*/
 		static TOptional<FString> GetDisplayString(
 			const TSharedRef<IConcertClient>& LocalConcertClient,
-			const FReplicationClientManager& ClientManager,
+			const FOnlineClientManager& ClientManager,
 			const FConcertPropertyChain& DisplayedProperty,
 			const TArray<TSoftObjectPtr<>>& EditedObjects
 			);
@@ -53,12 +53,12 @@ namespace UE::MultiUserClient::Replication
 		void Construct(const FArguments& InArgs,
            TSharedRef<ConcertSharedSlate::IMultiReplicationStreamEditor> InEditor,
            TSharedRef<IConcertClient> InConcertClient,
-           FReplicationClientManager& InClientManager
+           FOnlineClientManager& InClientManager
 		);
 
 	private:
 
-		FReplicationClientManager* ClientManager = nullptr;
+		FOnlineClientManager* ClientManager = nullptr;
 		
 		/** Used to obtain info about the streams */
 		TSharedPtr<ConcertSharedSlate::IMultiReplicationStreamEditor> Editor;
@@ -97,7 +97,7 @@ namespace UE::MultiUserClient::Replication
 		bool CanClickClear() const;
 
 		/** Unassigns this widget's property from all clients passing the predicate and removes the object from the model if it is a subobject. */
-		void UnassignPropertyFromClients(TFunctionRef<bool(const FReplicationClient& Client)> ShouldRemoveFromClient) const;
+		void UnassignPropertyFromClients(TFunctionRef<bool(const FOnlineClient& Client)> ShouldRemoveFromClient) const;
 
 		/** Resubscribes to all clients changing. */
 		void RebuildSubscriptions();
diff --git a/Engine/Plugins/Developer/Concert/ConcertApp/MultiUserClient/Source/MultiUserClient/Private/Widgets/ActiveSession/Replication/Client/Multi/IClientSelectionModel.h b/Engine/Plugins/Developer/Concert/ConcertApp/MultiUserClient/Source/MultiUserClient/Private/Widgets/ActiveSession/Replication/Client/Multi/IClientSelectionModel.h
--- a/Engine/Plugins/Developer/Concert/ConcertApp/MultiUserClient/Source/MultiUserClient/Private/Widgets/ActiveSession/Replication/Client/Multi/IClientSelectionModel.h
+++ b/Engine/Plugins/Developer/Concert/ConcertApp/MultiUserClient/Source/MultiUserClient/Private/Widgets/ActiveSession/Replication/Client/Multi/IClientSelectionModel.h
@@ -9,7 +9,7 @@ enum class EBreakBehavior : uint8;
 
 namespace UE::MultiUserClient::Replication
 {
-	class FReplicationClient;
+	class FOnlineClient;
 	
 	/** Abstracts the concept of a client selection. */
 	class IClientSelectionModel
@@ -17,7 +17,7 @@ namespace UE::MultiUserClient::Replication
 	public:
 
 		/** Iterates through every client in the selection. */
-		virtual void ForEachSelectedClient(TFunctionRef<EBreakBehavior(FReplicationClient&)> ProcessClient) const = 0;
+		virtual void ForEachSelectedClient(TFunctionRef<EBreakBehavior(FOnlineClient&)> ProcessClient) const = 0;
 
 		DECLARE_MULTICAST_DELEGATE(FOnSelectionChanged);
 		/** Called when the clients ForEachSelectedClient enumerates has changed. */
diff --git a/Engine/Plugins/Developer/Concert/ConcertApp/MultiUserClient/Source/MultiUserClient/Private/Widgets/ActiveSession/Replication/Client/Multi/MultiStreamModel.cpp b/Engine/Plugins/Developer/Concert/ConcertApp/MultiUserClient/Source/MultiUserClient/Private/Widgets/ActiveSession/Replication/Client/Multi/MultiStreamModel.cpp
--- a/Engine/Plugins/Developer/Concert/ConcertApp/MultiUserClient/Source/MultiUserClient/Private/Widgets/ActiveSession/Replication/Client/Multi/MultiStreamModel.cpp
+++ b/Engine/Plugins/Developer/Concert/ConcertApp/MultiUserClient/Source/MultiUserClient/Private/Widgets/ActiveSession/Replication/Client/Multi/MultiStreamModel.cpp
@@ -3,13 +3,13 @@
 #include "MultiStreamModel.h"
 
 #include "IClientSelectionModel.h"
-#include "Replication/Client/ReplicationClient.h"
-#include "Replication/Client/ReplicationClientManager.h"
+#include "Replication/Client/Online/OnlineClient.h"
+#include "Replication/Client/Online/OnlineClientManager.h"
 #include "Replication/Editor/Model/IEditableReplicationStreamModel.h"
 
 namespace UE::MultiUserClient::Replication
 {
-	FMultiStreamModel::FMultiStreamModel(IClientSelectionModel& InClientSelectionModel, FReplicationClientManager& InClientManager)
+	FMultiStreamModel::FMultiStreamModel(IClientSelectionModel& InClientSelectionModel, FOnlineClientManager& InClientManager)
 		: ClientSelectionModel(InClientSelectionModel)
 		, ClientManager(InClientManager)
 	{
@@ -17,16 +17,16 @@ namespace UE::MultiUserClient::Replication
 		RebuildStreamsSets();
 	}
 
-	void FMultiStreamModel::ForEachClient(TFunctionRef<EBreakBehavior(const FReplicationClient*)> ProcessClient) const
+	void FMultiStreamModel::ForEachClient(TFunctionRef<EBreakBehavior(const FOnlineClient*)> ProcessClient) const
 	{
-		for (const FReplicationClient* ReadOnlyClient : GetCachedReadOnlyClients())
+		for (const FOnlineClient* ReadOnlyClient : GetCachedReadOnlyClients())
 		{
 			if (ProcessClient(ReadOnlyClient) == EBreakBehavior::Break)
 			{
 				return;
 			}
 		}
-		for (const FReplicationClient* WritableClient : GetCachedWritableClients())
+		for (const FOnlineClient* WritableClient : GetCachedWritableClients())
 		{
 			if (ProcessClient(WritableClient) == EBreakBehavior::Break)
 			{
@@ -38,29 +38,29 @@ namespace UE::MultiUserClient::Replication
 	TSet<TSharedRef<ConcertSharedSlate::IReplicationStreamModel>> FMultiStreamModel::GetReadOnlyStreams() const
 	{
 		TSet<TSharedRef<ConcertSharedSlate::IReplicationStreamModel>> Result;
-		Algo::Transform(CachedReadOnlyClients, Result, [](const FReplicationClient* Client){ return Client->GetClientEditModel(); });
+		Algo::Transform(CachedReadOnlyClients, Result, [](const FOnlineClient* Client){ return Client->GetClientEditModel(); });
 		return Result;
 	}
 
 	TSet<TSharedRef<ConcertSharedSlate::IEditableReplicationStreamModel>> FMultiStreamModel::GetEditableStreams() const
 	{
 		TSet<TSharedRef<ConcertSharedSlate::IEditableReplicationStreamModel>> Result;
-		Algo::Transform(CachedWritableClients, Result, [](const FReplicationClient* Client){ return Client->GetClientEditModel(); });
+		Algo::Transform(CachedWritableClients, Result, [](const FOnlineClient* Client){ return Client->GetClientEditModel(); });
 		return Result;
 	}
 
 	void FMultiStreamModel::RebuildStreamsSets()
 	{
 		// Cannot just iterate through CachedReadOnlyClients because it may contain stale clients that were just removed
-		ClientManager.ForEachClient([this](FReplicationClient& Client)
+		ClientManager.ForEachClient([this](FOnlineClient& Client)
 		{
 			Client.OnModelChanged().RemoveAll(this);
 			return EBreakBehavior::Continue;
 		});
 		
-		TSet<const FReplicationClient*> ReadOnlyClients;
-		TSet<const FReplicationClient*> WritableClients;
-		ClientSelectionModel.ForEachSelectedClient([this, &ReadOnlyClients, &WritableClients](FReplicationClient& Client)
+		TSet<const FOnlineClient*> ReadOnlyClients;
+		TSet<const FOnlineClient*> WritableClients;
+		ClientSelectionModel.ForEachSelectedClient([this, &ReadOnlyClients, &WritableClients](FOnlineClient& Client)
 		{
 			const bool bIsUploadable = CanEverSubmit(Client.GetSubmissionWorkflow().GetUploadability());
 			const TSharedRef<ConcertSharedSlate::IEditableReplicationStreamModel> Stream = Client.GetClientEditModel();
diff --git a/Engine/Plugins/Developer/Concert/ConcertApp/MultiUserClient/Source/MultiUserClient/Private/Widgets/ActiveSession/Replication/Client/Multi/MultiStreamModel.h b/Engine/Plugins/Developer/Concert/ConcertApp/MultiUserClient/Source/MultiUserClient/Private/Widgets/ActiveSession/Replication/Client/Multi/MultiStreamModel.h
--- a/Engine/Plugins/Developer/Concert/ConcertApp/MultiUserClient/Source/MultiUserClient/Private/Widgets/ActiveSession/Replication/Client/Multi/MultiStreamModel.h
+++ b/Engine/Plugins/Developer/Concert/ConcertApp/MultiUserClient/Source/MultiUserClient/Private/Widgets/ActiveSession/Replication/Client/Multi/MultiStreamModel.h
@@ -12,8 +12,8 @@ enum class EBreakBehavior : uint8;
 
 namespace UE::MultiUserClient::Replication
 {
-	class FReplicationClient;
-	class FReplicationClientManager;
+	class FOnlineClient;
+	class FOnlineClientManager;
 	class IClientSelectionModel;
 
 	/** Checks whether clients accept remote changes and categorizes them into read-only and writable. */
@@ -21,11 +21,11 @@ namespace UE::MultiUserClient::Replication
 	{
 	public:
 		
-		FMultiStreamModel(IClientSelectionModel& InClientSelectionModel, FReplicationClientManager& InClientManager);
+		FMultiStreamModel(IClientSelectionModel& InClientSelectionModel, FOnlineClientManager& InClientManager);
 		
-		const TSet<const FReplicationClient*>& GetCachedReadOnlyClients() const { return CachedReadOnlyClients; }
-		const TSet<const FReplicationClient*>& GetCachedWritableClients() const { return CachedWritableClients; }
-		void ForEachClient(TFunctionRef<EBreakBehavior(const FReplicationClient*)> ProcessClient) const;
+		const TSet<const FOnlineClient*>& GetCachedReadOnlyClients() const { return CachedReadOnlyClients; }
+		const TSet<const FOnlineClient*>& GetCachedWritableClients() const { return CachedWritableClients; }
+		void ForEachClient(TFunctionRef<EBreakBehavior(const FOnlineClient*)> ProcessClient) const;
 
 		//~ Begin IEditableMultiReplicationStreamModel Interface
 		virtual TSet<TSharedRef<ConcertSharedSlate::IReplicationStreamModel>> GetReadOnlyStreams() const override;
@@ -39,10 +39,10 @@ namespace UE::MultiUserClient::Replication
 		/** Gets all clients to display and informs when they change. */
 		IClientSelectionModel& ClientSelectionModel;
 		/** Used to obtain a list of clients for unsubscribing */
-		FReplicationClientManager& ClientManager;
+		FOnlineClientManager& ClientManager;
 
-		TSet<const FReplicationClient*> CachedReadOnlyClients;
-		TSet<const FReplicationClient*> CachedWritableClients;
+		TSet<const FOnlineClient*> CachedReadOnlyClients;
+		TSet<const FOnlineClient*> CachedWritableClients;
 
 		FOnStreamExternallyChanged OnReadOnlyStreamChangedDelegate;
 		FOnStreamSetChanged OnStreamSetChangedDelegate;
diff --git a/Engine/Plugins/Developer/Concert/ConcertApp/MultiUserClient/Source/MultiUserClient/Private/Widgets/ActiveSession/Replication/Client/Multi/SAllClientsView.cpp b/Engine/Plugins/Developer/Concert/ConcertApp/MultiUserClient/Source/MultiUserClient/Private/Widgets/ActiveSession/Replication/Client/Multi/SAllClientsView.cpp
--- a/Engine/Plugins/Developer/Concert/ConcertApp/MultiUserClient/Source/MultiUserClient/Private/Widgets/ActiveSession/Replication/Client/Multi/SAllClientsView.cpp
+++ b/Engine/Plugins/Developer/Concert/ConcertApp/MultiUserClient/Source/MultiUserClient/Private/Widgets/ActiveSession/Replication/Client/Multi/SAllClientsView.cpp
@@ -4,7 +4,7 @@
 
 #include "AllClientsSelectionModel.h"
 #include "SMultiClientView.h"
-#include "Replication/Client/ReplicationClientManager.h"
+#include "Replication/Client/Online/OnlineClientManager.h"
 
 #include "Algo/Transform.h"
 #include "Replication/MultiUserReplicationManager.h"
@@ -23,10 +23,10 @@ namespace UE::MultiUserClient::Replication
 		];
 	}
 
-	TSet<const FReplicationClient*> SAllClientsView::GetAllClients() const
+	TSet<const FOnlineClient*> SAllClientsView::GetAllClients() const
 	{
-		TSet<const FReplicationClient*> Result;
-		Algo::Transform(ClientManager->GetRemoteClients(), Result, [](const TNonNullPtr<FRemoteReplicationClient>& Client)
+		TSet<const FOnlineClient*> Result;
+		Algo::Transform(ClientManager->GetRemoteClients(), Result, [](const TNonNullPtr<FRemoteClient>& Client)
 		{
 			return Client.Get();
 		});
diff --git a/Engine/Plugins/Developer/Concert/ConcertApp/MultiUserClient/Source/MultiUserClient/Private/Widgets/ActiveSession/Replication/Client/Multi/SAllClientsView.h b/Engine/Plugins/Developer/Concert/ConcertApp/MultiUserClient/Source/MultiUserClient/Private/Widgets/ActiveSession/Replication/Client/Multi/SAllClientsView.h
--- a/Engine/Plugins/Developer/Concert/ConcertApp/MultiUserClient/Source/MultiUserClient/Private/Widgets/ActiveSession/Replication/Client/Multi/SAllClientsView.h
+++ b/Engine/Plugins/Developer/Concert/ConcertApp/MultiUserClient/Source/MultiUserClient/Private/Widgets/ActiveSession/Replication/Client/Multi/SAllClientsView.h
@@ -13,8 +13,8 @@ class IConcertClient;
 namespace UE::MultiUserClient::Replication
 {
 	class FMultiUserReplicationManager;
-	class FReplicationClient;
-	class FReplicationClientManager;
+	class FOnlineClient;
+	class FOnlineClientManager;
 
 	/** Leverages SMultiClientView to display all replication clients. */
 	class SAllClientsView : public SCompoundWidget
@@ -27,12 +27,12 @@ namespace UE::MultiUserClient::Replication
 	private:
 
 		/** Used to get all the replication clients and listen for client changes. */
-		FReplicationClientManager* ClientManager = nullptr;
+		FOnlineClientManager* ClientManager = nullptr;
 
 		/** Keeps the SMultiClientView updated of any changes to clients (e.g. disconnects, etc.) */
 		TUniquePtr<FAllClientsSelectionModel> AllClientsModel;
 
 		/** Gets all the clients to display */
-		TSet<const FReplicationClient*> GetAllClients() const;
+		TSet<const FOnlineClient*> GetAllClients() const;
 	};
 }
diff --git a/Engine/Plugins/Developer/Concert/ConcertApp/MultiUserClient/Source/MultiUserClient/Private/Widgets/ActiveSession/Replication/Client/Multi/SMultiClientView.cpp b/Engine/Plugins/Developer/Concert/ConcertApp/MultiUserClient/Source/MultiUserClient/Private/Widgets/ActiveSession/Replication/Client/Multi/SMultiClientView.cpp
--- a/Engine/Plugins/Developer/Concert/ConcertApp/MultiUserClient/Source/MultiUserClient/Private/Widgets/ActiveSession/Replication/Client/Multi/SMultiClientView.cpp
+++ b/Engine/Plugins/Developer/Concert/ConcertApp/MultiUserClient/Source/MultiUserClient/Private/Widgets/ActiveSession/Replication/Client/Multi/SMultiClientView.cpp
@@ -7,8 +7,8 @@
 #include "Replication/ClientReplicationWidgetFactories.h"
 #include "Replication/MultiUserReplicationManager.h"
 #include "Replication/ReplicationWidgetFactories.h"
-#include "Replication/Client/ReplicationClient.h"
-#include "Replication/Client/ReplicationClientManager.h"
+#include "Replication/Client/Online/OnlineClient.h"
+#include "Replication/Client/Online/OnlineClientManager.h"
 #include "Replication/Editor/Model/Object/IObjectNameModel.h"
 #include "Replication/Editor/Model/PropertySource/SelectPropertyFromUClassModel.h"
 #include "Replication/Editor/View/IMultiObjectPropertyAssignmentView.h"
@@ -171,7 +171,7 @@ namespace UE::MultiUserClient::Replication
 	TSet<FGuid> SMultiClientView::GetReplicatableClientIds() const
 	{
 		TSet<FGuid> ClientIds;
-		StreamModel->ForEachClient([&ClientIds](const FReplicationClient* Client)
+		StreamModel->ForEachClient([&ClientIds](const FOnlineClient* Client)
 		{
 			ClientIds.Add(Client->GetEndpointId());
 			return EBreakBehavior::Continue;
@@ -181,7 +181,7 @@ namespace UE::MultiUserClient::Replication
 
 	void SMultiClientView::EnumerateObjectsInStreams(TFunctionRef<void(const FSoftObjectPath&)> Consumer) const
 	{
-		StreamModel->ForEachClient([&Consumer](const FReplicationClient* Client)
+		StreamModel->ForEachClient([&Consumer](const FOnlineClient* Client)
 		{
 			Client->GetClientEditModel()->ForEachReplicatedObject([&Consumer](const FSoftObjectPath& Object)
 			{
@@ -196,7 +196,7 @@ namespace UE::MultiUserClient::Replication
 	{
 		CleanClientSubscriptions();
 
-		SelectionModel->ForEachSelectedClient([this](FReplicationClient& Client)
+		SelectionModel->ForEachSelectedClient([this](FOnlineClient& Client)
 		{
 			Client.OnModelChanged().AddSP(this, &SMultiClientView::RefreshUI);
 			Client.OnHierarchyNeedsRefresh().AddRaw(this, &SMultiClientView::RefreshUI);
@@ -206,7 +206,7 @@ namespace UE::MultiUserClient::Replication
 
 	void SMultiClientView::CleanClientSubscriptions() const
 	{
-		ClientManager->ForEachClient([this](FReplicationClient& Client)
+		ClientManager->ForEachClient([this](FOnlineClient& Client)
 		{
 			Client.OnModelChanged().RemoveAll(this);
 			Client.OnHierarchyNeedsRefresh().RemoveAll(this);
diff --git a/Engine/Plugins/Developer/Concert/ConcertApp/MultiUserClient/Source/MultiUserClient/Private/Widgets/ActiveSession/Replication/Client/Multi/SMultiClientView.h b/Engine/Plugins/Developer/Concert/ConcertApp/MultiUserClient/Source/MultiUserClient/Private/Widgets/ActiveSession/Replication/Client/Multi/SMultiClientView.h
--- a/Engine/Plugins/Developer/Concert/ConcertApp/MultiUserClient/Source/MultiUserClient/Private/Widgets/ActiveSession/Replication/Client/Multi/SMultiClientView.h
+++ b/Engine/Plugins/Developer/Concert/ConcertApp/MultiUserClient/Source/MultiUserClient/Private/Widgets/ActiveSession/Replication/Client/Multi/SMultiClientView.h
@@ -29,8 +29,8 @@ namespace UE::MultiUserClient::Replication
 	class FGlobalAuthorityCache;
 	class FMultiStreamModel;
 	class FMultiUserReplicationManager;
-	class FReplicationClient;
-	class FReplicationClientManager;
+	class FOnlineClient;
+	class FOnlineClientManager;
 	class IClientSelectionModel;
 	class SPropertySelectionComboButton;
 
@@ -52,7 +52,7 @@ namespace UE::MultiUserClient::Replication
 	private:
 
 		TSharedPtr<IConcertClient> ConcertClient;
-		FReplicationClientManager* ClientManager = nullptr;
+		FOnlineClientManager* ClientManager = nullptr;
 		FUserPropertySelector* UserSelectedProperties = nullptr;
 		IClientSelectionModel* SelectionModel = nullptr;
 		
