Rename IClientSelectionModel to ISelectionModel and templatize it in preparation for introducing offline clients.

diff --git a/Engine/Plugins/Developer/Concert/ConcertApp/MultiUserClient/Source/MultiUserClient/Private/Widgets/ActiveSession/Replication/Client/Multi/AllClientsSelectionModel.cpp b/Engine/Plugins/Developer/Concert/ConcertApp/MultiUserClient/Source/MultiUserClient/Private/Widgets/ActiveSession/Replication/Client/Multi/AllClientsSelectionModel.cpp
deleted file mode 100644
diff --git a/Engine/Plugins/Developer/Concert/ConcertApp/MultiUserClient/Source/MultiUserClient/Private/Widgets/ActiveSession/Replication/Client/Multi/AllClientsSelectionModel.h b/Engine/Plugins/Developer/Concert/ConcertApp/MultiUserClient/Source/MultiUserClient/Private/Widgets/ActiveSession/Replication/Client/Multi/AllClientsSelectionModel.h
deleted file mode 100644
diff --git a/Engine/Plugins/Developer/Concert/ConcertApp/MultiUserClient/Source/MultiUserClient/Private/Widgets/ActiveSession/Replication/Client/Multi/IClientSelectionModel.h b/Engine/Plugins/Developer/Concert/ConcertApp/MultiUserClient/Source/MultiUserClient/Private/Widgets/ActiveSession/Replication/Client/Multi/IClientSelectionModel.h
deleted file mode 100644
diff --git a/Engine/Plugins/Developer/Concert/ConcertApp/MultiUserClient/Source/MultiUserClient/Private/Widgets/ActiveSession/Replication/Client/Multi/MultiStreamModel.cpp b/Engine/Plugins/Developer/Concert/ConcertApp/MultiUserClient/Source/MultiUserClient/Private/Widgets/ActiveSession/Replication/Client/Multi/MultiStreamModel.cpp
--- a/Engine/Plugins/Developer/Concert/ConcertApp/MultiUserClient/Source/MultiUserClient/Private/Widgets/ActiveSession/Replication/Client/Multi/MultiStreamModel.cpp
+++ b/Engine/Plugins/Developer/Concert/ConcertApp/MultiUserClient/Source/MultiUserClient/Private/Widgets/ActiveSession/Replication/Client/Multi/MultiStreamModel.cpp
@@ -2,18 +2,18 @@
 
 #include "MultiStreamModel.h"
 
-#include "IClientSelectionModel.h"
 #include "Replication/Client/Online/OnlineClient.h"
 #include "Replication/Client/Online/OnlineClientManager.h"
 #include "Replication/Editor/Model/IEditableReplicationStreamModel.h"
+#include "Selection/ISelectionModel.h"
 
 namespace UE::MultiUserClient::Replication
 {
-	FMultiStreamModel::FMultiStreamModel(IClientSelectionModel& InClientSelectionModel, FOnlineClientManager& InClientManager)
-		: ClientSelectionModel(InClientSelectionModel)
+	FMultiStreamModel::FMultiStreamModel(IOnlineClientSelectionModel& InOnlineClientSelectionModel, FOnlineClientManager& InClientManager)
+		: OnlineClientSelectionModel(InOnlineClientSelectionModel)
 		, ClientManager(InClientManager)
 	{
-		ClientSelectionModel.OnSelectionChanged().AddRaw(this, &FMultiStreamModel::RebuildStreamsSets);
+		OnlineClientSelectionModel.OnSelectionChanged().AddRaw(this, &FMultiStreamModel::RebuildStreamsSets);
 		RebuildStreamsSets();
 	}
 
@@ -58,35 +58,28 @@ namespace UE::MultiUserClient::Replication
 			return EBreakBehavior::Continue;
 		});
 		
-		TSet<const FOnlineClient*> ReadOnlyClients;
-		TSet<const FOnlineClient*> WritableClients;
-		ClientSelectionModel.ForEachSelectedClient([this, &ReadOnlyClients, &WritableClients](FOnlineClient& Client)
+		TSet<const FOnlineClient*> ReadOnlyStreams;
+		TSet<const FOnlineClient*> WritableStreams;
+		OnlineClientSelectionModel.ForEachItem([this, &ReadOnlyStreams, &WritableStreams](FOnlineClient& Client)
 		{
 			const bool bIsUploadable = CanEverSubmit(Client.GetSubmissionWorkflow().GetUploadability());
 			const TSharedRef<ConcertSharedSlate::IEditableReplicationStreamModel> Stream = Client.GetClientEditModel();
 			Client.OnModelChanged().AddRaw(this, &FMultiStreamModel::OnStreamExternallyChanged, Stream.ToWeakPtr());
 			
-			if (bIsUploadable)
-			{
-				WritableClients.Add(&Client);
-			}
-			else
-			{
-				ReadOnlyClients.Add(&Client);
-			}
-			
+			TSet<const FOnlineClient*>& StreamToAddTo = Client.AllowsEditing() ? WritableStreams : ReadOnlyStreams;
+			StreamToAddTo.Add(&Client);
 			return EBreakBehavior::Continue;
 		});
 
-		const bool bReadOnlyStayedSame = CachedReadOnlyClients.Num() == ReadOnlyClients.Num() && CachedReadOnlyClients.Includes(ReadOnlyClients);
+		const bool bReadOnlyStayedSame = CachedReadOnlyClients.Num() == ReadOnlyStreams.Num() && CachedReadOnlyClients.Includes(ReadOnlyStreams);
 		if (!bReadOnlyStayedSame)
 		{
-			CachedReadOnlyClients = MoveTemp(ReadOnlyClients);
+			CachedReadOnlyClients = MoveTemp(ReadOnlyStreams);
 		}
-		const bool bWritableStayedSame = CachedWritableClients.Num() == WritableClients.Num() && CachedWritableClients.Includes(WritableClients);
+		const bool bWritableStayedSame = CachedWritableClients.Num() == WritableStreams.Num() && CachedWritableClients.Includes(WritableStreams);
 		if (!bWritableStayedSame)
 		{
-			CachedWritableClients = MoveTemp(WritableClients);
+			CachedWritableClients = MoveTemp(WritableStreams);
 		}
 
 		if (!bReadOnlyStayedSame || !bWritableStayedSame)
diff --git a/Engine/Plugins/Developer/Concert/ConcertApp/MultiUserClient/Source/MultiUserClient/Private/Widgets/ActiveSession/Replication/Client/Multi/MultiStreamModel.h b/Engine/Plugins/Developer/Concert/ConcertApp/MultiUserClient/Source/MultiUserClient/Private/Widgets/ActiveSession/Replication/Client/Multi/MultiStreamModel.h
--- a/Engine/Plugins/Developer/Concert/ConcertApp/MultiUserClient/Source/MultiUserClient/Private/Widgets/ActiveSession/Replication/Client/Multi/MultiStreamModel.h
+++ b/Engine/Plugins/Developer/Concert/ConcertApp/MultiUserClient/Source/MultiUserClient/Private/Widgets/ActiveSession/Replication/Client/Multi/MultiStreamModel.h
@@ -5,6 +5,8 @@
 #include "Containers/Set.h"
 #include "Replication/Editor/Model/IEditableMultiReplicationStreamModel.h"
 #include "Replication/Editor/Model/IEditableReplicationStreamModel.h"
+#include "Selection/SelectionModelFwd.h"
+
 #include "Templates/Function.h"
 #include "Templates/SharedPointer.h"
 
@@ -14,14 +16,13 @@ namespace UE::MultiUserClient::Replication
 {
 	class FOnlineClient;
 	class FOnlineClientManager;
-	class IClientSelectionModel;
 
 	/** Checks whether clients accept remote changes and categorizes them into read-only and writable. */
 	class FMultiStreamModel : public ConcertSharedSlate::IEditableMultiReplicationStreamModel
 	{
 	public:
 		
-		FMultiStreamModel(IClientSelectionModel& InClientSelectionModel, FOnlineClientManager& InClientManager);
+		FMultiStreamModel(IOnlineClientSelectionModel& InOnlineClientSelectionModel, FOnlineClientManager& InClientManager);
 		
 		const TSet<const FOnlineClient*>& GetCachedReadOnlyClients() const { return CachedReadOnlyClients; }
 		const TSet<const FOnlineClient*>& GetCachedWritableClients() const { return CachedWritableClients; }
@@ -36,9 +37,9 @@ namespace UE::MultiUserClient::Replication
 
 	private:
 		
-		/** Gets all clients to display and informs when they change. */
-		IClientSelectionModel& ClientSelectionModel;
-		/** Used to obtain a list of clients for unsubscribing */
+		/** Gets all online clients that are supposed to be displayed. */
+		IOnlineClientSelectionModel& OnlineClientSelectionModel;
+		/** Used to obtain a list of clients for unsubscribing. */
 		FOnlineClientManager& ClientManager;
 
 		TSet<const FOnlineClient*> CachedReadOnlyClients;
diff --git a/Engine/Plugins/Developer/Concert/ConcertApp/MultiUserClient/Source/MultiUserClient/Private/Widgets/ActiveSession/Replication/Client/Multi/SAllClientsView.cpp b/Engine/Plugins/Developer/Concert/ConcertApp/MultiUserClient/Source/MultiUserClient/Private/Widgets/ActiveSession/Replication/Client/Multi/SAllClientsView.cpp
--- a/Engine/Plugins/Developer/Concert/ConcertApp/MultiUserClient/Source/MultiUserClient/Private/Widgets/ActiveSession/Replication/Client/Multi/SAllClientsView.cpp
+++ b/Engine/Plugins/Developer/Concert/ConcertApp/MultiUserClient/Source/MultiUserClient/Private/Widgets/ActiveSession/Replication/Client/Multi/SAllClientsView.cpp
@@ -2,12 +2,12 @@
 
 #include "SAllClientsView.h"
 
-#include "AllClientsSelectionModel.h"
+#include "Selection/AllClientsSelectionModel.h"
+#include "Replication/MultiUserReplicationManager.h"
 #include "SMultiClientView.h"
 #include "Replication/Client/Online/OnlineClientManager.h"
 
 #include "Algo/Transform.h"
-#include "Replication/MultiUserReplicationManager.h"
 
 namespace UE::MultiUserClient::Replication
 {
diff --git a/Engine/Plugins/Developer/Concert/ConcertApp/MultiUserClient/Source/MultiUserClient/Private/Widgets/ActiveSession/Replication/Client/Multi/SAllClientsView.h b/Engine/Plugins/Developer/Concert/ConcertApp/MultiUserClient/Source/MultiUserClient/Private/Widgets/ActiveSession/Replication/Client/Multi/SAllClientsView.h
--- a/Engine/Plugins/Developer/Concert/ConcertApp/MultiUserClient/Source/MultiUserClient/Private/Widgets/ActiveSession/Replication/Client/Multi/SAllClientsView.h
+++ b/Engine/Plugins/Developer/Concert/ConcertApp/MultiUserClient/Source/MultiUserClient/Private/Widgets/ActiveSession/Replication/Client/Multi/SAllClientsView.h
@@ -2,7 +2,7 @@
 
 #pragma once
 
-#include "AllClientsSelectionModel.h"
+#include "Selection/AllClientsSelectionModel.h"
 
 #include "HAL/Platform.h"
 #include "Widgets/DeclarativeSyntaxSupport.h"
diff --git a/Engine/Plugins/Developer/Concert/ConcertApp/MultiUserClient/Source/MultiUserClient/Private/Widgets/ActiveSession/Replication/Client/Multi/SMultiClientView.cpp b/Engine/Plugins/Developer/Concert/ConcertApp/MultiUserClient/Source/MultiUserClient/Private/Widgets/ActiveSession/Replication/Client/Multi/SMultiClientView.cpp
--- a/Engine/Plugins/Developer/Concert/ConcertApp/MultiUserClient/Source/MultiUserClient/Private/Widgets/ActiveSession/Replication/Client/Multi/SMultiClientView.cpp
+++ b/Engine/Plugins/Developer/Concert/ConcertApp/MultiUserClient/Source/MultiUserClient/Private/Widgets/ActiveSession/Replication/Client/Multi/SMultiClientView.cpp
@@ -2,7 +2,7 @@
 
 #include "SMultiClientView.h"
 
-#include "IClientSelectionModel.h"
+#include "Selection/ISelectionModel.h"
 #include "MultiStreamModel.h"
 #include "Replication/ClientReplicationWidgetFactories.h"
 #include "Replication/MultiUserReplicationManager.h"
@@ -34,17 +34,17 @@ namespace UE::MultiUserClient::Replication
 		const FArguments&,
 		TSharedRef<IConcertClient> InConcertClient,
 		FMultiUserReplicationManager& InMultiUserReplicationManager,
-		IClientSelectionModel& InDisplayClientsModel
+		IOnlineClientSelectionModel& InOnlineClientSelectionModel
 		)
 	{
 		ClientManager = InMultiUserReplicationManager.GetClientManager();
 		UserSelectedProperties = InMultiUserReplicationManager.GetUserPropertySelector();
-		StreamModel = MakeShared<FMultiStreamModel>(InDisplayClientsModel, *ClientManager);
-		SelectionModel = &InDisplayClientsModel;
+		StreamModel = MakeShared<FMultiStreamModel>(InOnlineClientSelectionModel, *ClientManager);
+		OnlineClientSelectionModel = &InOnlineClientSelectionModel;
 		ConcertClient = MoveTemp(InConcertClient);
 		
 		ClientManager->OnRemoteClientsChanged().AddSP(this, &SMultiClientView::RebuildClientSubscriptions);
-		SelectionModel->OnSelectionChanged().AddSP(this, &SMultiClientView::RebuildClientSubscriptions);
+		OnlineClientSelectionModel->OnSelectionChanged().AddSP(this, &SMultiClientView::RebuildClientSubscriptions);
 		UserSelectedProperties->OnPropertySelectionChanged().AddRaw(this, &SMultiClientView::RefreshUI);
 
 		TSharedPtr<SVerticalBox> Content;
@@ -196,7 +196,7 @@ namespace UE::MultiUserClient::Replication
 	{
 		CleanClientSubscriptions();
 
-		SelectionModel->ForEachSelectedClient([this](FOnlineClient& Client)
+		OnlineClientSelectionModel->ForEachItem([this](FOnlineClient& Client)
 		{
 			Client.OnModelChanged().AddSP(this, &SMultiClientView::RefreshUI);
 			Client.OnHierarchyNeedsRefresh().AddRaw(this, &SMultiClientView::RefreshUI);
diff --git a/Engine/Plugins/Developer/Concert/ConcertApp/MultiUserClient/Source/MultiUserClient/Private/Widgets/ActiveSession/Replication/Client/Multi/SMultiClientView.h b/Engine/Plugins/Developer/Concert/ConcertApp/MultiUserClient/Source/MultiUserClient/Private/Widgets/ActiveSession/Replication/Client/Multi/SMultiClientView.h
--- a/Engine/Plugins/Developer/Concert/ConcertApp/MultiUserClient/Source/MultiUserClient/Private/Widgets/ActiveSession/Replication/Client/Multi/SMultiClientView.h
+++ b/Engine/Plugins/Developer/Concert/ConcertApp/MultiUserClient/Source/MultiUserClient/Private/Widgets/ActiveSession/Replication/Client/Multi/SMultiClientView.h
@@ -4,9 +4,9 @@
 
 #include "Replication/Editor/Model/ObjectSource/IObjectSourceModel.h"
 #include "Replication/Editor/UnrealEditor/HideObjectsNotInWorldLogic.h"
+#include "Selection/SelectionModelFwd.h"
 #include "Widgets/DeclarativeSyntaxSupport.h"
 #include "Widgets/SCompoundWidget.h"
-#include "Widgets/ActiveSession/Replication/Client/PropertySelection/RootPropertySourceModel.h"
 
 namespace UE::MultiUserClient::Replication
 {
@@ -31,7 +31,6 @@ namespace UE::MultiUserClient::Replication
 	class FMultiUserReplicationManager;
 	class FOnlineClient;
 	class FOnlineClientManager;
-	class IClientSelectionModel;
 	class SPropertySelectionComboButton;
 
 	/** Displays a selection of clients. */
@@ -45,7 +44,7 @@ namespace UE::MultiUserClient::Replication
 			const FArguments& InArgs,
 			TSharedRef<IConcertClient> InConcertClient,
 			FMultiUserReplicationManager& InMultiUserReplicationManager UE_LIFETIMEBOUND,
-			IClientSelectionModel& InDisplayClientsModel UE_LIFETIMEBOUND
+			IOnlineClientSelectionModel& InOnlineClientSelectionModel UE_LIFETIMEBOUND
 			);
 		virtual ~SMultiClientView() override;
 
@@ -54,7 +53,7 @@ namespace UE::MultiUserClient::Replication
 		TSharedPtr<IConcertClient> ConcertClient;
 		FOnlineClientManager* ClientManager = nullptr;
 		FUserPropertySelector* UserSelectedProperties = nullptr;
-		IClientSelectionModel* SelectionModel = nullptr;
+		IOnlineClientSelectionModel* OnlineClientSelectionModel = nullptr;
 		
 		/** Combines the clients */
 		TSharedPtr<FMultiStreamModel> StreamModel;
diff --git a/Engine/Plugins/Developer/Concert/ConcertApp/MultiUserClient/Source/MultiUserClient/Private/Widgets/ActiveSession/Replication/Client/Multi/Selection/AllClientsSelectionModel.cpp b/Engine/Plugins/Developer/Concert/ConcertApp/MultiUserClient/Source/MultiUserClient/Private/Widgets/ActiveSession/Replication/Client/Multi/Selection/AllClientsSelectionModel.cpp
new file mode 100644
--- /dev/null
+++ b/Engine/Plugins/Developer/Concert/ConcertApp/MultiUserClient/Source/MultiUserClient/Private/Widgets/ActiveSession/Replication/Client/Multi/Selection/AllClientsSelectionModel.cpp
@@ -0,0 +1,29 @@
+// Copyright Epic Games, Inc. All Rights Reserved.
+
+#include "AllClientsSelectionModel.h"
+
+#include "Replication/Client/Online/OnlineClientManager.h"
+
+namespace UE::MultiUserClient::Replication
+{
+	FAllClientsSelectionModel::FAllClientsSelectionModel(FOnlineClientManager& InClientManager)
+		: ClientManager(InClientManager)
+	{
+		ClientManager.OnRemoteClientsChanged().AddRaw(this, &FAllClientsSelectionModel::OnRemoteClientsChanged);
+	}
+
+	FAllClientsSelectionModel::~FAllClientsSelectionModel()
+	{
+		ClientManager.OnRemoteClientsChanged().RemoveAll(this);
+	}
+
+	void FAllClientsSelectionModel::ForEachItem(TFunctionRef<EBreakBehavior(FOnlineClient&)> ProcessClient) const
+	{
+		ClientManager.ForEachClient(ProcessClient);
+	}
+
+	void FAllClientsSelectionModel::OnRemoteClientsChanged()
+	{
+		OnSelectionChangedDelegate.Broadcast();
+	}
+}
diff --git a/Engine/Plugins/Developer/Concert/ConcertApp/MultiUserClient/Source/MultiUserClient/Private/Widgets/ActiveSession/Replication/Client/Multi/Selection/AllClientsSelectionModel.h b/Engine/Plugins/Developer/Concert/ConcertApp/MultiUserClient/Source/MultiUserClient/Private/Widgets/ActiveSession/Replication/Client/Multi/Selection/AllClientsSelectionModel.h
new file mode 100644
--- /dev/null
+++ b/Engine/Plugins/Developer/Concert/ConcertApp/MultiUserClient/Source/MultiUserClient/Private/Widgets/ActiveSession/Replication/Client/Multi/Selection/AllClientsSelectionModel.h
@@ -0,0 +1,35 @@
+// Copyright Epic Games, Inc. All Rights Reserved.
+
+#pragma once
+
+#include "ISelectionModel.h"
+#include "SelectionModelFwd.h"
+
+namespace UE::MultiUserClient::Replication
+{
+	class FOnlineClientManager;
+	
+	/** Exposes all clients and detects when clients disconnect. */
+	class FAllClientsSelectionModel : public IOnlineClientSelectionModel
+	{
+	public:
+		
+		FAllClientsSelectionModel(FOnlineClientManager& InClientManager);
+		virtual ~FAllClientsSelectionModel() override;
+
+		//~ Begin IClientSelectionModel Interface
+		virtual void ForEachItem(TFunctionRef<EBreakBehavior(FOnlineClient&)> ProcessClient) const override;
+		virtual FOnSelectionChanged& OnSelectionChanged() override { return OnSelectionChangedDelegate; }
+		//~ End IClientSelectionModel Interface
+
+	private:
+
+		/** Informs us when the list of clients changes */
+		FOnlineClientManager& ClientManager;
+		
+		/** Called when the clients ForEachSelectedClient enumerates has changed. */
+		FOnSelectionChanged OnSelectionChangedDelegate;
+		
+		void OnRemoteClientsChanged();
+	};
+}
diff --git a/Engine/Plugins/Developer/Concert/ConcertApp/MultiUserClient/Source/MultiUserClient/Private/Widgets/ActiveSession/Replication/Client/Multi/Selection/ISelectionModel.h b/Engine/Plugins/Developer/Concert/ConcertApp/MultiUserClient/Source/MultiUserClient/Private/Widgets/ActiveSession/Replication/Client/Multi/Selection/ISelectionModel.h
new file mode 100644
--- /dev/null
+++ b/Engine/Plugins/Developer/Concert/ConcertApp/MultiUserClient/Source/MultiUserClient/Private/Widgets/ActiveSession/Replication/Client/Multi/Selection/ISelectionModel.h
@@ -0,0 +1,30 @@
+// Copyright Epic Games, Inc. All Rights Reserved.
+
+#pragma once
+
+#include "Delegates/Delegate.h"
+#include "Templates/Function.h"
+
+enum class EBreakBehavior : uint8;
+
+namespace UE::MultiUserClient::Replication
+{
+	class FOnlineClient;
+	
+	/** Abstracts the concept of a client selection. */
+	template<typename TItemType>
+	class ISelectionModel
+	{
+	public:
+
+		/** Iterates through every client in the selection. */
+		virtual void ForEachItem(TFunctionRef<EBreakBehavior(TItemType&)> ProcessItem) const = 0;
+
+		DECLARE_MULTICAST_DELEGATE(FOnSelectionChanged);
+		/** Called when the clients ForEachSelectedClient enumerates has changed. */
+		virtual FOnSelectionChanged& OnSelectionChanged() = 0;
+		
+		virtual ~ISelectionModel() = default;
+	};
+
+}
diff --git a/Engine/Plugins/Developer/Concert/ConcertApp/MultiUserClient/Source/MultiUserClient/Private/Widgets/ActiveSession/Replication/Client/Multi/Selection/SelectionModelFwd.h b/Engine/Plugins/Developer/Concert/ConcertApp/MultiUserClient/Source/MultiUserClient/Private/Widgets/ActiveSession/Replication/Client/Multi/Selection/SelectionModelFwd.h
new file mode 100644
--- /dev/null
+++ b/Engine/Plugins/Developer/Concert/ConcertApp/MultiUserClient/Source/MultiUserClient/Private/Widgets/ActiveSession/Replication/Client/Multi/Selection/SelectionModelFwd.h
@@ -0,0 +1,13 @@
+﻿// Copyright Epic Games, Inc. All Rights Reserved.
+
+#pragma once
+
+#include "ISelectionModel.h"
+
+namespace UE::MultiUserClient::Replication
+{
+	class FOnlineClient;
+	struct FClientItem;
+	
+	using IOnlineClientSelectionModel = ISelectionModel<FOnlineClient>;
+}
diff --git a/Engine/Plugins/Developer/Concert/ConcertApp/MultiUserClient/Source/MultiUserClient/Private/Widgets/ActiveSession/Replication/Client/PropertySelection/SPropertySelectionComboButton.h b/Engine/Plugins/Developer/Concert/ConcertApp/MultiUserClient/Source/MultiUserClient/Private/Widgets/ActiveSession/Replication/Client/PropertySelection/SPropertySelectionComboButton.h
--- a/Engine/Plugins/Developer/Concert/ConcertApp/MultiUserClient/Source/MultiUserClient/Private/Widgets/ActiveSession/Replication/Client/PropertySelection/SPropertySelectionComboButton.h
+++ b/Engine/Plugins/Developer/Concert/ConcertApp/MultiUserClient/Source/MultiUserClient/Private/Widgets/ActiveSession/Replication/Client/PropertySelection/SPropertySelectionComboButton.h
@@ -2,6 +2,7 @@
 
 #pragma once
 
+#include "RootPropertySourceModel.h"
 #include "Model/Item/SourceModelBuilders.h"
 #include "SelectionDelegates.h"
 
@@ -14,7 +15,6 @@ namespace UE::ConcertSharedSlate { struct FObjectGroup; }
 
 namespace UE::MultiUserClient::Replication
 {
-	class FRootPropertySourceModel;
 	class FUserPropertySelector;
 	struct FUserSelectableProperty;
 	
