Rename FAllClientsSelectionModel to FAllOnlineClientsSelectionModel

diff --git a/Engine/Plugins/Developer/Concert/ConcertApp/MultiUserClient/Source/MultiUserClient/Private/Widgets/ActiveSession/Replication/Client/Multi/SAllClientsView.cpp b/Engine/Plugins/Developer/Concert/ConcertApp/MultiUserClient/Source/MultiUserClient/Private/Widgets/ActiveSession/Replication/Client/Multi/SAllClientsView.cpp
--- a/Engine/Plugins/Developer/Concert/ConcertApp/MultiUserClient/Source/MultiUserClient/Private/Widgets/ActiveSession/Replication/Client/Multi/SAllClientsView.cpp
+++ b/Engine/Plugins/Developer/Concert/ConcertApp/MultiUserClient/Source/MultiUserClient/Private/Widgets/ActiveSession/Replication/Client/Multi/SAllClientsView.cpp
@@ -2,7 +2,7 @@
 
 #include "SAllClientsView.h"
 
-#include "Selection/AllClientsSelectionModel.h"
+#include "Selection/AllOnlineClientsSelectionModel.h"
 #include "Replication/MultiUserReplicationManager.h"
 #include "SMultiClientView.h"
 #include "Replication/Client/Online/OnlineClientManager.h"
@@ -15,7 +15,7 @@ namespace UE::MultiUserClient::Replication
 	{
 		ClientManager = InMultiUserReplicationManager.GetClientManager();
 		check(ClientManager);
-		AllClientsModel = MakeUnique<FAllClientsSelectionModel>(*ClientManager);
+		AllClientsModel = MakeUnique<FAllOnlineClientsSelectionModel>(*ClientManager);
 		
 		ChildSlot
 		[
diff --git a/Engine/Plugins/Developer/Concert/ConcertApp/MultiUserClient/Source/MultiUserClient/Private/Widgets/ActiveSession/Replication/Client/Multi/SAllClientsView.h b/Engine/Plugins/Developer/Concert/ConcertApp/MultiUserClient/Source/MultiUserClient/Private/Widgets/ActiveSession/Replication/Client/Multi/SAllClientsView.h
--- a/Engine/Plugins/Developer/Concert/ConcertApp/MultiUserClient/Source/MultiUserClient/Private/Widgets/ActiveSession/Replication/Client/Multi/SAllClientsView.h
+++ b/Engine/Plugins/Developer/Concert/ConcertApp/MultiUserClient/Source/MultiUserClient/Private/Widgets/ActiveSession/Replication/Client/Multi/SAllClientsView.h
@@ -2,7 +2,7 @@
 
 #pragma once
 
-#include "Selection/AllClientsSelectionModel.h"
+#include "Selection/AllOnlineClientsSelectionModel.h"
 
 #include "HAL/Platform.h"
 #include "Widgets/DeclarativeSyntaxSupport.h"
@@ -30,7 +30,7 @@ namespace UE::MultiUserClient::Replication
 		FOnlineClientManager* ClientManager = nullptr;
 
 		/** Keeps the SMultiClientView updated of any changes to clients (e.g. disconnects, etc.) */
-		TUniquePtr<FAllClientsSelectionModel> AllClientsModel;
+		TUniquePtr<FAllOnlineClientsSelectionModel> AllClientsModel;
 
 		/** Gets all the clients to display */
 		TSet<const FOnlineClient*> GetAllClients() const;
diff --git a/Engine/Plugins/Developer/Concert/ConcertApp/MultiUserClient/Source/MultiUserClient/Private/Widgets/ActiveSession/Replication/Client/Multi/Selection/AllClientsSelectionModel.cpp b/Engine/Plugins/Developer/Concert/ConcertApp/MultiUserClient/Source/MultiUserClient/Private/Widgets/ActiveSession/Replication/Client/Multi/Selection/AllClientsSelectionModel.cpp
deleted file mode 100644
diff --git a/Engine/Plugins/Developer/Concert/ConcertApp/MultiUserClient/Source/MultiUserClient/Private/Widgets/ActiveSession/Replication/Client/Multi/Selection/AllClientsSelectionModel.h b/Engine/Plugins/Developer/Concert/ConcertApp/MultiUserClient/Source/MultiUserClient/Private/Widgets/ActiveSession/Replication/Client/Multi/Selection/AllClientsSelectionModel.h
deleted file mode 100644
diff --git a/Engine/Plugins/Developer/Concert/ConcertApp/MultiUserClient/Source/MultiUserClient/Private/Widgets/ActiveSession/Replication/Client/Multi/Selection/AllOnlineClientsSelectionModel.cpp b/Engine/Plugins/Developer/Concert/ConcertApp/MultiUserClient/Source/MultiUserClient/Private/Widgets/ActiveSession/Replication/Client/Multi/Selection/AllOnlineClientsSelectionModel.cpp
new file mode 100644
--- /dev/null
+++ b/Engine/Plugins/Developer/Concert/ConcertApp/MultiUserClient/Source/MultiUserClient/Private/Widgets/ActiveSession/Replication/Client/Multi/Selection/AllOnlineClientsSelectionModel.cpp
@@ -0,0 +1,29 @@
+// Copyright Epic Games, Inc. All Rights Reserved.
+
+#include "AllOnlineClientsSelectionModel.h"
+
+#include "Replication/Client/Online/OnlineClientManager.h"
+
+namespace UE::MultiUserClient::Replication
+{
+	FAllOnlineClientsSelectionModel::FAllOnlineClientsSelectionModel(FOnlineClientManager& InClientManager)
+		: ClientManager(InClientManager)
+	{
+		ClientManager.OnRemoteClientsChanged().AddRaw(this, &FAllOnlineClientsSelectionModel::OnRemoteClientsChanged);
+	}
+
+	FAllOnlineClientsSelectionModel::~FAllOnlineClientsSelectionModel()
+	{
+		ClientManager.OnRemoteClientsChanged().RemoveAll(this);
+	}
+
+	void FAllOnlineClientsSelectionModel::ForEachItem(TFunctionRef<EBreakBehavior(FOnlineClient&)> ProcessClient) const
+	{
+		ClientManager.ForEachClient(ProcessClient);
+	}
+
+	void FAllOnlineClientsSelectionModel::OnRemoteClientsChanged()
+	{
+		OnSelectionChangedDelegate.Broadcast();
+	}
+}
diff --git a/Engine/Plugins/Developer/Concert/ConcertApp/MultiUserClient/Source/MultiUserClient/Private/Widgets/ActiveSession/Replication/Client/Multi/Selection/AllOnlineClientsSelectionModel.h b/Engine/Plugins/Developer/Concert/ConcertApp/MultiUserClient/Source/MultiUserClient/Private/Widgets/ActiveSession/Replication/Client/Multi/Selection/AllOnlineClientsSelectionModel.h
new file mode 100644
--- /dev/null
+++ b/Engine/Plugins/Developer/Concert/ConcertApp/MultiUserClient/Source/MultiUserClient/Private/Widgets/ActiveSession/Replication/Client/Multi/Selection/AllOnlineClientsSelectionModel.h
@@ -0,0 +1,37 @@
+// Copyright Epic Games, Inc. All Rights Reserved.
+
+#pragma once
+
+#include "ISelectionModel.h"
+#include "SelectionModelFwd.h"
+
+#include "HAL/Platform.h"
+
+namespace UE::MultiUserClient::Replication
+{
+	class FOnlineClientManager;
+	
+	/** Exposes all clients and detects when clients disconnect. */
+	class FAllOnlineClientsSelectionModel : public IOnlineClientSelectionModel
+	{
+	public:
+		
+		FAllOnlineClientsSelectionModel(FOnlineClientManager& InClientManager UE_LIFETIMEBOUND);
+		virtual ~FAllOnlineClientsSelectionModel() override;
+
+		//~ Begin IOnlineClientSelectionModel Interface
+		virtual void ForEachItem(TFunctionRef<EBreakBehavior(FOnlineClient&)> ProcessClient) const override;
+		virtual FOnSelectionChanged& OnSelectionChanged() override { return OnSelectionChangedDelegate; }
+		//~ End IOnlineClientSelectionModel Interface
+
+	private:
+
+		/** Informs us when the list of clients changes */
+		FOnlineClientManager& ClientManager;
+		
+		/** Called when the clients ForEachSelectedClient enumerates has changed. */
+		FOnSelectionChanged OnSelectionChangedDelegate;
+		
+		void OnRemoteClientsChanged();
+	};
+}
