Rename OnServerStateChanged in IClientStreamSynchronizer and IClientAuthoritySynchronizer to avoid ambigous calls when implementing both interfaces in the same class.

diff --git a/Engine/Plugins/Developer/Concert/ConcertApp/MultiUserClient/Source/MultiUserClient/Private/Replication/Authority/IClientAuthoritySynchronizer.h b/Engine/Plugins/Developer/Concert/ConcertApp/MultiUserClient/Source/MultiUserClient/Private/Replication/Authority/IClientAuthoritySynchronizer.h
--- a/Engine/Plugins/Developer/Concert/ConcertApp/MultiUserClient/Source/MultiUserClient/Private/Replication/Authority/IClientAuthoritySynchronizer.h
+++ b/Engine/Plugins/Developer/Concert/ConcertApp/MultiUserClient/Source/MultiUserClient/Private/Replication/Authority/IClientAuthoritySynchronizer.h
@@ -22,7 +22,7 @@ namespace UE::MultiUserClient::Replication
 
 		DECLARE_MULTICAST_DELEGATE(FOnServerStateChanged);
 		/** @return Event executed when authority state has been updated. */
-		virtual FOnServerStateChanged& OnServerStateChanged() = 0;
+		virtual FOnServerStateChanged& OnServerAuthorityChanged() = 0;
 
 		virtual ~IClientAuthoritySynchronizer() = default;
 	};
@@ -35,7 +35,7 @@ namespace UE::MultiUserClient::Replication
 	public:
 
 		//~ Begin IClientAuthoritySynchronizer Interface
-		virtual FOnServerStateChanged& OnServerStateChanged() override { return OnServerStateChangedDelegate; }
+		virtual FOnServerStateChanged& OnServerAuthorityChanged() override { return OnServerStateChangedDelegate; }
 		//~ End IClientAuthoritySynchronizer Interface
 
 	protected:
diff --git a/Engine/Plugins/Developer/Concert/ConcertApp/MultiUserClient/Source/MultiUserClient/Private/Replication/Client/Online/OnlineClient.cpp b/Engine/Plugins/Developer/Concert/ConcertApp/MultiUserClient/Source/MultiUserClient/Private/Replication/Client/Online/OnlineClient.cpp
--- a/Engine/Plugins/Developer/Concert/ConcertApp/MultiUserClient/Source/MultiUserClient/Private/Replication/Client/Online/OnlineClient.cpp
+++ b/Engine/Plugins/Developer/Concert/ConcertApp/MultiUserClient/Source/MultiUserClient/Private/Replication/Client/Online/OnlineClient.cpp
@@ -60,7 +60,7 @@ namespace UE::MultiUserClient::Replication
 		LocalAuthorityDiffer.OnChangedOwnedObjects().AddRaw(this, &FOnlineClient::DeferOnModelChanged);
 
 		SubmissionWorkflow->OnAuthorityRequestCompleted_AnyThread().AddRaw(this, &FOnlineClient::OnAuthoritySubmissionCompleted);
-		StreamSynchronizer->OnServerStateChanged().AddRaw(this, &FOnlineClient::OnServerStateChanged);
+		StreamSynchronizer->OnServerStreanChanged().AddRaw(this, &FOnlineClient::OnServerStateChanged);
 
 		LevelModificationHandler.OnHierarchyNeedsRefresh().AddLambda([this]()
 		{
diff --git a/Engine/Plugins/Developer/Concert/ConcertApp/MultiUserClient/Source/MultiUserClient/Private/Replication/Misc/Frequency/FrequencyChangeTracker.cpp b/Engine/Plugins/Developer/Concert/ConcertApp/MultiUserClient/Source/MultiUserClient/Private/Replication/Misc/Frequency/FrequencyChangeTracker.cpp
--- a/Engine/Plugins/Developer/Concert/ConcertApp/MultiUserClient/Source/MultiUserClient/Private/Replication/Misc/Frequency/FrequencyChangeTracker.cpp
+++ b/Engine/Plugins/Developer/Concert/ConcertApp/MultiUserClient/Source/MultiUserClient/Private/Replication/Misc/Frequency/FrequencyChangeTracker.cpp
@@ -13,12 +13,12 @@ namespace UE::MultiUserClient::Replication
 		: StreamSynchronizer(InStreamSynchronizer)
 	{
 		// Refresh locally recorded changes when the server state changes, e.g. after a request has been served.
-		StreamSynchronizer.OnServerStateChanged().AddRaw(this, &FFrequencyChangeTracker::RefreshChanges);
+		StreamSynchronizer.OnServerStreanChanged().AddRaw(this, &FFrequencyChangeTracker::RefreshChanges);
 	}
 
 	FFrequencyChangeTracker::~FFrequencyChangeTracker()
 	{
-		StreamSynchronizer.OnServerStateChanged().RemoveAll(this);
+		StreamSynchronizer.OnServerStreanChanged().RemoveAll(this);
 	}
 
 	void FFrequencyChangeTracker::AddOverride(FSoftObjectPath Object, FConcertObjectReplicationSettings NewSettings)
diff --git a/Engine/Plugins/Developer/Concert/ConcertApp/MultiUserClient/Source/MultiUserClient/Private/Replication/Misc/GlobalAuthorityCache.cpp b/Engine/Plugins/Developer/Concert/ConcertApp/MultiUserClient/Source/MultiUserClient/Private/Replication/Misc/GlobalAuthorityCache.cpp
--- a/Engine/Plugins/Developer/Concert/ConcertApp/MultiUserClient/Source/MultiUserClient/Private/Replication/Misc/GlobalAuthorityCache.cpp
+++ b/Engine/Plugins/Developer/Concert/ConcertApp/MultiUserClient/Source/MultiUserClient/Private/Replication/Misc/GlobalAuthorityCache.cpp
@@ -214,14 +214,14 @@ namespace UE::MultiUserClient::Replication
 	void FGlobalAuthorityCache::RegisterForClientEvents(const FOnlineClient& Client)
 	{
 		const FGuid& ClientEndpointId = Client.GetEndpointId();
-		Client.GetAuthoritySynchronizer().OnServerStateChanged().AddRaw(this, &FGlobalAuthorityCache::OnPostAuthorityChanged, ClientEndpointId);
-		Client.GetStreamSynchronizer().OnServerStateChanged().AddRaw(this, &FGlobalAuthorityCache::OnStreamChanged, ClientEndpointId);
+		Client.GetAuthoritySynchronizer().OnServerAuthorityChanged().AddRaw(this, &FGlobalAuthorityCache::OnPostAuthorityChanged, ClientEndpointId);
+		Client.GetStreamSynchronizer().OnServerStreanChanged().AddRaw(this, &FGlobalAuthorityCache::OnStreamChanged, ClientEndpointId);
 	}
 
 	void FGlobalAuthorityCache::UnregisterFromClientEvents(const FOnlineClient& Client) const
 	{
-		Client.GetAuthoritySynchronizer().OnServerStateChanged().RemoveAll(this);
-		Client.GetStreamSynchronizer().OnServerStateChanged().RemoveAll(this);
+		Client.GetAuthoritySynchronizer().OnServerAuthorityChanged().RemoveAll(this);
+		Client.GetStreamSynchronizer().OnServerStreanChanged().RemoveAll(this);
 	}
 
 	void FGlobalAuthorityCache::AddClient(const FGuid& ClientId)
diff --git a/Engine/Plugins/Developer/Concert/ConcertApp/MultiUserClient/Source/MultiUserClient/Private/Replication/Misc/PropertySelection/UserPropertySelector.cpp b/Engine/Plugins/Developer/Concert/ConcertApp/MultiUserClient/Source/MultiUserClient/Private/Replication/Misc/PropertySelection/UserPropertySelector.cpp
--- a/Engine/Plugins/Developer/Concert/ConcertApp/MultiUserClient/Source/MultiUserClient/Private/Replication/Misc/PropertySelection/UserPropertySelector.cpp
+++ b/Engine/Plugins/Developer/Concert/ConcertApp/MultiUserClient/Source/MultiUserClient/Private/Replication/Misc/PropertySelection/UserPropertySelector.cpp
@@ -31,7 +31,7 @@ namespace UE::MultiUserClient::Replication
 		ClientManager.OnPostRemoteClientAdded().RemoveAll(this);
 		ClientManager.ForEachClient([this](FOnlineClient& Client)
 		{
-			Client.GetStreamSynchronizer().OnServerStateChanged().RemoveAll(this);
+			Client.GetStreamSynchronizer().OnServerStreanChanged().RemoveAll(this);
 			return EBreakBehavior::Continue;
 		});
 		FCoreUObjectDelegates::OnObjectTransacted.RemoveAll(this);
@@ -73,7 +73,7 @@ namespace UE::MultiUserClient::Replication
 	{
 		IClientStreamSynchronizer& StreamSynchronizer = Client.GetStreamSynchronizer();
 		TrackProperties(StreamSynchronizer.GetServerState());
-		StreamSynchronizer.OnServerStateChanged().AddRaw(this, &FUserPropertySelector::OnServerStateChanged, Client.GetEndpointId());
+		StreamSynchronizer.OnServerStreanChanged().AddRaw(this, &FUserPropertySelector::OnServerStateChanged, Client.GetEndpointId());
 	}
 
 	void FUserPropertySelector::OnServerStateChanged(const FGuid ClientId)
diff --git a/Engine/Plugins/Developer/Concert/ConcertApp/MultiUserClient/Source/MultiUserClient/Private/Replication/MultiUserReplicationManager.cpp b/Engine/Plugins/Developer/Concert/ConcertApp/MultiUserClient/Source/MultiUserClient/Private/Replication/MultiUserReplicationManager.cpp
--- a/Engine/Plugins/Developer/Concert/ConcertApp/MultiUserClient/Source/MultiUserClient/Private/Replication/MultiUserReplicationManager.cpp
+++ b/Engine/Plugins/Developer/Concert/ConcertApp/MultiUserClient/Source/MultiUserClient/Private/Replication/MultiUserReplicationManager.cpp
@@ -161,8 +161,8 @@ namespace UE::MultiUserClient::Replication
 
 	void FMultiUserReplicationManager::SetupClientDelegates(FOnlineClient& InClient) const
 	{
-		InClient.GetStreamSynchronizer().OnServerStateChanged().AddRaw(this, &FMultiUserReplicationManager::OnClientStreamServerStateChanged, InClient.GetEndpointId());
-		InClient.GetAuthoritySynchronizer().OnServerStateChanged().AddRaw(this, &FMultiUserReplicationManager::OnClientAuthorityServerStateChanged, InClient.GetEndpointId());
+		InClient.GetStreamSynchronizer().OnServerStreanChanged().AddRaw(this, &FMultiUserReplicationManager::OnClientStreamServerStateChanged, InClient.GetEndpointId());
+		InClient.GetAuthoritySynchronizer().OnServerAuthorityChanged().AddRaw(this, &FMultiUserReplicationManager::OnClientAuthorityServerStateChanged, InClient.GetEndpointId());
 	}
 
 	const FConcertObjectReplicationMap* FMultiUserReplicationManager::FindReplicationMapForClient(const FGuid& ClientId) const
diff --git a/Engine/Plugins/Developer/Concert/ConcertApp/MultiUserClient/Source/MultiUserClient/Private/Replication/Stream/IClientStreamSynchronizer.h b/Engine/Plugins/Developer/Concert/ConcertApp/MultiUserClient/Source/MultiUserClient/Private/Replication/Stream/IClientStreamSynchronizer.h
--- a/Engine/Plugins/Developer/Concert/ConcertApp/MultiUserClient/Source/MultiUserClient/Private/Replication/Stream/IClientStreamSynchronizer.h
+++ b/Engine/Plugins/Developer/Concert/ConcertApp/MultiUserClient/Source/MultiUserClient/Private/Replication/Stream/IClientStreamSynchronizer.h
@@ -31,7 +31,7 @@ namespace UE::MultiUserClient::Replication
 
 		DECLARE_MULTICAST_DELEGATE(FOnServerStateChanged);
 		/** @return Event executed when the result of GetServerState has been updated. */
-		virtual FOnServerStateChanged& OnServerStateChanged() = 0;
+		virtual FOnServerStateChanged& OnServerStreanChanged() = 0;
 		
 		virtual ~IClientStreamSynchronizer() = default;
 	};
diff --git a/Engine/Plugins/Developer/Concert/ConcertApp/MultiUserClient/Source/MultiUserClient/Private/Replication/Stream/StreamChangeTracker.cpp b/Engine/Plugins/Developer/Concert/ConcertApp/MultiUserClient/Source/MultiUserClient/Private/Replication/Stream/StreamChangeTracker.cpp
--- a/Engine/Plugins/Developer/Concert/ConcertApp/MultiUserClient/Source/MultiUserClient/Private/Replication/Stream/StreamChangeTracker.cpp
+++ b/Engine/Plugins/Developer/Concert/ConcertApp/MultiUserClient/Source/MultiUserClient/Private/Replication/Stream/StreamChangeTracker.cpp
@@ -21,12 +21,12 @@ namespace UE::MultiUserClient::Replication
 		, OnModifyReplicationMapDelegate(MoveTemp(InOnModifyReplicationMapDelegate))
 	{
 		check(InStreamWithInProgressChangesAttribute.IsBound() || InStreamWithInProgressChangesAttribute.IsSet());
-		StreamSynchronizer.OnServerStateChanged().AddRaw(this, &FStreamChangeTracker::RefreshChangesCache);
+		StreamSynchronizer.OnServerStreanChanged().AddRaw(this, &FStreamChangeTracker::RefreshChangesCache);
 	}
 
 	FStreamChangeTracker::~FStreamChangeTracker()
 	{
-		StreamSynchronizer.OnServerStateChanged().RemoveAll(this);
+		StreamSynchronizer.OnServerStreanChanged().RemoveAll(this);
 	}
 
 	void FStreamChangeTracker::RefreshChangesCache()
diff --git a/Engine/Plugins/Developer/Concert/ConcertApp/MultiUserClient/Source/MultiUserClient/Private/Replication/Stream/StreamSynchronizer_LocalClient.h b/Engine/Plugins/Developer/Concert/ConcertApp/MultiUserClient/Source/MultiUserClient/Private/Replication/Stream/StreamSynchronizer_LocalClient.h
--- a/Engine/Plugins/Developer/Concert/ConcertApp/MultiUserClient/Source/MultiUserClient/Private/Replication/Stream/StreamSynchronizer_LocalClient.h
+++ b/Engine/Plugins/Developer/Concert/ConcertApp/MultiUserClient/Source/MultiUserClient/Private/Replication/Stream/StreamSynchronizer_LocalClient.h
@@ -23,7 +23,7 @@ namespace UE::MultiUserClient::Replication
 		virtual FGuid GetStreamId() const override { return LocalClientStreamId; }
 		virtual const FConcertObjectReplicationMap& GetServerState() const override;
 		virtual const FConcertStreamFrequencySettings& GetFrequencySettings() const override;
-		virtual FOnServerStateChanged& OnServerStateChanged() override { return OnServerStateChangedDelegate; }
+		virtual FOnServerStateChanged& OnServerStreanChanged() override { return OnServerStateChangedDelegate; }
 		//~ End IClientStreamSynchronizer Interface
 
 	private:
diff --git a/Engine/Plugins/Developer/Concert/ConcertApp/MultiUserClient/Source/MultiUserClient/Private/Replication/Stream/StreamSynchronizer_RemoteClient.h b/Engine/Plugins/Developer/Concert/ConcertApp/MultiUserClient/Source/MultiUserClient/Private/Replication/Stream/StreamSynchronizer_RemoteClient.h
--- a/Engine/Plugins/Developer/Concert/ConcertApp/MultiUserClient/Source/MultiUserClient/Private/Replication/Stream/StreamSynchronizer_RemoteClient.h
+++ b/Engine/Plugins/Developer/Concert/ConcertApp/MultiUserClient/Source/MultiUserClient/Private/Replication/Stream/StreamSynchronizer_RemoteClient.h
@@ -25,7 +25,7 @@ namespace UE::MultiUserClient::Replication
 		virtual FGuid GetStreamId() const override;
 		virtual const FConcertObjectReplicationMap& GetServerState() const override { return LastKnownServerState.ReplicationMap; }
 		virtual const FConcertStreamFrequencySettings& GetFrequencySettings() const override { return LastKnownServerState.FrequencySettings; }
-		virtual FOnServerStateChanged& OnServerStateChanged() override { return OnServerStateChangedDelegate; }
+		virtual FOnServerStateChanged& OnServerStreanChanged() override { return OnServerStateChangedDelegate; }
 		//~ End IClientStreamSynchronizer Interface
 
 	private:
