Fix duplicate loca key

diff --git a/Engine/Plugins/Developer/Concert/ConcertApp/MultiUserClient/Source/MultiUserClient/Private/Widgets/ActiveSession/Replication/Client/Context/ContextMenuUtils_Frequency.cpp b/Engine/Plugins/Developer/Concert/ConcertApp/MultiUserClient/Source/MultiUserClient/Private/Widgets/ActiveSession/Replication/Client/Context/ContextMenuUtils_Frequency.cpp
--- a/Engine/Plugins/Developer/Concert/ConcertApp/MultiUserClient/Source/MultiUserClient/Private/Widgets/ActiveSession/Replication/Client/Context/ContextMenuUtils_Frequency.cpp
+++ b/Engine/Plugins/Developer/Concert/ConcertApp/MultiUserClient/Source/MultiUserClient/Private/Widgets/ActiveSession/Replication/Client/Context/ContextMenuUtils_Frequency.cpp
@@ -246,7 +246,7 @@ namespace UE::MultiUserClient::Replication::ContextMenuUtils
 			)
 		{
 			// Change to Realtime / Specified Rate
-			MenuBuilder.BeginSection(NAME_None, LOCTEXT("Frequency", "Mode"));
+			MenuBuilder.BeginSection(NAME_None, LOCTEXT("Frequency.Mode", "Mode"));
 			AddReplicationModeSubMenu(MenuBuilder, ContextObject, GetClientsAttribute, InClientManager);
 			MenuBuilder.EndSection();
 
@@ -279,7 +279,7 @@ namespace UE::MultiUserClient::Replication::ContextMenuUtils
 				if (!bAddedSubmenu)
 				{
 					bAddedSubmenu = true;
-					MenuBuilder.BeginSection(NAME_None, LOCTEXT("Frequency", "Change Frequency"));
+					MenuBuilder.BeginSection(NAME_None, LOCTEXT("Frequency.Change", "Change Frequency"));
 				}
 			};
 			ON_SCOPE_EXIT{ if(bAddedSubmenu){ MenuBuilder.EndSection(); } };
