Make Multi User replication preset independent from UMultiUserReplicationStream.
Moved frequency settings out from UMultiUserReplicationStream because it was not being used.

diff --git a/Engine/Plugins/Developer/Concert/ConcertApp/MultiUserClient/Source/MultiUserClient/Private/Replication/Client/Online/OnlineClient.cpp b/Engine/Plugins/Developer/Concert/ConcertApp/MultiUserClient/Source/MultiUserClient/Private/Replication/Client/Online/OnlineClient.cpp
--- a/Engine/Plugins/Developer/Concert/ConcertApp/MultiUserClient/Source/MultiUserClient/Private/Replication/Client/Online/OnlineClient.cpp
+++ b/Engine/Plugins/Developer/Concert/ConcertApp/MultiUserClient/Source/MultiUserClient/Private/Replication/Client/Online/OnlineClient.cpp
@@ -2,7 +2,8 @@
 
 #include "OnlineClient.h"
 
-#include "Assets/MultiUserReplicationClientContent.h"
+#include "Assets/MultiUserReplicationStream.h"
+#include "Replication/ClientReplicationWidgetFactories.h"
 #include "Replication/Editor/Model/IEditableReplicationStreamModel.h"
 #include "Replication/ReplicationWidgetFactories.h"
 #include "Replication/Stream/Discovery/MultiUserStreamExtender.h"
@@ -11,7 +12,6 @@
 #include "Settings/MultiUserReplicationSettings.h"
 
 #include "Misc/CoreDelegates.h"
-#include "Replication/ClientReplicationWidgetFactories.h"
 
 namespace UE::MultiUserClient::Replication
 {
diff --git a/Engine/Plugins/Developer/Concert/ConcertApp/MultiUserClient/Source/MultiUserClient/Private/Replication/Client/Online/OnlineClient.h b/Engine/Plugins/Developer/Concert/ConcertApp/MultiUserClient/Source/MultiUserClient/Private/Replication/Client/Online/OnlineClient.h
--- a/Engine/Plugins/Developer/Concert/ConcertApp/MultiUserClient/Source/MultiUserClient/Private/Replication/Client/Online/OnlineClient.h
+++ b/Engine/Plugins/Developer/Concert/ConcertApp/MultiUserClient/Source/MultiUserClient/Private/Replication/Client/Online/OnlineClient.h
@@ -2,7 +2,6 @@
 
 #pragma once
 
-#include "Assets/MultiUserReplicationClientContent.h"
 #include "Replication/Authority/AuthorityChangeTracker.h"
 #include "Replication/Authority/IClientAuthoritySynchronizer.h"
 #include "Replication/Editor/UnrealEditor/ModifyObjectInLevelHandler.h"
@@ -18,8 +17,8 @@
 #include "Templates/SharedPointer.h"
 #include "Templates/UnrealTemplate.h"
 
+class UMultiUserReplicationStream;
 class IConcertClient;
-class UMultiUserReplicationClientContent;
 
 namespace UE::ConcertSharedSlate
 {
diff --git a/Engine/Plugins/Developer/Concert/ConcertApp/MultiUserClient/Source/MultiUserClient/Private/Replication/Client/Online/RemoteClient.cpp b/Engine/Plugins/Developer/Concert/ConcertApp/MultiUserClient/Source/MultiUserClient/Private/Replication/Client/Online/RemoteClient.cpp
--- a/Engine/Plugins/Developer/Concert/ConcertApp/MultiUserClient/Source/MultiUserClient/Private/Replication/Client/Online/RemoteClient.cpp
+++ b/Engine/Plugins/Developer/Concert/ConcertApp/MultiUserClient/Source/MultiUserClient/Private/Replication/Client/Online/RemoteClient.cpp
@@ -3,7 +3,7 @@
 #include "RemoteClient.h"
 
 #include "IConcertClient.h"
-#include "Assets/MultiUserReplicationClientContent.h"
+#include "Assets/MultiUserReplicationClientPreset.h"
 #include "Replication/Authority/AuthoritySynchronizer_RemoteClient.h"
 #include "Replication/Stream/StreamSynchronizer_RemoteClient.h"
 #include "Replication/Submission/Remote/SubmissionWorkflow_RemoteClient.h"
diff --git a/Engine/Plugins/Developer/Concert/ConcertApp/MultiUserClient/Source/MultiUserClient/Private/Replication/Preset/PresetManager.cpp b/Engine/Plugins/Developer/Concert/ConcertApp/MultiUserClient/Source/MultiUserClient/Private/Replication/Preset/PresetManager.cpp
--- a/Engine/Plugins/Developer/Concert/ConcertApp/MultiUserClient/Source/MultiUserClient/Private/Replication/Preset/PresetManager.cpp
+++ b/Engine/Plugins/Developer/Concert/ConcertApp/MultiUserClient/Source/MultiUserClient/Private/Replication/Preset/PresetManager.cpp
@@ -2,7 +2,9 @@
 
 #include "PresetManager.h"
 
+#include "Assets/MultiUserReplicationClientPreset.h"
 #include "Assets/MultiUserReplicationSessionPreset.h"
+#include "Assets/MultiUserReplicationStream.h"
 #include "ConcertLogGlobal.h"
 #include "IConcertSyncClient.h"
 #include "Replication/Client/Online/OnlineClientManager.h"
@@ -39,7 +41,7 @@ namespace UE::MultiUserClient::Replication
 		{
 			const auto AddClient = [&Preset, &Request, bClearUnreferencedClients](const FConcertSessionClientInfo& ClientSessionInfo)
 			{
-				const UMultiUserReplicationClientContent* ClientSessionContent = Preset.GetClientContent(ClientSessionInfo.ClientInfo);
+				const FMultiUserReplicationClientPreset* ClientSessionContent = Preset.GetClientContent(ClientSessionInfo.ClientInfo);
 				if (!ClientSessionContent)
 				{
 					if (bClearUnreferencedClients)
@@ -49,19 +51,19 @@ namespace UE::MultiUserClient::Replication
 					return;
 				}
 
-				const UMultiUserReplicationStream* SavedStream = ClientSessionContent->Stream;
-				if (SavedStream->ReplicationMap.IsEmpty())
+				const FConcertObjectReplicationMap& ReplicationMap = ClientSessionContent->ReplicationMap;
+				if (ReplicationMap.IsEmpty())
 				{
 					// This causes the client's content to be cleared.
 					Request.NewStreams.Add(ClientSessionInfo.ClientEndpointId, {});
 					return;
 				}
 				
-				const FGuid& StreamId = SavedStream->StreamId;
-				FConcertReplicationStream Stream { { .Identifier = StreamId, .ReplicationMap = SavedStream->ReplicationMap } };
+				const FGuid& StreamId = MultiUserStreamID;
+				FConcertReplicationStream Stream { { .Identifier = StreamId, .ReplicationMap = ReplicationMap } };
 				// Empty objects will be rejected by the server
 				RemoveEmptyObjectsFromRequest(Stream);
-				Stream.BaseDescription.FrequencySettings = ClientSessionContent->Stream->FrequencySettings;
+				Stream.BaseDescription.FrequencySettings = ClientSessionContent->FrequencySettings;
 				Request.NewStreams.Add(ClientSessionInfo.ClientEndpointId, { TArray{ Stream } });
 				
 				// MU automatically requests authority when it adds an object.
@@ -304,7 +306,7 @@ namespace UE::MultiUserClient::Replication
 		{
 			const auto[Client, ClientInfo] = ClientData;
 			UMultiUserReplicationStream* CopiedClientStream = Client->GetClientStreamObject();
-			UMultiUserReplicationClientContent* TargetClientPreset = Preset->AddClientIfUnique(ClientInfo, MultiUserStreamID);
+			FMultiUserReplicationClientPreset* TargetClientPreset = Preset->AddClientIfUnique(ClientInfo, MultiUserStreamID);
 			if (!TargetClientPreset)
 			{
 				UE_LOG(LogConcert, Warning,
@@ -315,9 +317,8 @@ namespace UE::MultiUserClient::Replication
 				continue;
 			}
 
-			TargetClientPreset->Stream->Copy(*CopiedClientStream);
-			// TODO UE-219834: Once UMultiUserReplicationStream::FrequencySettings reflect the server state, this can be removed.
-			TargetClientPreset->Stream->FrequencySettings = Client->GetStreamSynchronizer().GetFrequencySettings();
+			TargetClientPreset->ReplicationMap = CopiedClientStream->ReplicationMap;
+			TargetClientPreset->FrequencySettings = Client->GetStreamSynchronizer().GetFrequencySettings();
 		}
 
 		Preset->SetMuteContent(
diff --git a/Engine/Plugins/Developer/Concert/ConcertApp/MultiUserClient/Source/MultiUserClient/Private/Replication/Stream/StreamSynchronizer_RemoteClient.cpp b/Engine/Plugins/Developer/Concert/ConcertApp/MultiUserClient/Source/MultiUserClient/Private/Replication/Stream/StreamSynchronizer_RemoteClient.cpp
--- a/Engine/Plugins/Developer/Concert/ConcertApp/MultiUserClient/Source/MultiUserClient/Private/Replication/Stream/StreamSynchronizer_RemoteClient.cpp
+++ b/Engine/Plugins/Developer/Concert/ConcertApp/MultiUserClient/Source/MultiUserClient/Private/Replication/Stream/StreamSynchronizer_RemoteClient.cpp
@@ -2,7 +2,7 @@
 
 #include "StreamSynchronizer_RemoteClient.h"
 
-#include "Assets/MultiUserReplicationClientContent.h"
+#include "Assets/MultiUserReplicationClientPreset.h"
 #include "MultiUserStreamId.h"
 #include "Replication/Misc/Query/StreamAndAuthorityQueryService.h"
 
diff --git a/Engine/Plugins/Developer/Concert/ConcertApp/MultiUserClient/Source/MultiUserReplicationEditor/Internal/Assets/MultiUserReplicationClientContent.h b/Engine/Plugins/Developer/Concert/ConcertApp/MultiUserClient/Source/MultiUserReplicationEditor/Internal/Assets/MultiUserReplicationClientContent.h
deleted file mode 100644
diff --git a/Engine/Plugins/Developer/Concert/ConcertApp/MultiUserClient/Source/MultiUserReplicationEditor/Internal/Assets/MultiUserReplicationClientPreset.h b/Engine/Plugins/Developer/Concert/ConcertApp/MultiUserClient/Source/MultiUserReplicationEditor/Internal/Assets/MultiUserReplicationClientPreset.h
new file mode 100644
--- /dev/null
+++ b/Engine/Plugins/Developer/Concert/ConcertApp/MultiUserClient/Source/MultiUserReplicationEditor/Internal/Assets/MultiUserReplicationClientPreset.h
@@ -0,0 +1,36 @@
+// Copyright Epic Games, Inc. All Rights Reserved.
+
+#pragma once
+
+#include "UObject/Object.h"
+#include "Replication/Data/ReplicationStream.h"
+#include "MultiUserReplicationClientPreset.generated.h"
+
+/** Stores info about a client's content in a preset. */
+USTRUCT()
+struct FMultiUserReplicationClientPreset
+{
+	GENERATED_BODY()
+
+	/** The objects this stream will modify. */
+	UPROPERTY()
+	FConcertObjectReplicationMap ReplicationMap;
+	
+	/** The frequency setting the stream has. */
+	UPROPERTY()
+	FConcertStreamFrequencySettings FrequencySettings;
+
+	/** The FConcertClientInfo::DisplayName of the client. */
+	UPROPERTY()
+	FString DisplayName;
+	/** The FConcertClientInfo::DeviceName of the client. */
+	UPROPERTY()
+	FString DeviceName;
+
+	FMultiUserReplicationClientPreset() = default;
+
+	FMultiUserReplicationClientPreset(const FString& DisplayName, const FString& DeviceName)
+		: DisplayName(DisplayName)
+		, DeviceName(DeviceName)
+	{}
+};
diff --git a/Engine/Plugins/Developer/Concert/ConcertApp/MultiUserClient/Source/MultiUserReplicationEditor/Internal/Assets/MultiUserReplicationSessionPreset.h b/Engine/Plugins/Developer/Concert/ConcertApp/MultiUserClient/Source/MultiUserReplicationEditor/Internal/Assets/MultiUserReplicationSessionPreset.h
--- a/Engine/Plugins/Developer/Concert/ConcertApp/MultiUserClient/Source/MultiUserReplicationEditor/Internal/Assets/MultiUserReplicationSessionPreset.h
+++ b/Engine/Plugins/Developer/Concert/ConcertApp/MultiUserClient/Source/MultiUserReplicationEditor/Internal/Assets/MultiUserReplicationSessionPreset.h
@@ -2,12 +2,13 @@
 
 #pragma once
 
-#include "MultiUserReplicationClientContent.h"
-#include "Replication/Messages/Muting.h"
 #include "UObject/Object.h"
+
+#include "Replication/Messages/Muting.h"
 #include "MultiUserReplicationSessionPreset.generated.h"
 
 struct FConcertClientInfo;
+struct FMultiUserReplicationClientPreset;
 
 USTRUCT()
 struct FMultiUserMuteSessionContent
@@ -42,9 +43,9 @@ public:
 	/********** Clients **********/
 
 	/** @return The client preset that matches ClientInfo.DisplayName. If there are multiple, returns the one that matches ClientInfo.DeviceName, as well. */
-	UMultiUserReplicationClientContent* GetClientContent(const FConcertClientInfo& ClientInfo) const;
+	const FMultiUserReplicationClientPreset* GetClientContent(const FConcertClientInfo& ClientInfo) const;
 	/** @return The client preset that matches both the display and device name. */
-	UMultiUserReplicationClientContent* GetExactClientContent(const FConcertClientInfo& ClientInfo) const;
+	const FMultiUserReplicationClientPreset* GetExactClientContent(const FConcertClientInfo& ClientInfo) const;
 	
 	/** @return Whether a client that matches ClientInfo.DisplayName. */
 	bool ContainsClient(const FConcertClientInfo& ClientInfo) const { return GetClientContent(ClientInfo) != nullptr; }
@@ -52,9 +53,9 @@ public:
 	bool ContainsExactClient(const FConcertClientInfo& ClientInfo) const { return GetExactClientContent(ClientInfo) != nullptr; }
 	
 	/** Adds a client to the preset if it's not already present. */
-	UMultiUserReplicationClientContent* AddClientIfUnique(const FConcertClientInfo& ClientInfo, const FGuid& StreamId);
+	FMultiUserReplicationClientPreset* AddClientIfUnique(const FConcertClientInfo& ClientInfo, const FGuid& StreamId);
 	
-	const TArray<TObjectPtr<UMultiUserReplicationClientContent>>& GetClientPresets() const { return ClientPresets; }
+	const TArray<FMultiUserReplicationClientPreset>& GetClientPresets() const { return ClientPresets; }
 	
 	/********** Muting **********/
 	
@@ -63,8 +64,8 @@ public:
 	
 private:
 	
-	UPROPERTY(Instanced)
-	TArray<TObjectPtr<UMultiUserReplicationClientContent>> ClientPresets;
+	UPROPERTY()
+	TArray<FMultiUserReplicationClientPreset> ClientPresets;
 
 	UPROPERTY()
 	FMultiUserMuteSessionContent MuteContent;
diff --git a/Engine/Plugins/Developer/Concert/ConcertApp/MultiUserClient/Source/MultiUserReplicationEditor/Internal/Assets/MultiUserReplicationStream.h b/Engine/Plugins/Developer/Concert/ConcertApp/MultiUserClient/Source/MultiUserReplicationEditor/Internal/Assets/MultiUserReplicationStream.h
--- a/Engine/Plugins/Developer/Concert/ConcertApp/MultiUserClient/Source/MultiUserReplicationEditor/Internal/Assets/MultiUserReplicationStream.h
+++ b/Engine/Plugins/Developer/Concert/ConcertApp/MultiUserClient/Source/MultiUserReplicationEditor/Internal/Assets/MultiUserReplicationStream.h
@@ -21,16 +21,6 @@ public:
 	/** The objects this stream will modify. */
 	UPROPERTY()
 	FConcertObjectReplicationMap ReplicationMap;
-	
-	/**
-	 * The frequency setting the stream has.
-	 *
-	 * TODO UE-219834:
-	 * Currently, this is only written to by the preset system.
-	 * In the future, FFrequencyChangeTracker could / should use this so changes to frequencies can be transacted, as well.
-	 */
-	UPROPERTY()
-	FConcertStreamFrequencySettings FrequencySettings;
 
 	UMultiUserReplicationStream();
 
@@ -45,7 +35,4 @@ public:
 			return WeakThis.IsValid() ? &WeakThis->ReplicationMap : nullptr;
 		});
 	}
-	
-	/** Copies the stream content of OtherStream. */
-	void Copy(UMultiUserReplicationStream& OtherStream);
 };
diff --git a/Engine/Plugins/Developer/Concert/ConcertApp/MultiUserClient/Source/MultiUserReplicationEditor/Private/Assets/MultiUserReplicationClientContent.cpp b/Engine/Plugins/Developer/Concert/ConcertApp/MultiUserClient/Source/MultiUserReplicationEditor/Private/Assets/MultiUserReplicationClientContent.cpp
deleted file mode 100644
diff --git a/Engine/Plugins/Developer/Concert/ConcertApp/MultiUserClient/Source/MultiUserReplicationEditor/Private/Assets/MultiUserReplicationSessionPreset.cpp b/Engine/Plugins/Developer/Concert/ConcertApp/MultiUserClient/Source/MultiUserReplicationEditor/Private/Assets/MultiUserReplicationSessionPreset.cpp
--- a/Engine/Plugins/Developer/Concert/ConcertApp/MultiUserClient/Source/MultiUserReplicationEditor/Private/Assets/MultiUserReplicationSessionPreset.cpp
+++ b/Engine/Plugins/Developer/Concert/ConcertApp/MultiUserClient/Source/MultiUserReplicationEditor/Private/Assets/MultiUserReplicationSessionPreset.cpp
@@ -2,59 +2,55 @@
 
 #include "Assets/MultiUserReplicationSessionPreset.h"
 
+#include "Assets/MultiUserReplicationClientPreset.h"
 #include "ConcertLogGlobal.h"
 #include "ConcertMessageData.h"
 
-UMultiUserReplicationClientContent* UMultiUserReplicationSessionPreset::GetClientContent(const FConcertClientInfo& ClientInfo) const
+const FMultiUserReplicationClientPreset* UMultiUserReplicationSessionPreset::GetClientContent(const FConcertClientInfo& ClientInfo) const
 {
-	const auto IsMatch = [&ClientInfo](UMultiUserReplicationClientContent* Content)
+	const auto IsMatch = [&ClientInfo](const FMultiUserReplicationClientPreset& Content)
 	{
-		return Content && Content->DisplayName == ClientInfo.DisplayName;
+		return Content.DisplayName == ClientInfo.DisplayName;
 	};
-	const auto IsPerfectMatch = [&ClientInfo, &IsMatch](UMultiUserReplicationClientContent* Content)
+	const auto IsPerfectMatch = [&ClientInfo, &IsMatch](const FMultiUserReplicationClientPreset& Content)
 	{
-		return IsMatch(Content) && Content->DeviceName == ClientInfo.DeviceName;
+		return IsMatch(Content) && Content.DeviceName == ClientInfo.DeviceName;
 	};
 	
-	UMultiUserReplicationClientContent* BestMatch = nullptr;
+	const FMultiUserReplicationClientPreset* BestMatch = nullptr;
 	bool bIsBestMatchPerfect = false;
-	for (UMultiUserReplicationClientContent* Content : ClientPresets)
+	for (const FMultiUserReplicationClientPreset& Content : ClientPresets)
 	{
 		if (IsPerfectMatch(Content))
 		{
 			UE_CLOG(bIsBestMatchPerfect, LogConcert, Warning, TEXT("Preset %s contained client (name: %s, device: %s) multiple times"), *GetPathName(), *ClientInfo.DisplayName, *ClientInfo.DeviceName);
 			bIsBestMatchPerfect = true;
-			BestMatch = Content;
+			BestMatch = &Content;
 		}
 		else if (!bIsBestMatchPerfect && IsMatch(Content))
 		{
-			BestMatch = Content;
+			BestMatch = &Content;
 		}
 	}
 	
 	return BestMatch;
 }
 
-UMultiUserReplicationClientContent* UMultiUserReplicationSessionPreset::GetExactClientContent(const FConcertClientInfo& ClientInfo) const
+const FMultiUserReplicationClientPreset* UMultiUserReplicationSessionPreset::GetExactClientContent(const FConcertClientInfo& ClientInfo) const
 {
-	const TObjectPtr<UMultiUserReplicationClientContent>* Result = ClientPresets.FindByPredicate([&ClientInfo](const TObjectPtr<UMultiUserReplicationClientContent>& Content)
+	return ClientPresets.FindByPredicate([&ClientInfo](const FMultiUserReplicationClientPreset& Content)
 	{
-		return Content && Content->DisplayName == ClientInfo.DisplayName && Content->DeviceName == ClientInfo.DeviceName;
+		return Content.DisplayName == ClientInfo.DisplayName && Content.DeviceName == ClientInfo.DeviceName;
 	});
-	return Result ? *Result : nullptr;
 }
 
-UMultiUserReplicationClientContent* UMultiUserReplicationSessionPreset::AddClientIfUnique(const FConcertClientInfo& ClientInfo, const FGuid& StreamId)
+FMultiUserReplicationClientPreset* UMultiUserReplicationSessionPreset::AddClientIfUnique(const FConcertClientInfo& ClientInfo, const FGuid& StreamId)
 {
 	if (ContainsExactClient(ClientInfo))
 	{
 		return nullptr;
 	}
 	
-	UMultiUserReplicationClientContent* Result = NewObject<UMultiUserReplicationClientContent>(this);
-	Result->DisplayName = ClientInfo.DisplayName;
-	Result->DeviceName = ClientInfo.DeviceName;
-	Result->Stream->StreamId = StreamId;
-	ClientPresets.Add(Result);
-	return Result;
+	const int32 Index = ClientPresets.Emplace(ClientInfo.DisplayName, ClientInfo.DeviceName);
+	return &ClientPresets[Index];
 }
diff --git a/Engine/Plugins/Developer/Concert/ConcertApp/MultiUserClient/Source/MultiUserReplicationEditor/Private/Assets/MultiUserReplicationStream.cpp b/Engine/Plugins/Developer/Concert/ConcertApp/MultiUserClient/Source/MultiUserReplicationEditor/Private/Assets/MultiUserReplicationStream.cpp
--- a/Engine/Plugins/Developer/Concert/ConcertApp/MultiUserClient/Source/MultiUserReplicationEditor/Private/Assets/MultiUserReplicationStream.cpp
+++ b/Engine/Plugins/Developer/Concert/ConcertApp/MultiUserClient/Source/MultiUserReplicationEditor/Private/Assets/MultiUserReplicationStream.cpp
@@ -14,12 +14,5 @@ UMultiUserReplicationStream::UMultiUserReplicationStream()
 
 FConcertReplicationStream UMultiUserReplicationStream::GenerateDescription() const
 {
-	return { StreamId, ReplicationMap, FrequencySettings };
-}
-
-void UMultiUserReplicationStream::Copy(UMultiUserReplicationStream& OtherStream)
-{
-	StreamId = OtherStream.StreamId;
-	ReplicationMap = OtherStream.ReplicationMap;
-	FrequencySettings = OtherStream.FrequencySettings;
+	return { StreamId, ReplicationMap };
 }
diff --git a/Engine/Plugins/Developer/Concert/ConcertApp/MultiUserClient/Source/MultiUserReplicationEditor/Private/StreamEditor/ReplicationSessionPresetEditorToolkit.cpp b/Engine/Plugins/Developer/Concert/ConcertApp/MultiUserClient/Source/MultiUserReplicationEditor/Private/StreamEditor/ReplicationSessionPresetEditorToolkit.cpp
--- a/Engine/Plugins/Developer/Concert/ConcertApp/MultiUserClient/Source/MultiUserReplicationEditor/Private/StreamEditor/ReplicationSessionPresetEditorToolkit.cpp
+++ b/Engine/Plugins/Developer/Concert/ConcertApp/MultiUserClient/Source/MultiUserReplicationEditor/Private/StreamEditor/ReplicationSessionPresetEditorToolkit.cpp
@@ -2,7 +2,7 @@
 
 #include "ReplicationSessionPresetEditorToolkit.h"
 
-#include "Assets/MultiUserReplicationClientContent.h"
+#include "Assets/MultiUserReplicationClientPreset.h"
 #include "Assets/MultiUserReplicationSessionPreset.h"
 #include "Replication/Editor/Model/PropertySource/SelectPropertyFromUClassModel.h"
 #include "Replication/Editor/View/IReplicationStreamEditor.h"
