Multi-User presets now associate replication settings by actor label.

Before being applied, the presets are translated using the new actor label remapping algorithm.
This allows presets to be reused in different sessions / worlds.

diff --git a/Engine/Plugins/Developer/Concert/ConcertApp/MultiUserClient/Source/MultiUserClient/Private/Replication/Preset/PresetManager.cpp b/Engine/Plugins/Developer/Concert/ConcertApp/MultiUserClient/Source/MultiUserClient/Private/Replication/Preset/PresetManager.cpp
--- a/Engine/Plugins/Developer/Concert/ConcertApp/MultiUserClient/Source/MultiUserClient/Private/Replication/Preset/PresetManager.cpp
+++ b/Engine/Plugins/Developer/Concert/ConcertApp/MultiUserClient/Source/MultiUserClient/Private/Replication/Preset/PresetManager.cpp
@@ -8,197 +8,231 @@
 #include "ConcertLogGlobal.h"
 #include "IConcertSyncClient.h"
 #include "Replication/Client/Online/OnlineClientManager.h"
+#include "Replication/Misc/ActorLabelRemappingEditor.h"
 #include "Replication/Misc/ReplicationStreamUtils.h"
 #include "Replication/Muting/MuteStateManager.h"
 #include "Replication/Stream/MultiUserStreamId.h"
 #include "Widgets/ActiveSession/Replication/Client/ClientUtils.h"
 
+#include "Engine/World.h"
 #include "FileHelpers.h"
 
-namespace UE::MultiUserClient::Replication
+namespace UE::MultiUserClient::Replication::Private
 {
-	namespace Private
+	static void RemoveEmptyObjectsFromRequest(FConcertReplicationStream& Stream)
 	{
-		static void RemoveEmptyObjectsFromRequest(FConcertReplicationStream& Stream)
+		TMap<FSoftObjectPath, FConcertReplicatedObjectInfo>& ReplicationMap = Stream.BaseDescription.ReplicationMap.ReplicatedObjects;
+		for (auto It = ReplicationMap.CreateIterator(); It; ++It)
 		{
-			TMap<FSoftObjectPath, FConcertReplicatedObjectInfo>& ReplicationMap = Stream.BaseDescription.ReplicationMap.ReplicatedObjects;
-			for (auto It = ReplicationMap.CreateIterator(); It; ++It)
+			const bool bIsEmpty = It->Value.PropertySelection.ReplicatedProperties.IsEmpty();
+			if (bIsEmpty)
 			{
-				const bool bIsEmpty = It->Value.PropertySelection.ReplicatedProperties.IsEmpty();
-				if (bIsEmpty)
-				{
-					It.RemoveCurrent();
-				}
+				It.RemoveCurrent();
 			}
 		}
+	}
+
+	static void ProcessNonEmptyRequest(const FConcertSessionClientInfo& ClientSessionInfo, FConcertReplication_PutState_Request& Request, const FMultiUserReplicationClientPreset* ClientSessionContent, const FConcertObjectReplicationMap& OriginalReplicationMap)
+	{
+		const FGuid& StreamId = MultiUserStreamID;
+		FConcertReplicationStream Stream { { .Identifier = StreamId, .ReplicationMap = OriginalReplicationMap } };
+		// Empty objects will be rejected by the server
+		RemoveEmptyObjectsFromRequest(Stream);
+		Stream.BaseDescription.FrequencySettings = ClientSessionContent->FrequencySettings;
+		Request.NewStreams.Add(ClientSessionInfo.ClientEndpointId, { TArray{ Stream } });
+			
+		// MU automatically requests authority when it adds an object.
+		// We'll assume that that authority was granted when the preset was created - if it actually was not, our request may fail due to overlapping authority.
+		TArray<FConcertObjectInStreamID>& OwnedObjects = Request.NewAuthorityState.Add(ClientSessionInfo.ClientEndpointId).Objects;
+		Algo::Transform(Stream.BaseDescription.ReplicationMap.ReplicatedObjects, OwnedObjects, [&StreamId](const TPair<FSoftObjectPath, FConcertReplicatedObjectInfo>& Pair)
+		{
+			return FConcertObjectInStreamID{ StreamId, Pair.Key };
+		});
+	}
 
-		static void FillStreamAndAuthorityRequest(
-			FConcertReplication_PutState_Request& Request,
-			const UMultiUserReplicationSessionPreset& Preset,
-			const IConcertClientSession& Session,
-			bool bClearUnreferencedClients
-			)
+	static void AddClientToRequest(
+		const FConcertSessionClientInfo& ClientSessionInfo,
+		FConcertReplication_PutState_Request& Request,
+		const UMultiUserReplicationSessionPreset& Preset,
+		bool bClearUnreferencedClients
+	)
+	{
+		const auto AddRequestToClearClient = [&Request, &ClientSessionInfo]()
+		{
+			Request.NewStreams.Add(ClientSessionInfo.ClientEndpointId, {});
+		};
+		
+		const FMultiUserReplicationClientPreset* ClientSessionContent = Preset.GetClientContent(ClientSessionInfo.ClientInfo);
+		if (!ClientSessionContent)
 		{
-			const auto AddClient = [&Preset, &Request, bClearUnreferencedClients](const FConcertSessionClientInfo& ClientSessionInfo)
+			if (bClearUnreferencedClients)
 			{
-				const FMultiUserReplicationClientPreset* ClientSessionContent = Preset.GetClientContent(ClientSessionInfo.ClientInfo);
-				if (!ClientSessionContent)
-				{
-					if (bClearUnreferencedClients)
-					{
-						Request.NewStreams.Add(ClientSessionInfo.ClientEndpointId, {});
-					}
-					return;
-				}
+				AddRequestToClearClient();
+			}
+			return;
+		}
 
-				const FConcertObjectReplicationMap& ReplicationMap = ClientSessionContent->ReplicationMap;
-				if (ReplicationMap.IsEmpty())
-				{
-					// This causes the client's content to be cleared.
-					Request.NewStreams.Add(ClientSessionInfo.ClientEndpointId, {});
-					return;
-				}
-				
-				const FGuid& StreamId = MultiUserStreamID;
-				FConcertReplicationStream Stream { { .Identifier = StreamId, .ReplicationMap = ReplicationMap } };
-				// Empty objects will be rejected by the server
-				RemoveEmptyObjectsFromRequest(Stream);
-				Stream.BaseDescription.FrequencySettings = ClientSessionContent->FrequencySettings;
-				Request.NewStreams.Add(ClientSessionInfo.ClientEndpointId, { TArray{ Stream } });
-				
-				// MU automatically requests authority when it adds an object.
-				// We'll assume that that authority was granted when the preset was created - if it actually was not, our request may fail due to overlapping authority.
-				TArray<FConcertObjectInStreamID>& OwnedObjects = Request.NewAuthorityState.Add(ClientSessionInfo.ClientEndpointId).Objects;
-				Algo::Transform(Stream.BaseDescription.ReplicationMap.ReplicatedObjects, OwnedObjects, [&StreamId](const TPair<FSoftObjectPath, FConcertReplicatedObjectInfo>& Pair)
-				{
-					return FConcertObjectInStreamID{ StreamId, Pair.Key };
-				});
-			};
-			AddClient({ Session.GetSessionClientEndpointId(), Session.GetLocalClientInfo() });
+		const FConcertObjectReplicationMap& OriginalReplicationMap = ClientSessionContent->ReplicationMap;
+		if (OriginalReplicationMap.IsEmpty())
+		{
+			AddRequestToClearClient();
+			return;
+		}
+
+		check(GWorld);
+		const FConcertObjectReplicationMap TranslatedReplicationMap = ConcertSyncCore::RemapReplicationMap(
+			OriginalReplicationMap,
+			ClientSessionContent->ActorLabelRemappingData,
+			*GWorld
+			);
+		if (TranslatedReplicationMap.IsEmpty())
+		{
+			AddRequestToClearClient();
+		}
+		else
+		{
+			ProcessNonEmptyRequest(ClientSessionInfo, Request, ClientSessionContent, TranslatedReplicationMap);
+		}
+	}
+
+	static void FillStreamAndAuthorityRequest(
+		FConcertReplication_PutState_Request& Request,
+		const UMultiUserReplicationSessionPreset& Preset,
+		const IConcertClientSession& Session,
+		bool bClearUnreferencedClients
+		)
+	{
+		// GWorld is required by AddClientToRequest
+		if (ensure(GWorld))
+		{
+			AddClientToRequest({ Session.GetSessionClientEndpointId(), Session.GetLocalClientInfo() }, Request, Preset, bClearUnreferencedClients);
 			for (const FConcertSessionClientInfo& ClientSessionInfo : Session.GetSessionClients())
 			{
-				AddClient(ClientSessionInfo);
+				AddClientToRequest(ClientSessionInfo, Request, Preset, bClearUnreferencedClients);
 			}
 		}
+	}
 
-		static void FillMuteStateRequest(
-			FConcertReplication_PutState_Request& Request,
-			const UMultiUserReplicationSessionPreset& Preset,
-			const IConcertClientSession& Session
-			)
+	static void FillMuteStateRequest(
+		FConcertReplication_PutState_Request& Request,
+		const UMultiUserReplicationSessionPreset& Preset,
+		const IConcertClientSession& Session
+		)
+	{
+		FConcertReplication_ChangeMuteState_Request& MuteRequest = Request.MuteChange;
+		MuteRequest.Flags = EConcertReplicationMuteRequestFlags::ClearMuteState;
+		
+		const auto IsReferencedByConnectedClient = [&Request, &Session](const FSoftObjectPath& ObjectPath)
 		{
-			FConcertReplication_ChangeMuteState_Request& MuteRequest = Request.MuteChange;
-			MuteRequest.Flags = EConcertReplicationMuteRequestFlags::ClearMuteState;
-			
-			const auto IsReferencedByConnectedClient = [&Request, &Session](const FSoftObjectPath& ObjectPath)
+			return Algo::AnyOf(Request.NewStreams, [&Session, &ObjectPath](const TPair<FGuid, FConcertReplicationStreamArray>& ClientContent)
 			{
-				return Algo::AnyOf(Request.NewStreams, [&Session, &ObjectPath](const TPair<FGuid, FConcertReplicationStreamArray>& ClientContent)
-				{
-					const FGuid& EndpointId = ClientContent.Key;
-					FConcertSessionClientInfo Dummy;
-					
-					const bool bIsConnected = Session.GetSessionClientEndpointId() == EndpointId || Session.FindSessionClient(EndpointId, Dummy);
-					// Case: User muted Floor but only Floor.StaticMeshComponent0 is replicated. Hence, also look for child objects being referenced.
-					const bool bIsReferenced = ConcertSyncCore::IsObjectOrChildReferenced(ClientContent.Value.Streams, ObjectPath);
-					
-					return bIsConnected && bIsReferenced;
-				});
-			};
+				const FGuid& EndpointId = ClientContent.Key;
+				FConcertSessionClientInfo Dummy;
+				
+				const bool bIsConnected = Session.GetSessionClientEndpointId() == EndpointId || Session.FindSessionClient(EndpointId, Dummy);
+				// Case: User muted Floor but only Floor.StaticMeshComponent0 is replicated. Hence, also look for child objects being referenced.
+				const bool bIsReferenced = ConcertSyncCore::IsObjectOrChildReferenced(ClientContent.Value.Streams, ObjectPath);
+				
+				return bIsConnected && bIsReferenced;
+			});
+		};
 
-			const FMultiUserMuteSessionContent& MuteContent = Preset.GetMuteContent();
-			for (const TPair<FSoftObjectPath, FConcertReplication_ObjectMuteSetting>& MutedObject : MuteContent.MutedObjects)
+		const FMultiUserMuteSessionContent& MuteContent = Preset.GetMuteContent();
+		for (const TPair<FSoftObjectPath, FConcertReplication_ObjectMuteSetting>& MutedObject : MuteContent.MutedObjects)
+		{
+			if (IsReferencedByConnectedClient(MutedObject.Key))
 			{
-				if (IsReferencedByConnectedClient(MutedObject.Key))
-				{
-					MuteRequest.ObjectsToMute.Add(MutedObject.Key, MutedObject.Value);
-				}
+				MuteRequest.ObjectsToMute.Add(MutedObject.Key, MutedObject.Value);
 			}
-			for (const TPair<FSoftObjectPath, FConcertReplication_ObjectMuteSetting>& UnmutedObject : MuteContent.UnmutedObjects)
+		}
+		for (const TPair<FSoftObjectPath, FConcertReplication_ObjectMuteSetting>& UnmutedObject : MuteContent.UnmutedObjects)
+		{
+			if (IsReferencedByConnectedClient(UnmutedObject.Key))
 			{
-				if (IsReferencedByConnectedClient(UnmutedObject.Key))
-				{
-					MuteRequest.ObjectsToUnmute.Add(UnmutedObject.Key, UnmutedObject.Value);
-				}
+				MuteRequest.ObjectsToUnmute.Add(UnmutedObject.Key, UnmutedObject.Value);
 			}
 		}
-		
-		static FConcertReplication_PutState_Request BuildRequest(
-			const UMultiUserReplicationSessionPreset& Preset,
-			const IConcertClientSession& Session,
-			EApplyPresetFlags Flags
-			)
-		{
-			FConcertReplication_PutState_Request Request;
+	}
+	
+	static FConcertReplication_PutState_Request BuildRequest(
+		const UMultiUserReplicationSessionPreset& Preset,
+		const IConcertClientSession& Session,
+		EApplyPresetFlags Flags
+		)
+	{
+		FConcertReplication_PutState_Request Request;
 
-			const bool bClearUnreferencedClients = EnumHasAnyFlags(Flags, EApplyPresetFlags::ClearUnreferencedClients);
-			FillStreamAndAuthorityRequest(Request, Preset, Session, bClearUnreferencedClients);
+		const bool bClearUnreferencedClients = EnumHasAnyFlags(Flags, EApplyPresetFlags::ClearUnreferencedClients);
+		FillStreamAndAuthorityRequest(Request, Preset, Session, bClearUnreferencedClients);
 
-			// TODO UE-219829: Once the server allows sending the mute state disconnected clients should have when they rejoin,
-			// simply send over all mute state instead of doing filtering here.
-			FillMuteStateRequest(Request, Preset, Session);
-			
-			return Request;
+		// TODO UE-219829: Once the server allows sending the mute state disconnected clients should have when they rejoin,
+		// simply send over all mute state instead of doing filtering here.
+		FillMuteStateRequest(Request, Preset, Session);
+		
+		return Request;
+	}
+
+	static EReplaceSessionContentErrorCode ExtractErrorCode(const FConcertReplication_PutState_Response& Response)
+	{
+		switch (Response.ResponseCode)
+		{
+		case EConcertReplicationPutStateResponseCode::Success: return EReplaceSessionContentErrorCode::Success;
+		case EConcertReplicationPutStateResponseCode::Timeout: return EReplaceSessionContentErrorCode::Timeout;
+		case EConcertReplicationPutStateResponseCode::FeatureDisabled: return EReplaceSessionContentErrorCode::FeatureDisabled;
+						
+		case EConcertReplicationPutStateResponseCode::ClientUnknown: [[fallthrough]];
+		case EConcertReplicationPutStateResponseCode::StreamError: [[fallthrough]];
+		case EConcertReplicationPutStateResponseCode::AuthorityConflict: [[fallthrough]];
+		case EConcertReplicationPutStateResponseCode::MuteError: [[fallthrough]];
+		default: return EReplaceSessionContentErrorCode::Rejected; 
 		}
+	}
 
-		static EReplaceSessionContentErrorCode ExtractErrorCode(const FConcertReplication_PutState_Response& Response)
+	static void RemoveEmptyObjectsFromLocalClient(ConcertSharedSlate::IEditableReplicationStreamModel& EditModel)
+	{
+		TArray<FSoftObjectPath> EmptyObjects;
+		EditModel.ForEachReplicatedObject([&EditModel, &EmptyObjects](const FSoftObjectPath& Object)
 		{
-			switch (Response.ResponseCode)
+			if (EditModel.GetNumProperties(Object) == 0)
 			{
-			case EConcertReplicationPutStateResponseCode::Success: return EReplaceSessionContentErrorCode::Success;
-			case EConcertReplicationPutStateResponseCode::Timeout: return EReplaceSessionContentErrorCode::Timeout;
-			case EConcertReplicationPutStateResponseCode::FeatureDisabled: return EReplaceSessionContentErrorCode::FeatureDisabled;
-							
-			case EConcertReplicationPutStateResponseCode::ClientUnknown: [[fallthrough]];
-			case EConcertReplicationPutStateResponseCode::StreamError: [[fallthrough]];
-			case EConcertReplicationPutStateResponseCode::AuthorityConflict: [[fallthrough]];
-			case EConcertReplicationPutStateResponseCode::MuteError: [[fallthrough]];
-			default: return EReplaceSessionContentErrorCode::Rejected; 
+				EmptyObjects.Add(Object);
 			}
-		}
+			return EBreakBehavior::Continue;
+		});
+		EditModel.RemoveObjects(EmptyObjects);
+	}
 
-		static void RemoveEmptyObjectsFromLocalClient(ConcertSharedSlate::IEditableReplicationStreamModel& EditModel)
+	static TArray<TPair<const FOnlineClient*, FConcertClientInfo>> DetermineSavedClients(
+		const FOnlineClientManager& ClientManager,
+		const IConcertClientSession& Session,
+		const FSavePresetOptions& Options
+		)
+	{
+		TArray<TPair<const FOnlineClient*, FConcertClientInfo>> IncludedClients;
+		ClientManager.ForEachClient([&Session, &IncludedClients, &Options](const FOnlineClient& Client)
 		{
-			TArray<FSoftObjectPath> EmptyObjects;
-			EditModel.ForEachReplicatedObject([&EditModel, &EmptyObjects](const FSoftObjectPath& Object)
+			FConcertClientInfo ClientInfo;
+			const bool bGotClientInfo = ClientUtils::GetClientDisplayInfo(Session, Client.GetEndpointId(), ClientInfo);
+			if (!ensure(bGotClientInfo))
 			{
-				if (EditModel.GetNumProperties(Object) == 0)
-				{
-					EmptyObjects.Add(Object);
-				}
 				return EBreakBehavior::Continue;
-			});
-			EditModel.RemoveObjects(EmptyObjects);
-		}
+			}
 
-		static TArray<TPair<const FOnlineClient*, FConcertClientInfo>> DetermineSavedClients(
-			const FOnlineClientManager& ClientManager,
-			const IConcertClientSession& Session,
-			const FSavePresetOptions& Options
-			)
-		{
-			TArray<TPair<const FOnlineClient*, FConcertClientInfo>> IncludedClients;
-			ClientManager.ForEachClient([&Session, &IncludedClients, &Options](const FOnlineClient& Client)
+			const bool bIsFilteredOut = Options.ClientFilterDelegate.IsBound() && Options.ClientFilterDelegate.Execute(ClientInfo) == EFilterResult::Exclude; 
+			if (bIsFilteredOut)
 			{
-				FConcertClientInfo ClientInfo;
-				const bool bGotClientInfo = ClientUtils::GetClientDisplayInfo(Session, Client.GetEndpointId(), ClientInfo);
-				if (!ensure(bGotClientInfo))
-				{
-					return EBreakBehavior::Continue;
-				}
-
-				const bool bIsFilteredOut = Options.ClientFilterDelegate.IsBound() && Options.ClientFilterDelegate.Execute(ClientInfo) == EFilterResult::Exclude; 
-				if (bIsFilteredOut)
-				{
-					return EBreakBehavior::Continue;
-				}
-			
-				IncludedClients.Emplace(&Client, ClientInfo);
 				return EBreakBehavior::Continue;
-			});
-			return IncludedClients;
-		}
+			}
+		
+			IncludedClients.Emplace(&Client, ClientInfo);
+			return EBreakBehavior::Continue;
+		});
+		return IncludedClients;
 	}
-	
+}
+
+namespace UE::MultiUserClient::Replication
+{
 	FPresetManager::FPresetManager(
 		const IConcertSyncClient& SyncClient,
 		const FOnlineClientManager& ClientManager,
@@ -319,6 +353,7 @@ namespace UE::MultiUserClient::Replication
 
 			TargetClientPreset->ReplicationMap = CopiedClientStream->ReplicationMap;
 			TargetClientPreset->FrequencySettings = Client->GetStreamSynchronizer().GetFrequencySettings();
+			TargetClientPreset->ActorLabelRemappingData = ConcertSyncCore::GenerateRemappingData(TargetClientPreset->ReplicationMap);
 		}
 
 		Preset->SetMuteContent(
diff --git a/Engine/Plugins/Developer/Concert/ConcertApp/MultiUserClient/Source/MultiUserReplicationEditor/Internal/Assets/MultiUserReplicationClientPreset.h b/Engine/Plugins/Developer/Concert/ConcertApp/MultiUserClient/Source/MultiUserReplicationEditor/Internal/Assets/MultiUserReplicationClientPreset.h
--- a/Engine/Plugins/Developer/Concert/ConcertApp/MultiUserClient/Source/MultiUserReplicationEditor/Internal/Assets/MultiUserReplicationClientPreset.h
+++ b/Engine/Plugins/Developer/Concert/ConcertApp/MultiUserClient/Source/MultiUserReplicationEditor/Internal/Assets/MultiUserReplicationClientPreset.h
@@ -3,7 +3,10 @@
 #pragma once
 
 #include "UObject/Object.h"
+
+#include "Replication/Data/ActorLabelRemapping.h"
 #include "Replication/Data/ReplicationStream.h"
+
 #include "MultiUserReplicationClientPreset.generated.h"
 
 /** Stores info about a client's content in a preset. */
@@ -20,6 +23,15 @@ struct FMultiUserReplicationClientPreset
 	UPROPERTY()
 	FConcertStreamFrequencySettings FrequencySettings;
 
+	/**
+	 * For each FSoftObjectPath in ReplicationMap that references an actor this saves its actor label.
+	 * 
+	 * When the preset is applied, potentially in a new world or different session, this allows rebinding the original FSoftObjectPaths
+	 * to FSoftObjectPaths that now point to different objects but that share the same actor label and share the same subobject hierarchy.
+	 */
+	UPROPERTY()
+	FConcertReplicationRemappingData ActorLabelRemappingData;
+
 	/** The FConcertClientInfo::DisplayName of the client. */
 	UPROPERTY()
 	FString DisplayName;
