Remove FOnlineClient::AllowsEditing, which effectively always returns true.

diff --git a/Engine/Plugins/Developer/Concert/ConcertApp/MultiUserClient/Source/MultiUserClient/Private/Replication/Misc/Frequency/FrequencyUtils.h b/Engine/Plugins/Developer/Concert/ConcertApp/MultiUserClient/Source/MultiUserClient/Private/Replication/Misc/Frequency/FrequencyUtils.h
--- a/Engine/Plugins/Developer/Concert/ConcertApp/MultiUserClient/Source/MultiUserClient/Private/Replication/Misc/Frequency/FrequencyUtils.h
+++ b/Engine/Plugins/Developer/Concert/ConcertApp/MultiUserClient/Source/MultiUserClient/Private/Replication/Misc/Frequency/FrequencyUtils.h
@@ -78,9 +78,8 @@ namespace UE::MultiUserClient::Replication::FrequencyUtils
 	inline bool CanChangeFrequencySettings(const FSoftObjectPath& ContextObject, const FOnlineClient& Client)
 	{
 		const bool bHasProperties = Client.GetStreamSynchronizer().GetServerState().HasProperties(ContextObject);
-		const bool bAllowsEditing = Client.AllowsEditing();
 		const bool bNotProcessingOtherRequests = Client.GetSubmissionWorkflow().CanSubmit();
-		return bHasProperties && bAllowsEditing && bNotProcessingOtherRequests;
+		return bHasProperties && bNotProcessingOtherRequests;
 	}
 
 	/** @return Whether ContextObject's frequency settings can be changed for all given Clients. */
diff --git a/Engine/Plugins/Developer/Concert/ConcertApp/MultiUserClient/Source/MultiUserClient/Private/Replication/Submission/MultiEdit/ReassignObjectPropertiesLogic.cpp b/Engine/Plugins/Developer/Concert/ConcertApp/MultiUserClient/Source/MultiUserClient/Private/Replication/Submission/MultiEdit/ReassignObjectPropertiesLogic.cpp
--- a/Engine/Plugins/Developer/Concert/ConcertApp/MultiUserClient/Source/MultiUserClient/Private/Replication/Submission/MultiEdit/ReassignObjectPropertiesLogic.cpp
+++ b/Engine/Plugins/Developer/Concert/ConcertApp/MultiUserClient/Source/MultiUserClient/Private/Replication/Submission/MultiEdit/ReassignObjectPropertiesLogic.cpp
@@ -50,8 +50,7 @@ namespace UE::MultiUserClient::Replication
 					{
 						// Get the property assignments of the client being reassigned from
 						const FOnlineClient* ClientToReassignFrom = ClientManager.FindClient(ClientId);
-						if (!ensureMsgf(ClientToReassignFrom, TEXT("Authority cache out of sync"))
-							|| !ClientToReassignFrom->AllowsEditing())
+						if (!ensureMsgf(ClientToReassignFrom, TEXT("Authority cache out of sync")))
 						{
 							return EBreakBehavior::Continue;
 						}
@@ -202,39 +201,39 @@ namespace UE::MultiUserClient::Replication
 		}
 
 		/** @return Whether TargetClientObjects includes all properties ClientsToConsider have registered for ObjectToCheck. */
-		static bool DoesTargetIncludeOthers_SingleObject(
+		static bool DoesTargetIncludeOtherClientsContent_SingleObject(
 			const FConcertObjectReplicationMap& TargetClientObjects,
-			const TArray<TNonNullPtr<const FOnlineClient>> ClientsToConsider,
-			const FSoftObjectPath& ObjectToCheck)
+			const FOnlineClientManager& ClientManager,
+			const FSoftObjectPath& ObjectToCheck
+			)
 		{
 			const FConcertReplicatedObjectInfo* ObjectInfo = TargetClientObjects.ReplicatedObjects.Find(ObjectToCheck);
-				
-			for (const FOnlineClient* OtherClient : ClientsToConsider)
+
+			bool bIncludesAllOtherClients = true;
+			ClientManager.ForEachClient([&ObjectToCheck, ObjectInfo, &bIncludesAllOtherClients](const FOnlineClient& OtherClient)
 			{
-				const FConcertObjectReplicationMap& OtherClientObjects = OtherClient->GetStreamSynchronizer().GetServerState();
+				const FConcertObjectReplicationMap& OtherClientObjects = OtherClient.GetStreamSynchronizer().GetServerState();
 				const FConcertReplicatedObjectInfo* OtherObjectInfo = OtherClientObjects.ReplicatedObjects.Find(ObjectToCheck);
 					
 				// TargetClient has at least as much if OtherClient has nothing
 				if (!OtherObjectInfo || OtherObjectInfo->PropertySelection.ReplicatedProperties.IsEmpty())
 				{
-					continue;
+					return EBreakBehavior::Continue;
 				}
 
-				if (!ObjectInfo || !ObjectInfo->PropertySelection.Includes(OtherObjectInfo->PropertySelection))
-				{
-					return false;
-				}
-			}
+				const bool bTargetContainsOther = ObjectInfo && ObjectInfo->PropertySelection.Includes(OtherObjectInfo->PropertySelection);
+				bIncludesAllOtherClients &= bTargetContainsOther;
+				return bIncludesAllOtherClients ? EBreakBehavior::Continue : EBreakBehavior::Break;
+			});
 
-			return true;
+			return bIncludesAllOtherClients;
 		}
 		
-		/** @return Whether ClientId includes all properties assigned to ObjectsToCheck as all other clients (passing the predicate) do. */
-		static bool DoesTargetIncludeOthers(
+		/** @return Whether ClientId includes all properties assigned to ObjectsToCheck as all other clients do. */
+		static bool DoesTargetIncludeOtherClientsContent(
 			const FGuid& TargetClientId,
 			TConstArrayView<FSoftObjectPath> ObjectsToCheck,
-			const FOnlineClientManager& ClientManager,
-			TFunctionRef<bool(const FOnlineClient& Client)> ShouldConsiderClientPredicate
+			const FOnlineClientManager& ClientManager
 			)
 		{
 			const FOnlineClient* TargetClient = ClientManager.FindClient(TargetClientId);
@@ -244,10 +243,9 @@ namespace UE::MultiUserClient::Replication
 			}
 
 			const FConcertObjectReplicationMap& TargetClientObjects = TargetClient->GetStreamSynchronizer().GetServerState();
-			const TArray<TNonNullPtr<const FOnlineClient>> ClientsToConsider = ClientManager.GetClients(ShouldConsiderClientPredicate);
-			return Algo::AllOf(ObjectsToCheck, [&TargetClientObjects, &ClientsToConsider](const FSoftObjectPath& ObjectToCheck)
+			return Algo::AllOf(ObjectsToCheck, [&ClientManager, &TargetClientObjects](const FSoftObjectPath& ObjectToCheck)
 			{
-				return DoesTargetIncludeOthers_SingleObject(TargetClientObjects, ClientsToConsider, ObjectToCheck);
+				return DoesTargetIncludeOtherClientsContent_SingleObject(TargetClientObjects, ClientManager, ObjectToCheck);
 			});
 		}
 
@@ -343,16 +341,9 @@ namespace UE::MultiUserClient::Replication
 			SET_REASON(LOCTEXT("Reason.ClientDisconnected", "Client disconnected"));
 			return false;
 		}
-		
-		if (!Client->AllowsEditing())
-		{
-			SET_REASON(LOCTEXT("Reason.ClientNoEditing", "Client does not allow remote editing"));
-			return false;
-		}
 
 		// Relevant after user has re-assigned all properties to a given client
-		auto ShouldConsiderClient = [this](const FOnlineClient& Client) { return Client.AllowsEditing(); };
-		if (ReassignObjectProperties::Private::DoesTargetIncludeOthers(ClientId, ObjectsToReassign, ClientManager, ShouldConsiderClient))
+		if (ReassignObjectProperties::Private::DoesTargetIncludeOtherClientsContent(ClientId, ObjectsToReassign, ClientManager))
 		{
 			const bool bHasAnyProperties = ReassignObjectProperties::Private::HasAnyProperties(ClientId, ObjectsToReassign, ClientManager);
 			SET_REASON(
@@ -373,7 +364,7 @@ namespace UE::MultiUserClient::Replication
 		using namespace ConcertSyncClient::Replication;
 		
 		const FOnlineClient* ClientToAssignTo = ClientManager.FindClient(ClientId);
-		if (!ensure(ClientToAssignTo && ClientToAssignTo->AllowsEditing()))
+		if (!ensure(ClientToAssignTo))
 		{
 			UE_LOG(LogConcert, Error, TEXT("Property Reassignment: The target client is not editable."));
 			return;
@@ -477,8 +468,7 @@ namespace UE::MultiUserClient::Replication
 		}
 		
 		const FOnlineClient* ClientToAssignTo = ClientManager.FindClient(InProgressOperation->AssignedToClient);
-		if (!ensureMsgf(ClientToAssignTo, TEXT("FReassignObjectPropertiesLogic should have cancelled the change upon disconnect"))
-			|| !ClientToAssignTo->AllowsEditing())
+		if (!ensureMsgf(ClientToAssignTo, TEXT("FReassignObjectPropertiesLogic should have cancelled the change upon disconnect")))
 		{
 			UE_LOG(LogConcert, Error, TEXT("Property Reassignment: The target client is no longer editable."));
 			return;
diff --git a/Engine/Plugins/Developer/Concert/ConcertApp/MultiUserClient/Source/MultiUserClient/Private/Widgets/ActiveSession/Replication/Client/Multi/Columns/SAssignPropertyComboBox.cpp b/Engine/Plugins/Developer/Concert/ConcertApp/MultiUserClient/Source/MultiUserClient/Private/Widgets/ActiveSession/Replication/Client/Multi/Columns/SAssignPropertyComboBox.cpp
--- a/Engine/Plugins/Developer/Concert/ConcertApp/MultiUserClient/Source/MultiUserClient/Private/Widgets/ActiveSession/Replication/Client/Multi/Columns/SAssignPropertyComboBox.cpp
+++ b/Engine/Plugins/Developer/Concert/ConcertApp/MultiUserClient/Source/MultiUserClient/Private/Widgets/ActiveSession/Replication/Client/Multi/Columns/SAssignPropertyComboBox.cpp
@@ -245,7 +245,7 @@ namespace UE::MultiUserClient::Replication
 		bool bCanRemoveFromOwners = true;
 		ClientManager->ForEachClient([this, &EndpointId, &Reason, Client, &bCanRemoveFromOwners](const FOnlineClient& ClientToRemoveFrom)
 		{
-			if (*Client != ClientToRemoveFrom && !ClientToRemoveFrom.AllowsEditing())
+			if (*Client != ClientToRemoveFrom)
 			{
 				const bool bHasAnySelectedObject = Algo::AnyOf(EditedObjects, [this, &ClientToRemoveFrom](const TSoftObjectPtr<>& Object)
 				{
@@ -267,15 +267,7 @@ namespace UE::MultiUserClient::Replication
 			return false;
 		}
 		
-		const bool bAllowsEditing = Client->AllowsEditing();
-		if (!bAllowsEditing)
-		{
-			SET_REASON(FText::Format(
-				LOCTEXT("RemoteEditingDisabled", "Client {0} does not allow remote editing of its properties."),
-				FText::FromString(ClientUtils::GetClientDisplayName(*ConcertClient, EndpointId))
-				));
-		}
-		return bAllowsEditing;
+		return true;
 	}
 #undef SET_REASON
 	
@@ -353,7 +345,7 @@ namespace UE::MultiUserClient::Replication
 		ClientManager->ForEachClient([this, &ShouldRemoveFromClient](const FOnlineClient& ClientToRemoveFrom)
 		{
 			const TSharedRef<ConcertSharedSlate::IEditableReplicationStreamModel> EditModel = ClientToRemoveFrom.GetClientEditModel();
-			if (ClientToRemoveFrom.AllowsEditing() && ShouldRemoveFromClient(ClientToRemoveFrom))
+			if (ShouldRemoveFromClient(ClientToRemoveFrom))
 			{
 				for (const TSoftObjectPtr<>& Object : EditedObjects)
 				{
diff --git a/Engine/Plugins/Developer/Concert/ConcertApp/MultiUserClient/Source/MultiUserClient/Private/Widgets/ActiveSession/Replication/Client/Multi/MultiStreamModel.cpp b/Engine/Plugins/Developer/Concert/ConcertApp/MultiUserClient/Source/MultiUserClient/Private/Widgets/ActiveSession/Replication/Client/Multi/MultiStreamModel.cpp
--- a/Engine/Plugins/Developer/Concert/ConcertApp/MultiUserClient/Source/MultiUserClient/Private/Widgets/ActiveSession/Replication/Client/Multi/MultiStreamModel.cpp
+++ b/Engine/Plugins/Developer/Concert/ConcertApp/MultiUserClient/Source/MultiUserClient/Private/Widgets/ActiveSession/Replication/Client/Multi/MultiStreamModel.cpp
@@ -19,14 +19,7 @@ namespace UE::MultiUserClient::Replication
 
 	void FMultiStreamModel::ForEachClient(TFunctionRef<EBreakBehavior(const FOnlineClient*)> ProcessClient) const
 	{
-		for (const FOnlineClient* ReadOnlyClient : GetCachedReadOnlyClients())
-		{
-			if (ProcessClient(ReadOnlyClient) == EBreakBehavior::Break)
-			{
-				return;
-			}
-		}
-		for (const FOnlineClient* WritableClient : GetCachedWritableClients())
+		for (const FOnlineClient* WritableClient : CachedWritableClients)
 		{
 			if (ProcessClient(WritableClient) == EBreakBehavior::Break)
 			{
@@ -35,13 +28,6 @@ namespace UE::MultiUserClient::Replication
 		}
 	}
 
-	TSet<TSharedRef<ConcertSharedSlate::IReplicationStreamModel>> FMultiStreamModel::GetReadOnlyStreams() const
-	{
-		TSet<TSharedRef<ConcertSharedSlate::IReplicationStreamModel>> Result;
-		Algo::Transform(CachedReadOnlyClients, Result, [](const FOnlineClient* Client){ return Client->GetClientEditModel(); });
-		return Result;
-	}
-
 	TSet<TSharedRef<ConcertSharedSlate::IEditableReplicationStreamModel>> FMultiStreamModel::GetEditableStreams() const
 	{
 		TSet<TSharedRef<ConcertSharedSlate::IEditableReplicationStreamModel>> Result;
@@ -51,39 +37,27 @@ namespace UE::MultiUserClient::Replication
 
 	void FMultiStreamModel::RebuildStreamsSets()
 	{
-		// Cannot just iterate through CachedReadOnlyClients because it may contain stale clients that were just removed
+		// It is not safe to iterate through our cached client array because it may contain stale clients that were just removed.
 		ClientManager.ForEachClient([this](FOnlineClient& Client)
 		{
 			Client.OnModelChanged().RemoveAll(this);
 			return EBreakBehavior::Continue;
 		});
 		
-		TSet<const FOnlineClient*> ReadOnlyStreams;
 		TSet<const FOnlineClient*> WritableStreams;
-		OnlineClientSelectionModel.ForEachItem([this, &ReadOnlyStreams, &WritableStreams](FOnlineClient& Client)
+		OnlineClientSelectionModel.ForEachItem([this, &WritableStreams](FOnlineClient& Client)
 		{
-			const bool bIsUploadable = CanEverSubmit(Client.GetSubmissionWorkflow().GetUploadability());
 			const TSharedRef<ConcertSharedSlate::IEditableReplicationStreamModel> Stream = Client.GetClientEditModel();
 			Client.OnModelChanged().AddRaw(this, &FMultiStreamModel::OnStreamExternallyChanged, Stream.ToWeakPtr());
 			
-			TSet<const FOnlineClient*>& StreamToAddTo = Client.AllowsEditing() ? WritableStreams : ReadOnlyStreams;
-			StreamToAddTo.Add(&Client);
+			WritableStreams.Add(&Client);
 			return EBreakBehavior::Continue;
 		});
 
-		const bool bReadOnlyStayedSame = CachedReadOnlyClients.Num() == ReadOnlyStreams.Num() && CachedReadOnlyClients.Includes(ReadOnlyStreams);
-		if (!bReadOnlyStayedSame)
-		{
-			CachedReadOnlyClients = MoveTemp(ReadOnlyStreams);
-		}
 		const bool bWritableStayedSame = CachedWritableClients.Num() == WritableStreams.Num() && CachedWritableClients.Includes(WritableStreams);
 		if (!bWritableStayedSame)
 		{
 			CachedWritableClients = MoveTemp(WritableStreams);
-		}
-
-		if (!bReadOnlyStayedSame || !bWritableStayedSame)
-		{
 			OnStreamSetChangedDelegate.Broadcast();
 		}
 	}
diff --git a/Engine/Plugins/Developer/Concert/ConcertApp/MultiUserClient/Source/MultiUserClient/Private/Widgets/ActiveSession/Replication/Client/Multi/MultiStreamModel.h b/Engine/Plugins/Developer/Concert/ConcertApp/MultiUserClient/Source/MultiUserClient/Private/Widgets/ActiveSession/Replication/Client/Multi/MultiStreamModel.h
--- a/Engine/Plugins/Developer/Concert/ConcertApp/MultiUserClient/Source/MultiUserClient/Private/Widgets/ActiveSession/Replication/Client/Multi/MultiStreamModel.h
+++ b/Engine/Plugins/Developer/Concert/ConcertApp/MultiUserClient/Source/MultiUserClient/Private/Widgets/ActiveSession/Replication/Client/Multi/MultiStreamModel.h
@@ -24,12 +24,10 @@ namespace UE::MultiUserClient::Replication
 		
 		FMultiStreamModel(IOnlineClientSelectionModel& InOnlineClientSelectionModel, FOnlineClientManager& InClientManager);
 		
-		const TSet<const FOnlineClient*>& GetCachedReadOnlyClients() const { return CachedReadOnlyClients; }
-		const TSet<const FOnlineClient*>& GetCachedWritableClients() const { return CachedWritableClients; }
 		void ForEachClient(TFunctionRef<EBreakBehavior(const FOnlineClient*)> ProcessClient) const;
 
 		//~ Begin IEditableMultiReplicationStreamModel Interface
-		virtual TSet<TSharedRef<ConcertSharedSlate::IReplicationStreamModel>> GetReadOnlyStreams() const override;
+		virtual TSet<TSharedRef<ConcertSharedSlate::IReplicationStreamModel>> GetReadOnlyStreams() const override { return {}; }
 		virtual TSet<TSharedRef<ConcertSharedSlate::IEditableReplicationStreamModel>> GetEditableStreams() const override;
 		virtual FOnStreamExternallyChanged& OnStreamExternallyChanged() override { return OnReadOnlyStreamChangedDelegate; }
 		virtual FOnStreamSetChanged& OnStreamSetChanged() override { return OnStreamSetChangedDelegate; }
@@ -42,7 +40,6 @@ namespace UE::MultiUserClient::Replication
 		/** Used to obtain a list of clients for unsubscribing. */
 		FOnlineClientManager& ClientManager;
 
-		TSet<const FOnlineClient*> CachedReadOnlyClients;
 		TSet<const FOnlineClient*> CachedWritableClients;
 
 		FOnStreamExternallyChanged OnReadOnlyStreamChangedDelegate;
