Rename GetClientManager to GetOnlineClientManger.

diff --git a/Engine/Plugins/Developer/Concert/ConcertApp/MultiUserClient/Source/MultiUserClient/Private/Replication/MultiUserReplicationManager.h b/Engine/Plugins/Developer/Concert/ConcertApp/MultiUserClient/Source/MultiUserClient/Private/Replication/MultiUserReplicationManager.h
--- a/Engine/Plugins/Developer/Concert/ConcertApp/MultiUserClient/Source/MultiUserClient/Private/Replication/MultiUserReplicationManager.h
+++ b/Engine/Plugins/Developer/Concert/ConcertApp/MultiUserClient/Source/MultiUserClient/Private/Replication/MultiUserReplicationManager.h
@@ -62,8 +62,8 @@ namespace UE::MultiUserClient::Replication
 		void JoinReplicationSession();
 
 		/** @note You're not supposed to keep any reference to the ClientManager since it can become invalid depending on connection state. */
-		FOnlineClientManager* GetClientManager() { return ConnectedState ? &ConnectedState->ClientManager : nullptr; }
-		const FOnlineClientManager* GetClientManager() const { return ConnectedState ? &ConnectedState->ClientManager : nullptr; }
+		FOnlineClientManager* GetOnlineClientManager() { return ConnectedState ? &ConnectedState->ClientManager : nullptr; }
+		const FOnlineClientManager* GetOnlineClientManager() const { return ConnectedState ? &ConnectedState->ClientManager : nullptr; }
 		
 		/** @note You're not supposed to keep any reference to the MuteManager since it can become invalid depending on connection state. */
 		FMuteStateManager* GetMuteManager() { return ConnectedState ? &ConnectedState->MuteManager : nullptr; }
diff --git a/Engine/Plugins/Developer/Concert/ConcertApp/MultiUserClient/Source/MultiUserClient/Private/Widgets/ActiveSession/Replication/Client/Multi/SAllClientsView.cpp b/Engine/Plugins/Developer/Concert/ConcertApp/MultiUserClient/Source/MultiUserClient/Private/Widgets/ActiveSession/Replication/Client/Multi/SAllClientsView.cpp
--- a/Engine/Plugins/Developer/Concert/ConcertApp/MultiUserClient/Source/MultiUserClient/Private/Widgets/ActiveSession/Replication/Client/Multi/SAllClientsView.cpp
+++ b/Engine/Plugins/Developer/Concert/ConcertApp/MultiUserClient/Source/MultiUserClient/Private/Widgets/ActiveSession/Replication/Client/Multi/SAllClientsView.cpp
@@ -13,7 +13,7 @@ namespace UE::MultiUserClient::Replication
 {
 	void SAllClientsView::Construct(const FArguments&, TSharedRef<IConcertClient> InConcertClient, FMultiUserReplicationManager& InMultiUserReplicationManager)
 	{
-		ClientManager = InMultiUserReplicationManager.GetClientManager();
+		ClientManager = InMultiUserReplicationManager.GetOnlineClientManager();
 		check(ClientManager);
 		AllClientsModel = MakeUnique<FAllOnlineClientsSelectionModel>(*ClientManager);
 		
diff --git a/Engine/Plugins/Developer/Concert/ConcertApp/MultiUserClient/Source/MultiUserClient/Private/Widgets/ActiveSession/Replication/Client/Multi/SMultiClientView.cpp b/Engine/Plugins/Developer/Concert/ConcertApp/MultiUserClient/Source/MultiUserClient/Private/Widgets/ActiveSession/Replication/Client/Multi/SMultiClientView.cpp
--- a/Engine/Plugins/Developer/Concert/ConcertApp/MultiUserClient/Source/MultiUserClient/Private/Widgets/ActiveSession/Replication/Client/Multi/SMultiClientView.cpp
+++ b/Engine/Plugins/Developer/Concert/ConcertApp/MultiUserClient/Source/MultiUserClient/Private/Widgets/ActiveSession/Replication/Client/Multi/SMultiClientView.cpp
@@ -37,7 +37,7 @@ namespace UE::MultiUserClient::Replication
 		IOnlineClientSelectionModel& InOnlineClientSelectionModel
 		)
 	{
-		ClientManager = InMultiUserReplicationManager.GetClientManager();
+		ClientManager = InMultiUserReplicationManager.GetOnlineClientManager();
 		UserSelectedProperties = InMultiUserReplicationManager.GetUserPropertySelector();
 		StreamModel = MakeShared<FMultiStreamModel>(InOnlineClientSelectionModel, *ClientManager);
 		OnlineClientSelectionModel = &InOnlineClientSelectionModel;
