Fix not being able to assign property to remote client

diff --git a/Engine/Plugins/Developer/Concert/ConcertApp/MultiUserClient/Source/MultiUserClient/Private/Widgets/ActiveSession/Replication/Client/Multi/Columns/SAssignPropertyComboBox.cpp b/Engine/Plugins/Developer/Concert/ConcertApp/MultiUserClient/Source/MultiUserClient/Private/Widgets/ActiveSession/Replication/Client/Multi/Columns/SAssignPropertyComboBox.cpp
--- a/Engine/Plugins/Developer/Concert/ConcertApp/MultiUserClient/Source/MultiUserClient/Private/Widgets/ActiveSession/Replication/Client/Multi/Columns/SAssignPropertyComboBox.cpp
+++ b/Engine/Plugins/Developer/Concert/ConcertApp/MultiUserClient/Source/MultiUserClient/Private/Widgets/ActiveSession/Replication/Client/Multi/Columns/SAssignPropertyComboBox.cpp
@@ -240,32 +240,6 @@ namespace UE::MultiUserClient::Replication
 			SET_REASON(LOCTEXT("ClientDisconnected", "Client disconnected."));
 			return false;
 		}
-
-		// The combo box assigns the property to the clicked client and removes from the others... check that the currently assigned clients allow it.
-		bool bCanRemoveFromOwners = true;
-		ClientManager->ForEachClient([this, &EndpointId, &Reason, Client, &bCanRemoveFromOwners](const FOnlineClient& ClientToRemoveFrom)
-		{
-			if (*Client != ClientToRemoveFrom)
-			{
-				const bool bHasAnySelectedObject = Algo::AnyOf(EditedObjects, [this, &ClientToRemoveFrom](const TSoftObjectPtr<>& Object)
-				{
-					return ClientToRemoveFrom.GetClientEditModel()->HasProperty(Object.GetUniqueID(), Property);
-				});
-				bCanRemoveFromOwners = !bHasAnySelectedObject;
-				
-				SET_REASON(FText::Format(
-					LOCTEXT("OwningClientDoesNotAllow", "Client {0} does not allow remote editing of its properties but has registered this property."),
-					FText::FromString(ClientUtils::GetClientDisplayName(*ConcertClient, EndpointId))
-					));
-				return EBreakBehavior::Break;
-			}
-			return EBreakBehavior::Continue;
-		});
-
-		if (!bCanRemoveFromOwners)
-		{
-			return false;
-		}
 		
 		return true;
 	}
