Fix replicated object of client 2 still being shown in UI of client 1 after client 2 leaves the session.

diff --git a/Engine/Plugins/Developer/Concert/ConcertApp/MultiUserClient/Source/MultiUserClient/Private/Widgets/ActiveSession/Replication/Client/Multi/SMultiClientView.cpp b/Engine/Plugins/Developer/Concert/ConcertApp/MultiUserClient/Source/MultiUserClient/Private/Widgets/ActiveSession/Replication/Client/Multi/SMultiClientView.cpp
--- a/Engine/Plugins/Developer/Concert/ConcertApp/MultiUserClient/Source/MultiUserClient/Private/Widgets/ActiveSession/Replication/Client/Multi/SMultiClientView.cpp
+++ b/Engine/Plugins/Developer/Concert/ConcertApp/MultiUserClient/Source/MultiUserClient/Private/Widgets/ActiveSession/Replication/Client/Multi/SMultiClientView.cpp
@@ -202,6 +202,8 @@ namespace UE::MultiUserClient::Replication
 			Client.OnHierarchyNeedsRefresh().AddRaw(this, &SMultiClientView::RefreshUI);
 			return EBreakBehavior::Continue;
 		});
+
+		RefreshUI();
 	}
 
 	void SMultiClientView::CleanClientSubscriptions() const
