Extract assignment logic from column widget to new FAssignedClientsModel.

diff --git a/Engine/Plugins/Developer/Concert/ConcertApp/MultiUserClient/Source/MultiUserClient/Private/Replication/Submission/MultiEdit/ReassignObjectPropertiesLogic.cpp b/Engine/Plugins/Developer/Concert/ConcertApp/MultiUserClient/Source/MultiUserClient/Private/Replication/Submission/MultiEdit/ReassignObjectPropertiesLogic.cpp
--- a/Engine/Plugins/Developer/Concert/ConcertApp/MultiUserClient/Source/MultiUserClient/Private/Replication/Submission/MultiEdit/ReassignObjectPropertiesLogic.cpp
+++ b/Engine/Plugins/Developer/Concert/ConcertApp/MultiUserClient/Source/MultiUserClient/Private/Replication/Submission/MultiEdit/ReassignObjectPropertiesLogic.cpp
@@ -13,265 +13,278 @@
 
 #define LOCTEXT_NAMESPACE "FReassignObjectPropertiesLogic"
 
-namespace UE::MultiUserClient::Replication
+namespace UE::MultiUserClient::Replication::ReassignObjectProperties
 {
-	namespace ReassignObjectProperties::Private
-	{
-		using FClientId = FGuid;
-		
-		struct FObjectReassignment
-		{
-			/** Used to build the final changelist to apply to the target client. Maps ClientIds to the object data they had registered. */
-			TMap<FClientId, FConcertObjectReplicationMap> OldRegisteredObjects;
-			/** The frequency settings clients when reassignment operation was started. */
-			TMap<FClientId, FConcertStreamFrequencySettings> OldFrequencies;
-			/** Changes to make to all the other clients */
-			TMap<FClientId, FConcertReplication_ChangeStream_Request> ReassignedClientRequests;
-			/** Authority that the clients had before. This is used to determine whether to give the assigned to client authority once the stream changes have been made. */
-			TMap<FClientId, TArray<FSoftObjectPath>> ReassignedAuthority;
-		};
+	using FClientId = FGuid;
 		
-		/** Builds all the changes that must be made to the clients we are transferring from.  */
-		static FObjectReassignment BuildChangesForTransferal(const FOnlineClientManager& ClientManager, TConstArrayView<FSoftObjectPath> ObjectsToReassign, const FGuid& StreamId)
+	struct FObjectReassignment
+	{
+		/** Used to build the final changelist to apply to the target client. Maps ClientIds to the object data they had registered. */
+		TMap<FClientId, FConcertObjectReplicationMap> OldRegisteredObjects;
+		/** The frequency settings clients when reassignment operation was started. */
+		TMap<FClientId, FConcertStreamFrequencySettings> OldFrequencies;
+		/** Changes to make to all the other clients */
+		TMap<FClientId, FConcertReplication_ChangeStream_Request> ReassignedClientRequests;
+		/** Authority that the clients had before. This is used to determine whether to give the assigned to client authority once the stream changes have been made. */
+		TMap<FClientId, TArray<FSoftObjectPath>> ReassignedAuthority;
+	};
+	
+	/** Builds all the changes that must be made to the clients we are transferring from.  */
+	static FObjectReassignment BuildChangesForTransferal(const FOnlineClientManager& ClientManager, TConstArrayView<FSoftObjectPath> ObjectsToReassign, const FGuid& StreamId)
+	{
+		FObjectReassignment Result;
+		TMap<FGuid, FConcertObjectReplicationMap>& OldRegisteredObjects = Result.OldRegisteredObjects;
+		TMap<FGuid, FConcertStreamFrequencySettings>& OldFrequencies = Result.OldFrequencies;
+		TMap<FGuid, FConcertReplication_ChangeStream_Request>& ReassignedClientRequests = Result.ReassignedClientRequests;
+		TMap<FGuid, TArray<FSoftObjectPath>>& ReassignedAuthority = Result.ReassignedAuthority;
+	
+		for (const FSoftObjectPath& ObjectPath : ObjectsToReassign)
 		{
-			FObjectReassignment Result;
-			TMap<FGuid, FConcertObjectReplicationMap>& OldRegisteredObjects = Result.OldRegisteredObjects;
-			TMap<FGuid, FConcertStreamFrequencySettings>& OldFrequencies = Result.OldFrequencies;
-			TMap<FGuid, FConcertReplication_ChangeStream_Request>& ReassignedClientRequests = Result.ReassignedClientRequests;
-			TMap<FGuid, TArray<FSoftObjectPath>>& ReassignedAuthority = Result.ReassignedAuthority;
-		
-			for (const FSoftObjectPath& ObjectPath : ObjectsToReassign)
-			{
-				const FConcertObjectInStreamID ObjectId{ StreamId, ObjectPath };
+			const FConcertObjectInStreamID ObjectId{ StreamId, ObjectPath };
 
-				const FGlobalAuthorityCache& AuthorityCache = ClientManager.GetAuthorityCache();
-				AuthorityCache.ForEachClientWithObjectInStream(ObjectPath,
-					[&ClientManager, &OldRegisteredObjects, &OldFrequencies, &ReassignedClientRequests, &ReassignedAuthority, &ObjectPath, &ObjectId, &AuthorityCache](const FGuid& ClientId)
+			const FGlobalAuthorityCache& AuthorityCache = ClientManager.GetAuthorityCache();
+			AuthorityCache.ForEachClientWithObjectInStream(ObjectPath,
+				[&ClientManager, &OldRegisteredObjects, &OldFrequencies, &ReassignedClientRequests, &ReassignedAuthority, &ObjectPath, &ObjectId, &AuthorityCache](const FGuid& ClientId)
+				{
+					// Get the property assignments of the client being reassigned from
+					const FOnlineClient* ClientToReassignFrom = ClientManager.FindClient(ClientId);
+					if (!ensureMsgf(ClientToReassignFrom, TEXT("Authority cache out of sync")))
 					{
-						// Get the property assignments of the client being reassigned from
-						const FOnlineClient* ClientToReassignFrom = ClientManager.FindClient(ClientId);
-						if (!ensureMsgf(ClientToReassignFrom, TEXT("Authority cache out of sync")))
-						{
-							return EBreakBehavior::Continue;
-						}
-						
-						const IClientStreamSynchronizer& ReassignedStreamSynchronizer = ClientToReassignFrom->GetStreamSynchronizer();
-						const FConcertReplicatedObjectInfo* ObjectInfo = ReassignedStreamSynchronizer.GetServerState().ReplicatedObjects.Find(ObjectId.Object);
-						if (!ensureMsgf(ObjectInfo, TEXT("ForEachClientWithObjectInStream lied")))
-						{
-							return EBreakBehavior::Continue;
-						}
-
-						// Transfer from client being reassigned from to the client being assigned to
-						OldRegisteredObjects.FindOrAdd(ClientId).ReplicatedObjects.Add(ObjectPath, *ObjectInfo);
-						// Remove the object completely from the reassigned client.
-						ReassignedClientRequests.FindOrAdd(ClientId).ObjectsToRemove.Add(ObjectId);
-						// Later we can determine whether to give the assigned to client authority
-						if (AuthorityCache.HasAuthorityOverObject(ObjectPath, ClientId))
-						{
-							ReassignedAuthority.FindOrAdd(ClientId).Add(ObjectPath);
-						}
-
-						// Keep track of this object's frequency settings so it can be transferred
-						const FConcertObjectReplicationSettings* OverrideFrequencySettings = ReassignedStreamSynchronizer.GetFrequencySettings().ObjectOverrides.Find(ObjectId.Object);
-						if (OverrideFrequencySettings)
-						{
-							OldFrequencies.FindOrAdd(ClientId).ObjectOverrides.Add(ObjectId.Object, *OverrideFrequencySettings);
-						}
+						return EBreakBehavior::Continue;
+					}
 					
+					const IClientStreamSynchronizer& ReassignedStreamSynchronizer = ClientToReassignFrom->GetStreamSynchronizer();
+					const FConcertReplicatedObjectInfo* ObjectInfo = ReassignedStreamSynchronizer.GetServerState().ReplicatedObjects.Find(ObjectId.Object);
+					if (!ensureMsgf(ObjectInfo, TEXT("ForEachClientWithObjectInStream lied")))
+					{
 						return EBreakBehavior::Continue;
-					});
-			}
+					}
 
-			// Need to set Defaults of the OldFrequencies because that's not been done, yet.
-			for (TPair<FGuid, FConcertStreamFrequencySettings>& OldFrequency : OldFrequencies)
-			{
-				const FOnlineClient* ClientToReassignFrom = ClientManager.FindClient(OldFrequency.Key);
-				check(ClientToReassignFrom);
-				OldFrequency.Value.Defaults = ClientToReassignFrom->GetStreamSynchronizer().GetFrequencySettings().Defaults;
-			}
+					// Transfer from client being reassigned from to the client being assigned to
+					OldRegisteredObjects.FindOrAdd(ClientId).ReplicatedObjects.Add(ObjectPath, *ObjectInfo);
+					// Remove the object completely from the reassigned client.
+					ReassignedClientRequests.FindOrAdd(ClientId).ObjectsToRemove.Add(ObjectId);
+					// Later we can determine whether to give the assigned to client authority
+					if (AuthorityCache.HasAuthorityOverObject(ObjectPath, ClientId))
+					{
+						ReassignedAuthority.FindOrAdd(ClientId).Add(ObjectPath);
+					}
 
-			return Result;
+					// Keep track of this object's frequency settings so it can be transferred
+					const FConcertObjectReplicationSettings* OverrideFrequencySettings = ReassignedStreamSynchronizer.GetFrequencySettings().ObjectOverrides.Find(ObjectId.Object);
+					if (OverrideFrequencySettings)
+					{
+						OldFrequencies.FindOrAdd(ClientId).ObjectOverrides.Add(ObjectId.Object, *OverrideFrequencySettings);
+					}
+				
+					return EBreakBehavior::Continue;
+				});
+		}
+
+		// Need to set Defaults of the OldFrequencies because that's not been done, yet.
+		for (TPair<FGuid, FConcertStreamFrequencySettings>& OldFrequency : OldFrequencies)
+		{
+			const FOnlineClient* ClientToReassignFrom = ClientManager.FindClient(OldFrequency.Key);
+			check(ClientToReassignFrom);
+			OldFrequency.Value.Defaults = ClientToReassignFrom->GetStreamSynchronizer().GetFrequencySettings().Defaults;
 		}
 
-		static FConcertObjectReplicationSettings GetHighestFrequency(const FSoftObjectPath& Object, const TMap<FClientId, FConcertStreamFrequencySettings>& OldFrequencies)
+		return Result;
+	}
+
+	static FConcertObjectReplicationSettings GetHighestFrequency(const FSoftObjectPath& Object, const TMap<FClientId, FConcertStreamFrequencySettings>& OldFrequencies)
+	{
+		TOptional<FConcertObjectReplicationSettings> Settings;
+		for (const TPair<FClientId, FConcertStreamFrequencySettings>& Pair : OldFrequencies)
 		{
-			TOptional<FConcertObjectReplicationSettings> Settings;
-			for (const TPair<FClientId, FConcertStreamFrequencySettings>& Pair : OldFrequencies)
+			const FConcertObjectReplicationSettings& ClientSettings = Pair.Value.GetSettingsFor(Object);
+			if (!Settings)
 			{
-				const FConcertObjectReplicationSettings& ClientSettings = Pair.Value.GetSettingsFor(Object);
-				if (!Settings)
-				{
-					Settings = ClientSettings;
-				}
-				else if (*Settings <= ClientSettings)
-				{
-					Settings = ClientSettings;
-				}
+				Settings = ClientSettings;
+			}
+			else if (*Settings <= ClientSettings)
+			{
+				Settings = ClientSettings;
 			}
-			return Settings.Get(FConcertObjectReplicationSettings{});
 		}
+		return Settings.Get(FConcertObjectReplicationSettings{});
+	}
 
-		/** Builds a changelist based on which remote clients we managed to change successfully. */
-		static TPair<FStreamChangelist, FFrequencyChangelist> MakeChangelistFromAppliedChanges(
-			const TMap<FGuid, FConcertObjectReplicationMap>& OldRegisteredObjects,
-			const TMap<FClientId, FConcertStreamFrequencySettings>& OldFrequencies,
-			const FParallelExecutionResult& ParallelExecutionResult,
-			const IClientStreamSynchronizer& AssignedToStreamSynchronizer
-			)
+	/** Builds a changelist based on which remote clients we managed to change successfully. */
+	static TPair<FStreamChangelist, FFrequencyChangelist> MakeChangelistFromAppliedChanges(
+		const TMap<FGuid, FConcertObjectReplicationMap>& OldRegisteredObjects,
+		const TMap<FClientId, FConcertStreamFrequencySettings>& OldFrequencies,
+		const FParallelExecutionResult& ParallelExecutionResult,
+		const IClientStreamSynchronizer& AssignedToStreamSynchronizer
+		)
+	{
+		FStreamChangelist ObjectChanges;
+		FFrequencyChangelist FrequencyChanges;
+		
+		const FConcertObjectReplicationMap& CurrentTargetClientState = AssignedToStreamSynchronizer.GetServerState();
+		const FConcertStreamFrequencySettings& CurrentFrequencySettings = AssignedToStreamSynchronizer.GetFrequencySettings();
+		const FGuid TargetStreamId = AssignedToStreamSynchronizer.GetStreamId();
+		for (const TPair<FGuid, FSubmitStreamChangesResponse>& ChangesRequestedOnRemote : ParallelExecutionResult.StreamResponses)
 		{
-			FStreamChangelist ObjectChanges;
-			FFrequencyChangelist FrequencyChanges;
+			const TOptional<FCompletedChangeSubmission>& SubmissionInfo = ChangesRequestedOnRemote.Value.SubmissionInfo;
+			if (!SubmissionInfo || SubmissionInfo->Response.IsFailure())
+			{
+				continue;
+			}
 			
-			const FConcertObjectReplicationMap& CurrentTargetClientState = AssignedToStreamSynchronizer.GetServerState();
-			const FConcertStreamFrequencySettings& CurrentFrequencySettings = AssignedToStreamSynchronizer.GetFrequencySettings();
-			const FGuid TargetStreamId = AssignedToStreamSynchronizer.GetStreamId();
-			for (const TPair<FGuid, FSubmitStreamChangesResponse>& ChangesRequestedOnRemote : ParallelExecutionResult.StreamResponses)
+			const FConcertObjectReplicationMap& ObjectReplicationMap = OldRegisteredObjects[ChangesRequestedOnRemote.Key];
+			for (const TPair<FSoftObjectPath, FConcertReplicatedObjectInfo>& ChangesToApply : ObjectReplicationMap.ReplicatedObjects)
 			{
-				const TOptional<FCompletedChangeSubmission>& SubmissionInfo = ChangesRequestedOnRemote.Value.SubmissionInfo;
-				if (!SubmissionInfo || SubmissionInfo->Response.IsFailure())
+				const FSoftObjectPath& ObjectPath = ChangesToApply.Key;
+				const FConcertObjectInStreamID ObjectId { TargetStreamId, ObjectPath };
+				TOptional<FConcertReplication_ChangeStream_PutObject> PutRequest = FConcertReplication_ChangeStream_PutObject::MakeFromInfo(ChangesToApply.Value);
+				if (!ensure(PutRequest))
 				{
 					continue;
 				}
-				
-				const FConcertObjectReplicationMap& ObjectReplicationMap = OldRegisteredObjects[ChangesRequestedOnRemote.Key];
-				for (const TPair<FSoftObjectPath, FConcertReplicatedObjectInfo>& ChangesToApply : ObjectReplicationMap.ReplicatedObjects)
-				{
-					const FSoftObjectPath& ObjectPath = ChangesToApply.Key;
-					const FConcertObjectInStreamID ObjectId { TargetStreamId, ObjectPath };
-					TOptional<FConcertReplication_ChangeStream_PutObject> PutRequest = FConcertReplication_ChangeStream_PutObject::MakeFromInfo(ChangesToApply.Value);
-					if (!ensure(PutRequest))
-					{
-						continue;
-					}
 
-					// We want to append the other client's properties to the ones the target client already has.
-					// TODO UE-201166: This step would not be necessary if we had an append operation in FConcertReplication_ChangeStream_PutObject
-					const FConcertReplicatedObjectInfo* CurrentInfo = CurrentTargetClientState.ReplicatedObjects.Find(ObjectId.Object);
-					if (CurrentInfo)
+				// We want to append the other client's properties to the ones the target client already has.
+				// TODO UE-201166: This step would not be necessary if we had an append operation in FConcertReplication_ChangeStream_PutObject
+				const FConcertReplicatedObjectInfo* CurrentInfo = CurrentTargetClientState.ReplicatedObjects.Find(ObjectId.Object);
+				if (CurrentInfo)
+				{
+					for (const FConcertPropertyChain& PropertyChain : CurrentInfo->PropertySelection.ReplicatedProperties)
 					{
-						for (const FConcertPropertyChain& PropertyChain : CurrentInfo->PropertySelection.ReplicatedProperties)
-						{
-							PutRequest->Properties.ReplicatedProperties.Add(PropertyChain);
-						}
+						PutRequest->Properties.ReplicatedProperties.Add(PropertyChain);
 					}
-					
-					ObjectChanges.ObjectsToPut.Add(ObjectId, MoveTemp(*PutRequest));
+				}
+				
+				ObjectChanges.ObjectsToPut.Add(ObjectId, MoveTemp(*PutRequest));
 
-					// The reassigned to client will replicate at the highest rate of either its current rate or the highest from the source clients.
-					const FConcertObjectReplicationSettings& CurrentObjectFrequencySettings = CurrentFrequencySettings.GetSettingsFor(ObjectPath);
-					const FConcertObjectReplicationSettings HighestSourceFrequencySettings = GetHighestFrequency(ObjectPath, OldFrequencies);
-					if (CurrentObjectFrequencySettings < HighestSourceFrequencySettings)
-					{
-						FrequencyChanges.OverridesToAdd.Add(ObjectPath, HighestSourceFrequencySettings);
-					}
+				// The reassigned to client will replicate at the highest rate of either its current rate or the highest from the source clients.
+				const FConcertObjectReplicationSettings& CurrentObjectFrequencySettings = CurrentFrequencySettings.GetSettingsFor(ObjectPath);
+				const FConcertObjectReplicationSettings HighestSourceFrequencySettings = GetHighestFrequency(ObjectPath, OldFrequencies);
+				if (CurrentObjectFrequencySettings < HighestSourceFrequencySettings)
+				{
+					FrequencyChanges.OverridesToAdd.Add(ObjectPath, HighestSourceFrequencySettings);
 				}
 			}
-			return { ObjectChanges, FrequencyChanges };
 		}
+		return { ObjectChanges, FrequencyChanges };
+	}
 
-		/** Builds list of authority changes based on which clients we managed to change successfully. */
-		static FConcertReplication_ChangeAuthority_Request MakeAuthorityRequestFrom(
-			const TMap<FGuid, TArray<FSoftObjectPath>>& ReassignedAuthority, 
-			const FParallelExecutionResult& ParallelExecutionResult,
-			const FGuid& TargetStreamId
-			)
+	/** Builds list of authority changes based on which clients we managed to change successfully. */
+	static FConcertReplication_ChangeAuthority_Request MakeAuthorityRequestFrom(
+		const TMap<FGuid, TArray<FSoftObjectPath>>& ReassignedAuthority, 
+		const FParallelExecutionResult& ParallelExecutionResult,
+		const FGuid& TargetStreamId
+		)
+	{
+		FConcertReplication_ChangeAuthority_Request ResultRequest;
+		for (const TPair<FGuid, FSubmitStreamChangesResponse>& ChangesRequestedOnRemote : ParallelExecutionResult.StreamResponses)
 		{
-			FConcertReplication_ChangeAuthority_Request ResultRequest;
-			for (const TPair<FGuid, FSubmitStreamChangesResponse>& ChangesRequestedOnRemote : ParallelExecutionResult.StreamResponses)
+			const TOptional<FCompletedChangeSubmission>& SubmissionInfo = ChangesRequestedOnRemote.Value.SubmissionInfo;
+			if (!SubmissionInfo || SubmissionInfo->Response.IsFailure())
 			{
-				const TOptional<FCompletedChangeSubmission>& SubmissionInfo = ChangesRequestedOnRemote.Value.SubmissionInfo;
-				if (!SubmissionInfo || SubmissionInfo->Response.IsFailure())
-				{
-					continue;
-				}
+				continue;
+			}
 
-				const TArray<FSoftObjectPath>* PreviousAuthority = ReassignedAuthority.Find(ChangesRequestedOnRemote.Key);
-				if (!PreviousAuthority)
-				{
-					continue;
-				}
+			const TArray<FSoftObjectPath>* PreviousAuthority = ReassignedAuthority.Find(ChangesRequestedOnRemote.Key);
+			if (!PreviousAuthority)
+			{
+				continue;
+			}
 
-				for (const FSoftObjectPath& ObjectPath : *PreviousAuthority)
-				{
-					ResultRequest.TakeAuthority.Add(ObjectPath, FConcertStreamArray{{ TargetStreamId }});
-				}
+			for (const FSoftObjectPath& ObjectPath : *PreviousAuthority)
+			{
+				ResultRequest.TakeAuthority.Add(ObjectPath, FConcertStreamArray{{ TargetStreamId }});
 			}
-			return ResultRequest;
 		}
+		return ResultRequest;
+	}
 
-		/** @return Whether TargetClientObjects includes all properties ClientsToConsider have registered for ObjectToCheck. */
-		static bool DoesTargetIncludeOtherClientsContent_SingleObject(
-			const FConcertObjectReplicationMap& TargetClientObjects,
-			const FOnlineClientManager& ClientManager,
-			const FSoftObjectPath& ObjectToCheck
-			)
-		{
-			const FConcertReplicatedObjectInfo* ObjectInfo = TargetClientObjects.ReplicatedObjects.Find(ObjectToCheck);
+	/** @return Whether TargetClientObjects includes all properties ClientsToConsider have registered for ObjectToCheck. */
+	static bool DoesTargetIncludeOtherClientsContent_SingleObject(
+		const FConcertObjectReplicationMap& TargetClientObjects,
+		const FOnlineClientManager& ClientManager,
+		const FSoftObjectPath& ObjectToCheck
+		)
+	{
+		const FConcertReplicatedObjectInfo* ObjectInfo = TargetClientObjects.ReplicatedObjects.Find(ObjectToCheck);
 
-			bool bIncludesAllOtherClients = true;
-			ClientManager.ForEachClient([&ObjectToCheck, ObjectInfo, &bIncludesAllOtherClients](const FOnlineClient& OtherClient)
+		bool bIncludesAllOtherClients = true;
+		ClientManager.ForEachClient([&ObjectToCheck, ObjectInfo, &bIncludesAllOtherClients](const FOnlineClient& OtherClient)
+		{
+			const FConcertObjectReplicationMap& OtherClientObjects = OtherClient.GetStreamSynchronizer().GetServerState();
+			const FConcertReplicatedObjectInfo* OtherObjectInfo = OtherClientObjects.ReplicatedObjects.Find(ObjectToCheck);
+				
+			// TargetClient has at least as much if OtherClient has nothing
+			if (!OtherObjectInfo || OtherObjectInfo->PropertySelection.ReplicatedProperties.IsEmpty())
 			{
-				const FConcertObjectReplicationMap& OtherClientObjects = OtherClient.GetStreamSynchronizer().GetServerState();
-				const FConcertReplicatedObjectInfo* OtherObjectInfo = OtherClientObjects.ReplicatedObjects.Find(ObjectToCheck);
-					
-				// TargetClient has at least as much if OtherClient has nothing
-				if (!OtherObjectInfo || OtherObjectInfo->PropertySelection.ReplicatedProperties.IsEmpty())
-				{
-					return EBreakBehavior::Continue;
-				}
+				return EBreakBehavior::Continue;
+			}
 
-				const bool bTargetContainsOther = ObjectInfo && ObjectInfo->PropertySelection.Includes(OtherObjectInfo->PropertySelection);
-				bIncludesAllOtherClients &= bTargetContainsOther;
-				return bIncludesAllOtherClients ? EBreakBehavior::Continue : EBreakBehavior::Break;
-			});
+			const bool bTargetContainsOther = ObjectInfo && ObjectInfo->PropertySelection.Includes(OtherObjectInfo->PropertySelection);
+			bIncludesAllOtherClients &= bTargetContainsOther;
+			return bIncludesAllOtherClients ? EBreakBehavior::Continue : EBreakBehavior::Break;
+		});
 
-			return bIncludesAllOtherClients;
+		return bIncludesAllOtherClients;
+	}
+	
+	/** @return Whether ClientId includes all properties assigned to ObjectsToCheck as all other clients do. */
+	static bool DoesTargetIncludeOtherClientsContent(
+		const FGuid& TargetClientId,
+		TConstArrayView<FSoftObjectPath> ObjectsToCheck,
+		const FOnlineClientManager& ClientManager
+		)
+	{
+		const FOnlineClient* TargetClient = ClientManager.FindClient(TargetClientId);
+		if (!ensure(TargetClient))
+		{
+			return false;
 		}
-		
-		/** @return Whether ClientId includes all properties assigned to ObjectsToCheck as all other clients do. */
-		static bool DoesTargetIncludeOtherClientsContent(
-			const FGuid& TargetClientId,
-			TConstArrayView<FSoftObjectPath> ObjectsToCheck,
-			const FOnlineClientManager& ClientManager
-			)
+
+		const FConcertObjectReplicationMap& TargetClientObjects = TargetClient->GetStreamSynchronizer().GetServerState();
+		return Algo::AllOf(ObjectsToCheck, [&ClientManager, &TargetClientObjects](const FSoftObjectPath& ObjectToCheck)
 		{
-			const FOnlineClient* TargetClient = ClientManager.FindClient(TargetClientId);
-			if (!ensure(TargetClient))
-			{
-				return false;
-			}
+			return DoesTargetIncludeOtherClientsContent_SingleObject(TargetClientObjects, ClientManager, ObjectToCheck);
+		});
+	}
 
-			const FConcertObjectReplicationMap& TargetClientObjects = TargetClient->GetStreamSynchronizer().GetServerState();
-			return Algo::AllOf(ObjectsToCheck, [&ClientManager, &TargetClientObjects](const FSoftObjectPath& ObjectToCheck)
-			{
-				return DoesTargetIncludeOtherClientsContent_SingleObject(TargetClientObjects, ClientManager, ObjectToCheck);
-			});
+	/** @return Whether ClientId has any properties assigned to any of the objects in ObjectsToCheck. */
+	static bool HasAnyProperties(const FGuid& ClientId, TConstArrayView<FSoftObjectPath> ObjectsToCheck, const FOnlineClientManager& ClientManager)
+	{
+		const FOnlineClient* TargetClient = ClientManager.FindClient(ClientId);
+		if (!ensure(TargetClient))
+		{
+			return false;
 		}
 
-		/** @return Whether ClientId has any properties assigned to any of the objects in ObjectsToCheck. */
-		static bool HasAnyProperties(const FGuid& ClientId, TConstArrayView<FSoftObjectPath> ObjectsToCheck, const FOnlineClientManager& ClientManager)
+		const FConcertObjectReplicationMap& TargetClientObjects = TargetClient->GetStreamSynchronizer().GetServerState();
+		return Algo::AnyOf(ObjectsToCheck, [&TargetClientObjects](const FSoftObjectPath& ObjectPath)
 		{
-			const FOnlineClient* TargetClient = ClientManager.FindClient(ClientId);
-			if (!ensure(TargetClient))
-			{
-				return false;
-			}
+			const FConcertReplicatedObjectInfo* ObjectInfo = TargetClientObjects.ReplicatedObjects.Find(ObjectPath);
+			return ObjectInfo && !ObjectInfo->PropertySelection.ReplicatedProperties.IsEmpty();
+		});
+	}
 
-			const FConcertObjectReplicationMap& TargetClientObjects = TargetClient->GetStreamSynchronizer().GetServerState();
-			return Algo::AnyOf(ObjectsToCheck, [&TargetClientObjects](const FSoftObjectPath& ObjectPath)
-			{
-				const FConcertReplicatedObjectInfo* ObjectInfo = TargetClientObjects.ReplicatedObjects.Find(ObjectPath);
-				return ObjectInfo && !ObjectInfo->PropertySelection.ReplicatedProperties.IsEmpty();
-			});
-		}
+	template<typename TCallback> requires std::is_invocable_r_v<EBreakBehavior, TCallback, const FGuid&>
+	static void EnumerateOwningOnlineClients(
+		const FOnlineClientManager& ClientManager,
+		const FSoftObjectPath& ObjectPath,
+		TCallback&& Callback
+		) 
+	{
+		ClientManager.ForEachClient([&ObjectPath, &Callback](const FOnlineClient& ReplicationClient)
+		{
+			const bool bHasProperties = ReplicationClient.GetStreamSynchronizer().GetServerState().HasProperties(ObjectPath);
+			return bHasProperties ? Callback(ReplicationClient.GetEndpointId()) : EBreakBehavior::Continue;
+		});
 	}
-	
+}
+
+namespace UE::MultiUserClient::Replication
+{
 	FReassignObjectPropertiesLogic::FReassignObjectPropertiesLogic(FOnlineClientManager& InClientManager)
 		: ClientManager(InClientManager)
 	{
 		ClientManager.OnRemoteClientsChanged().AddRaw(this, &FReassignObjectPropertiesLogic::OnRemoteClientsChanged);
-		ClientManager.GetAuthorityCache().OnCacheChanged().AddRaw(this, &FReassignObjectPropertiesLogic::OnClientCacheChanged);
 	}
 
 	FReassignObjectPropertiesLogic::~FReassignObjectPropertiesLogic()
@@ -279,7 +292,6 @@ namespace UE::MultiUserClient::Replication
 		// Note that all of this clean up is not REALLY needed because the FReassignObjectPropertiesLogic is a member of ClientManager.
 		// However, we'll follow good RAII here in case that should ever change.
 		ClientManager.OnRemoteClientsChanged().RemoveAll(this);
-		ClientManager.GetAuthorityCache().OnCacheChanged().RemoveAll(this);
 
 		if (FSubmissionQueue* SubmissionQueue = FindSubmissionQueueForTargetClient())
 		{
@@ -287,24 +299,14 @@ namespace UE::MultiUserClient::Replication
 		}
 	}
 
-	void FReassignObjectPropertiesLogic::EnumerateClientOwnershipState(const FSoftObjectPath& ObjectPath, FProcessClientOwnership Callback) const
-	{
-		ClientManager.ForEachClient([&ObjectPath, &Callback](const FOnlineClient& ReplicationClient)
-		{
-			const bool bHasProperties = ReplicationClient.GetStreamSynchronizer().GetServerState().HasProperties(ObjectPath);
-			const EOwnershipState Ownership = bHasProperties ? EOwnershipState::HasObjectRegistered : EOwnershipState::NoOwnership;
-			return Callback(ReplicationClient.GetEndpointId(), Ownership);
-		});
-	}
-
 	bool FReassignObjectPropertiesLogic::OwnsAnyOf(TConstArrayView<FSoftObjectPath> Objects, const FGuid& TargetClientId) const
 	{
 		bool bOwnsAny = false;
 		for (int32 i = 0; i < Objects.Num() && !bOwnsAny; ++i)
 		{
-			EnumerateClientOwnershipState(Objects[i], [&bOwnsAny, &TargetClientId](const FGuid& ClientId, EOwnershipState Ownership)
+			ReassignObjectProperties::EnumerateOwningOnlineClients(ClientManager, Objects[i],[&bOwnsAny, &TargetClientId](const FGuid& ClientId)
 			{
-				bOwnsAny = TargetClientId == ClientId && Ownership == EOwnershipState::HasObjectRegistered;
+				bOwnsAny = true;
 				return bOwnsAny ? EBreakBehavior::Break : EBreakBehavior::Continue;
 			});
 		}
@@ -316,9 +318,9 @@ namespace UE::MultiUserClient::Replication
 		bool bOwnsAny = false;
 		for (int32 i = 0; i < Objects.Num() && !bOwnsAny; ++i)
 		{
-			EnumerateClientOwnershipState(Objects[i], [&bOwnsAny](const FGuid& ClientId, EOwnershipState Ownership)
+			ReassignObjectProperties::EnumerateOwningOnlineClients(ClientManager,Objects[i], [&bOwnsAny](const FGuid& ClientId)
 			{
-				bOwnsAny = Ownership == EOwnershipState::HasObjectRegistered;
+				bOwnsAny = true;
 				return bOwnsAny ? EBreakBehavior::Break : EBreakBehavior::Continue;
 			});
 		}
@@ -343,9 +345,9 @@ namespace UE::MultiUserClient::Replication
 		}
 
 		// Relevant after user has re-assigned all properties to a given client
-		if (ReassignObjectProperties::Private::DoesTargetIncludeOtherClientsContent(ClientId, ObjectsToReassign, ClientManager))
+		if (ReassignObjectProperties::DoesTargetIncludeOtherClientsContent(ClientId, ObjectsToReassign, ClientManager))
 		{
-			const bool bHasAnyProperties = ReassignObjectProperties::Private::HasAnyProperties(ClientId, ObjectsToReassign, ClientManager);
+			const bool bHasAnyProperties = ReassignObjectProperties::HasAnyProperties(ClientId, ObjectsToReassign, ClientManager);
 			SET_REASON(
 				FText::Format(
 					LOCTEXT("Reason.NothingToAssignFmt", "Nothing to assign: {0}"),
@@ -373,7 +375,7 @@ namespace UE::MultiUserClient::Replication
 		const IClientStreamSynchronizer& StreamSynchronizer = ClientToAssignTo->GetStreamSynchronizer();
 		const FGuid& TargetStreamId = StreamSynchronizer.GetStreamId();
 		auto[OldRegisteredObjects, OldFrequencies, ReassignedClientRequests, ReassignedAuthority]
-			= ReassignObjectProperties::Private::BuildChangesForTransferal(
+			= ReassignObjectProperties::BuildChangesForTransferal(
 				ClientManager,
 				ObjectsToReassign,
 				TargetStreamId
@@ -432,8 +434,6 @@ namespace UE::MultiUserClient::Replication
 				);
 			InProgressOperation.Reset();
 		}
-
-		BroadcastOwnershipChanged();
 	}
 
 	void FReassignObjectPropertiesLogic::OnAssignedFromClientsCompleted(FParallelExecutionResult ParallelResult)
@@ -490,7 +490,7 @@ namespace UE::MultiUserClient::Replication
 		const IClientStreamSynchronizer& StreamSynchronizer = ClientToAssignTo->GetStreamSynchronizer();
 		const FGuid& StreamId = StreamSynchronizer.GetStreamId();
 		const bool bIsStreamRegistered = !StreamSynchronizer.GetServerState().IsEmpty();
-		auto[ObjectChanges, FrequencyChanges] = ReassignObjectProperties::Private::MakeChangelistFromAppliedChanges(
+		auto[ObjectChanges, FrequencyChanges] = ReassignObjectProperties::MakeChangelistFromAppliedChanges(
 			InProgressOperation->OldRegisteredObjects,
 			InProgressOperation->OldFrequencies,
 			ParallelResult,
@@ -502,7 +502,7 @@ namespace UE::MultiUserClient::Replication
 		
 		const TSharedPtr<ISubmissionOperation> SubmitOperation = Workflow.SubmitChanges({
 			MoveTemp(AssignedToClientRequest),
-			ReassignObjectProperties::Private::MakeAuthorityRequestFrom(InProgressOperation->OldAuthority, ParallelResult, StreamId)
+			ReassignObjectProperties::MakeAuthorityRequestFrom(InProgressOperation->OldAuthority, ParallelResult, StreamId)
 		});
 		if (!ensure(SubmitOperation))
 		{
diff --git a/Engine/Plugins/Developer/Concert/ConcertApp/MultiUserClient/Source/MultiUserClient/Private/Replication/Submission/MultiEdit/ReassignObjectPropertiesLogic.h b/Engine/Plugins/Developer/Concert/ConcertApp/MultiUserClient/Source/MultiUserClient/Private/Replication/Submission/MultiEdit/ReassignObjectPropertiesLogic.h
--- a/Engine/Plugins/Developer/Concert/ConcertApp/MultiUserClient/Source/MultiUserClient/Private/Replication/Submission/MultiEdit/ReassignObjectPropertiesLogic.h
+++ b/Engine/Plugins/Developer/Concert/ConcertApp/MultiUserClient/Source/MultiUserClient/Private/Replication/Submission/MultiEdit/ReassignObjectPropertiesLogic.h
@@ -44,8 +44,6 @@ namespace UE::MultiUserClient::Replication
 		FReassignObjectPropertiesLogic(FOnlineClientManager& InClientManager);
 		virtual ~FReassignObjectPropertiesLogic() override;
 
-		/** Enumerates every client that has registered properties to this object. */
-		void EnumerateClientOwnershipState(const FSoftObjectPath& ObjectPath, FProcessClientOwnership Callback) const;
 		/** @return Whether ClientId has ownership over at least one */
 		bool OwnsAnyOf(TConstArrayView<FSoftObjectPath> Objects, const FGuid& TargetClientId) const;
 		/** @return Whether any of the following objects has at least one owner */
@@ -61,9 +59,6 @@ namespace UE::MultiUserClient::Replication
 		/** @return The local time at which reassignment operation was started. Useful for UI showing throbber but only if operation takes long.*/
 		TOptional<FDateTime> GetTimeReassignmentWasStarted() const;
 
-		/** Broadcast when the result of EnumerateClientOwnershipState may have changed. */
-		FSimpleMulticastDelegate& OnOwnershipChanged() { return OnOwnershipChangedDelegate; }
-
 	private:
 
 		/** Detects when "this" is destroyed to avoid dereferencing it in latent actions. */
@@ -107,8 +102,6 @@ namespace UE::MultiUserClient::Replication
 		FSimpleMulticastDelegate OnOwnershipChangedDelegate;
 
 		void OnRemoteClientsChanged();
-		void OnClientCacheChanged(const FGuid&) const { BroadcastOwnershipChanged(); }
-		void BroadcastOwnershipChanged() const { OnOwnershipChangedDelegate.Broadcast(); }
 
 		/** Called once all clients being reassigned have remove their object bindings. The target client will now receive all the previous bindings. */
 		void OnAssignedFromClientsCompleted(FParallelExecutionResult ParallelResult);
diff --git a/Engine/Plugins/Developer/Concert/ConcertApp/MultiUserClient/Source/MultiUserClient/Private/Widgets/ActiveSession/Replication/Client/Multi/Columns/AssignedClients/AssignedClientsModel.cpp b/Engine/Plugins/Developer/Concert/ConcertApp/MultiUserClient/Source/MultiUserClient/Private/Widgets/ActiveSession/Replication/Client/Multi/Columns/AssignedClients/AssignedClientsModel.cpp
new file mode 100644
--- /dev/null
+++ b/Engine/Plugins/Developer/Concert/ConcertApp/MultiUserClient/Source/MultiUserClient/Private/Widgets/ActiveSession/Replication/Client/Multi/Columns/AssignedClients/AssignedClientsModel.cpp
@@ -0,0 +1,74 @@
+﻿// Copyright Epic Games, Inc. All Rights Reserved.
+
+#include "AssignedClientsModel.h"
+
+#include "Replication/Client/Online/OnlineClient.h"
+#include "Replication/Client/Online/OnlineClientManager.h"
+#include "Replication/Editor/Model/Object/IObjectHierarchyModel.h"
+
+#include "Containers/Array.h"
+#include "Misc/Guid.h"
+
+namespace UE::MultiUserClient::Replication::MultiStreamColumns::Private
+{
+	template<typename TCallback> requires std::is_invocable_v<TCallback, const FGuid&>
+	static void EnumerateOwningOnlineClients(
+		const FOnlineClientManager& ClientManager,
+		const FSoftObjectPath& ObjectPath,
+		TCallback&& Callback
+		) 
+	{
+		ClientManager.ForEachClient([&ObjectPath, &Callback](const FOnlineClient& ReplicationClient)
+		{
+			const bool bHasProperties = ReplicationClient.GetStreamSynchronizer().GetServerState().HasProperties(ObjectPath);
+			if (bHasProperties)
+			{
+				Callback(ReplicationClient.GetEndpointId());
+			}
+			return EBreakBehavior::Continue;
+		});
+	}
+}
+
+namespace UE::MultiUserClient::Replication::MultiStreamColumns
+{
+	FAssignedClientsModel::FAssignedClientsModel(
+		const ConcertSharedSlate::IObjectHierarchyModel& ObjectHierarchy,
+		FOnlineClientManager& InClientManager
+		)
+		: ObjectHierarchy(ObjectHierarchy)
+		, ClientManager(InClientManager)
+	{
+		ClientManager.OnRemoteClientsChanged().AddRaw(this, &FAssignedClientsModel::BroadcastOwnershipChanged);
+		ClientManager.GetAuthorityCache().OnCacheChanged().AddRaw(this, &FAssignedClientsModel::OnClientChanged);
+	}
+
+	FAssignedClientsModel::~FAssignedClientsModel()
+	{
+		ClientManager.OnRemoteClientsChanged().RemoveAll(this);
+		ClientManager.GetAuthorityCache().OnCacheChanged().RemoveAll(this);
+	}
+
+	TArray<FGuid> FAssignedClientsModel::GetAssignedClients(const FSoftObjectPath& ObjectPath) const
+	{
+		TArray<FGuid> ClientsWithOwnership;
+		const auto ProcessObject = [this, &ClientsWithOwnership](const FSoftObjectPath& ObjectPath)
+		{
+			Private::EnumerateOwningOnlineClients(ClientManager, ObjectPath, [&ClientsWithOwnership](const FGuid& ClientId)
+			{
+				ClientsWithOwnership.AddUnique(ClientId);
+			});
+		};
+			
+		ProcessObject(ObjectPath);
+		ObjectHierarchy.ForEachChildRecursive(
+			TSoftObjectPtr{ ObjectPath },
+			[&ProcessObject](const TSoftObjectPtr<>&, const TSoftObjectPtr<>& ChildObject, ConcertSharedSlate::EChildRelationship)
+			{
+				ProcessObject(ChildObject.GetUniqueID());
+				return EBreakBehavior::Continue;
+			});
+			
+		return ClientsWithOwnership;
+	}
+}
diff --git a/Engine/Plugins/Developer/Concert/ConcertApp/MultiUserClient/Source/MultiUserClient/Private/Widgets/ActiveSession/Replication/Client/Multi/Columns/AssignedClients/AssignedClientsModel.h b/Engine/Plugins/Developer/Concert/ConcertApp/MultiUserClient/Source/MultiUserClient/Private/Widgets/ActiveSession/Replication/Client/Multi/Columns/AssignedClients/AssignedClientsModel.h
new file mode 100644
--- /dev/null
+++ b/Engine/Plugins/Developer/Concert/ConcertApp/MultiUserClient/Source/MultiUserClient/Private/Widgets/ActiveSession/Replication/Client/Multi/Columns/AssignedClients/AssignedClientsModel.h
@@ -0,0 +1,50 @@
+﻿// Copyright Epic Games, Inc. All Rights Reserved.
+
+#pragma once
+
+#include "Containers/ContainersFwd.h"
+#include "Delegates/Delegate.h"
+#include "HAL/Platform.h"
+#include "Templates/UnrealTemplate.h"
+
+struct FGuid;
+
+namespace UE::ConcertSharedSlate { class IObjectHierarchyModel; }
+namespace UE::MultiUserClient::Replication
+{
+	class FOnlineClientManager;
+	class FReassignObjectPropertiesLogic;
+}
+
+namespace UE::MultiUserClient::Replication::MultiStreamColumns
+{
+	/** Implements the model in the MVC pattern for the assigned clients column. */
+	class FAssignedClientsModel : public FNoncopyable
+	{
+	public:
+		
+		FAssignedClientsModel(
+			const ConcertSharedSlate::IObjectHierarchyModel& ObjectHierarchy UE_LIFETIMEBOUND,
+			FOnlineClientManager& InClientManager UE_LIFETIMEBOUND
+			);
+		~FAssignedClientsModel();
+
+		TArray<FGuid> GetAssignedClients(const FSoftObjectPath& ObjectPath) const;
+		
+		/** Broadcast when the result of GetAssignedClients may have changed. */
+		FSimpleMulticastDelegate& OnOwnershipChanged() { return OnOwnershipChangedDelegate; }
+
+	private:
+
+		/** Used to compute clients's recursive ownership of subobjects for a managed object. */
+		const ConcertSharedSlate::IObjectHierarchyModel& ObjectHierarchy;
+		FOnlineClientManager& ClientManager;
+		
+		/** Broadcast when the result of GetAssignedClients may have changed. */
+		FSimpleMulticastDelegate OnOwnershipChangedDelegate;
+		
+		void OnClientChanged(const FGuid&) const { BroadcastOwnershipChanged(); }
+		void BroadcastOwnershipChanged() const { OnOwnershipChangedDelegate.Broadcast(); }
+	};
+}
+
diff --git a/Engine/Plugins/Developer/Concert/ConcertApp/MultiUserClient/Source/MultiUserClient/Private/Widgets/ActiveSession/Replication/Client/Multi/Columns/AssignedClients/MultiStreamColumns_AssignedClients.cpp b/Engine/Plugins/Developer/Concert/ConcertApp/MultiUserClient/Source/MultiUserClient/Private/Widgets/ActiveSession/Replication/Client/Multi/Columns/AssignedClients/MultiStreamColumns_AssignedClients.cpp
new file mode 100644
--- /dev/null
+++ b/Engine/Plugins/Developer/Concert/ConcertApp/MultiUserClient/Source/MultiUserClient/Private/Widgets/ActiveSession/Replication/Client/Multi/Columns/AssignedClients/MultiStreamColumns_AssignedClients.cpp
@@ -0,0 +1,121 @@
+// Copyright Epic Games, Inc. All Rights Reserved.
+
+#include "AssignedClientsModel.h"
+
+#include "IConcertClient.h"
+#include "MultiUserReplicationStyle.h"
+#include "Replication/Editor/View/IMultiReplicationStreamEditor.h"
+#include "SAssignedClientsWidget.h"
+#include "Widgets/ActiveSession/Replication/Client/ClientUtils.h"
+#include "Widgets/ActiveSession/Replication/Client/Multi/Columns/MultiStreamColumns.h"
+#include "Widgets/Client/ClientInfoHelpers.h"
+#include "Widgets/Client/SHorizontalClientList.h"
+
+#include "Widgets/Text/STextBlock.h"
+
+#define LOCTEXT_NAMESPACE "AssignedClientsColumnId"
+
+namespace UE::MultiUserClient::Replication::MultiStreamColumns
+{
+	const FName AssignedClientsColumnId(TEXT("AssignedClientsColumn"));
+
+	ConcertSharedSlate::FObjectColumnEntry AssignedClientsColumn(
+		TSharedRef<IConcertClient> ConcertClient,
+		TAttribute<TSharedPtr<ConcertSharedSlate::IMultiReplicationStreamEditor>> MultiStreamModelAttribute,
+		const ConcertSharedSlate::IObjectHierarchyModel& ObjectHierarchy,
+		FOnlineClientManager& ClientManager,
+		const int32 ColumnsSortPriority
+		)
+	{
+		class FObjectColumn_AssignedClients : public ConcertSharedSlate::IObjectTreeColumn
+		{
+		public:
+
+			FObjectColumn_AssignedClients(
+				TSharedRef<IConcertClient> ConcertClient,
+				TAttribute<TSharedPtr<ConcertSharedSlate::IMultiReplicationStreamEditor>> MultiStreamModelAttribute,
+				const ConcertSharedSlate::IObjectHierarchyModel& ObjectHierarchy UE_LIFETIMEBOUND,
+				FOnlineClientManager& ClientManager UE_LIFETIMEBOUND
+				)
+				: ConcertClient(MoveTemp(ConcertClient))
+				, MultiStreamModelAttribute(MoveTemp(MultiStreamModelAttribute))
+				, Model(ObjectHierarchy, ClientManager)
+			{}
+			
+			virtual SHeaderRow::FColumn::FArguments CreateHeaderRowArgs() const override
+			{
+				return SHeaderRow::Column(AssignedClientsColumnId)
+					.DefaultLabel(LOCTEXT("Owner.Label", "Author"))
+					.ToolTipText(LOCTEXT("Owner.ToolTip", "Clients that have registered properties for the object"))
+					.FillSized(FMultiUserReplicationStyle::Get()->GetFloat(TEXT("AllClients.Object.OwnerSize")));
+			}
+			
+			virtual TSharedRef<SWidget> GenerateColumnWidget(const FBuildArgs& InArgs) override
+			{
+				return SNew(SAssignedClientsWidget, ConcertClient, Model)
+					.ManagedObject(InArgs.RowItem.RowData.GetObjectPath())
+					.HighlightText_Lambda([HighlightText = InArgs.HighlightText](){ return HighlightText ? *HighlightText : FText::GetEmpty(); });
+			}
+			virtual void PopulateSearchString(const ConcertSharedSlate::FObjectTreeRowContext& InItem, TArray<FString>& InOutSearchStrings) const override
+			{
+				for (const FGuid& ClientId : Model.GetAssignedClients(InItem.RowData.GetObjectPath()))
+				{
+					InOutSearchStrings.Add(ClientUtils::GetClientDisplayName(*ConcertClient->GetCurrentSession(), ClientId));
+				}
+			}
+
+			virtual bool CanBeSorted() const override { return true; }
+			virtual bool IsLessThan(const ConcertSharedSlate::FObjectTreeRowContext& Left, const ConcertSharedSlate::FObjectTreeRowContext& Right) const override
+			{
+				const TOptional<FString> LeftClientDisplayString = GetDisplayString(Left.RowData.GetObjectPath());
+				const TOptional<FString> RightClientDisplayString = GetDisplayString(Right.RowData.GetObjectPath());
+			
+				if (LeftClientDisplayString && RightClientDisplayString)
+				{
+					return *LeftClientDisplayString < *RightClientDisplayString;
+				}
+				// Our rule: set < unset. This way unassigned appears last.
+				return LeftClientDisplayString.IsSet() && !RightClientDisplayString.IsSet();
+			}
+
+		private:
+			
+			const TSharedRef<IConcertClient> ConcertClient;
+			const TAttribute<TSharedPtr<ConcertSharedSlate::IMultiReplicationStreamEditor>> MultiStreamModelAttribute;
+
+			/** The model that the view displays; the model of the MVC pattern. */
+			FAssignedClientsModel Model;
+			
+			TOptional<FString> GetDisplayString(const FSoftObjectPath& ManagedObject) const
+			{
+				using SWidgetType = ConcertSharedSlate::SHorizontalClientList;
+				const TArray<FGuid> Clients = Model.GetAssignedClients(ManagedObject);
+
+				const ConcertSharedSlate::FGetClientParenthesesContent GetParenthesesContent =
+					ConcertClientSharedSlate::MakeGetLocalClientParenthesesContent(ConcertClient);
+				const auto SortPredicate = [&GetParenthesesContent](const FConcertSessionClientInfo& Left, const FConcertSessionClientInfo& Right)
+				{
+					return ConcertSharedSlate::SortLocalClientParenthesesFirstThenThenAlphabetical(Left, Right, GetParenthesesContent);
+				};
+				return SWidgetType::GetDisplayString(
+					Clients,
+					ConcertClientSharedSlate::MakeClientInfoGetter(ConcertClient),
+					ConcertSharedSlate::FClientSortPredicate::CreateLambda(SortPredicate),
+					GetParenthesesContent
+					);
+			}
+		};
+
+		return {
+			ConcertSharedSlate::TReplicationColumnDelegates<ConcertSharedSlate::FObjectTreeRowContext>::FCreateColumn::CreateLambda(
+				[ConcertClient = MoveTemp(ConcertClient), MultiStreamModelAttribute = MoveTemp(MultiStreamModelAttribute), &ObjectHierarchy, &ClientManager]()
+				{
+					return MakeShared<FObjectColumn_AssignedClients>(ConcertClient, MultiStreamModelAttribute, ObjectHierarchy, ClientManager);
+				}),
+			AssignedClientsColumnId,
+			{ ColumnsSortPriority }
+		};
+	}
+}
+
+#undef LOCTEXT_NAMESPACE
\ No newline at end of file
diff --git a/Engine/Plugins/Developer/Concert/ConcertApp/MultiUserClient/Source/MultiUserClient/Private/Widgets/ActiveSession/Replication/Client/Multi/Columns/AssignedClients/SAssignedClientsWidget.cpp b/Engine/Plugins/Developer/Concert/ConcertApp/MultiUserClient/Source/MultiUserClient/Private/Widgets/ActiveSession/Replication/Client/Multi/Columns/AssignedClients/SAssignedClientsWidget.cpp
new file mode 100644
--- /dev/null
+++ b/Engine/Plugins/Developer/Concert/ConcertApp/MultiUserClient/Source/MultiUserClient/Private/Widgets/ActiveSession/Replication/Client/Multi/Columns/AssignedClients/SAssignedClientsWidget.cpp
@@ -0,0 +1,56 @@
+﻿// Copyright Epic Games, Inc. All Rights Reserved.
+
+#include "SAssignedClientsWidget.h"
+
+#include "AssignedClientsModel.h"
+#include "Widgets/Client/ClientInfoHelpers.h"
+#include "Widgets/Client/SHorizontalClientList.h"
+
+#include "Widgets/Text/STextBlock.h"
+
+#define LOCTEXT_NAMESPACE "SAssignedClientsWidget"
+
+namespace UE::MultiUserClient::Replication::MultiStreamColumns
+{
+	void SAssignedClientsWidget::Construct(
+		const FArguments& InArgs,
+		const TSharedRef<IConcertClient>& InConcertClient,
+		FAssignedClientsModel& InModel
+		)
+	{
+		Model = &InModel;
+		ManagedObject = InArgs._ManagedObject;
+		
+		ChildSlot
+		[
+			SAssignNew(ClientList, ConcertSharedSlate::SHorizontalClientList)
+			.GetClientParenthesesContent(ConcertClientSharedSlate::MakeGetLocalClientParenthesesContent(InConcertClient))
+			.GetClientInfo(ConcertClientSharedSlate::MakeClientInfoGetter(InConcertClient))
+			.Font(FAppStyle::GetFontStyle("PropertyWindow.NormalFont"))
+			.HighlightText(InArgs._HighlightText)
+			.ListToolTipText(LOCTEXT("Clients.ToolTip", "These clients will replicate their assigned properties when replication is active.\nYou can pause & resume replication at the beginnig of this row."))
+			.EmptyListSlot()
+			[
+				SNew(STextBlock)
+				.Text(LOCTEXT("NoClients.Label", "No assigned properties"))
+				.ToolTipText(LOCTEXT("NoClients.ToolTip", "Click this row and then assign the properties to the client that should replicate them."))
+				.ColorAndOpacity(FSlateColor::UseSubduedForeground())
+			]
+		];
+		
+		RefreshClientList();
+		Model->OnOwnershipChanged().AddSP(this, &SAssignedClientsWidget::RefreshClientList);
+	}
+	
+	SAssignedClientsWidget::~SAssignedClientsWidget()
+	{
+		Model->OnOwnershipChanged().RemoveAll(this);
+	}
+
+	void SAssignedClientsWidget::RefreshClientList() const
+	{
+		ClientList->RefreshList(Model->GetAssignedClients(ManagedObject));
+	}
+}
+
+#undef LOCTEXT_NAMESPACE
\ No newline at end of file
diff --git a/Engine/Plugins/Developer/Concert/ConcertApp/MultiUserClient/Source/MultiUserClient/Private/Widgets/ActiveSession/Replication/Client/Multi/Columns/AssignedClients/SAssignedClientsWidget.h b/Engine/Plugins/Developer/Concert/ConcertApp/MultiUserClient/Source/MultiUserClient/Private/Widgets/ActiveSession/Replication/Client/Multi/Columns/AssignedClients/SAssignedClientsWidget.h
new file mode 100644
--- /dev/null
+++ b/Engine/Plugins/Developer/Concert/ConcertApp/MultiUserClient/Source/MultiUserClient/Private/Widgets/ActiveSession/Replication/Client/Multi/Columns/AssignedClients/SAssignedClientsWidget.h
@@ -0,0 +1,54 @@
+﻿// Copyright Epic Games, Inc. All Rights Reserved.
+
+#pragma once
+
+#include "Widgets/DeclarativeSyntaxSupport.h"
+#include "Widgets/SCompoundWidget.h"
+
+class IConcertClient;
+
+namespace UE::ConcertSharedSlate
+{
+	class IObjectHierarchyModel;
+	class SHorizontalClientList;
+}
+
+namespace UE::MultiUserClient::Replication { class FReassignObjectPropertiesLogic; }
+
+namespace UE::MultiUserClient::Replication::MultiStreamColumns
+{
+	class FAssignedClientsModel;
+
+	/** Placed into object rows. Display the clients that own it and its children. */
+	class SAssignedClientsWidget : public SCompoundWidget
+	{
+		
+	public:
+	
+		SLATE_BEGIN_ARGS(SAssignedClientsWidget){}
+			SLATE_ARGUMENT(FSoftObjectPath, ManagedObject)
+			SLATE_ATTRIBUTE(FText, HighlightText)
+		SLATE_END_ARGS()
+
+		void Construct(
+			const FArguments& InArgs,
+			const TSharedRef<IConcertClient>& InConcertClient,
+			FAssignedClientsModel& InModel UE_LIFETIMEBOUND
+			);
+	
+		virtual ~SAssignedClientsWidget() override;
+
+		void RefreshClientList() const;
+
+	private:
+
+		/** The model this view displays. */
+		FAssignedClientsModel* Model = nullptr;
+		
+		/** This widget displays the owning clients and is refreshed when the list of owning clients changes. */
+		TSharedPtr<ConcertSharedSlate::SHorizontalClientList> ClientList;
+		
+		/** The object for which we're displaying the owning clients */
+		FSoftObjectPath ManagedObject;
+	};
+}
diff --git a/Engine/Plugins/Developer/Concert/ConcertApp/MultiUserClient/Source/MultiUserClient/Private/Widgets/ActiveSession/Replication/Client/Multi/Columns/MultiStreamColumns.h b/Engine/Plugins/Developer/Concert/ConcertApp/MultiUserClient/Source/MultiUserClient/Private/Widgets/ActiveSession/Replication/Client/Multi/Columns/MultiStreamColumns.h
--- a/Engine/Plugins/Developer/Concert/ConcertApp/MultiUserClient/Source/MultiUserClient/Private/Widgets/ActiveSession/Replication/Client/Multi/Columns/MultiStreamColumns.h
+++ b/Engine/Plugins/Developer/Concert/ConcertApp/MultiUserClient/Source/MultiUserClient/Private/Widgets/ActiveSession/Replication/Client/Multi/Columns/MultiStreamColumns.h
@@ -57,7 +57,6 @@ namespace UE::MultiUserClient::Replication::MultiStreamColumns
 	 * @param ConcertClient Used to look up client names
 	 * @param MultiStreamModelAttribute Used for requesting resorting the column
 	 * @param ObjectHierarchy Used to display clients when a child object is replicated
-	 * @param ReassignmentLogic Used to get the assigned clients and detect when they change.
 	 * @param ClientManager Used to access all clients for display in the combo box drop-down
 	 * @param ColumnsSortPriority The order relative to the other columns
 	 */
@@ -65,8 +64,7 @@ namespace UE::MultiUserClient::Replication::MultiStreamColumns
 		TSharedRef<IConcertClient> ConcertClient,
 		TAttribute<TSharedPtr<ConcertSharedSlate::IMultiReplicationStreamEditor>> MultiStreamModelAttribute,
 		const ConcertSharedSlate::IObjectHierarchyModel& ObjectHierarchy UE_LIFETIMEBOUND,
-		FReassignObjectPropertiesLogic& ReassignmentLogic UE_LIFETIMEBOUND,
-		const FOnlineClientManager& ClientManager UE_LIFETIMEBOUND,
+		FOnlineClientManager& ClientManager UE_LIFETIMEBOUND,
 		int32 ColumnsSortPriority = static_cast<int32>(EColumnSortOrder::ReassignOwnership)
 		);
 	
diff --git a/Engine/Plugins/Developer/Concert/ConcertApp/MultiUserClient/Source/MultiUserClient/Private/Widgets/ActiveSession/Replication/Client/Multi/Columns/MultiStreamColumns_AssignedClients.cpp b/Engine/Plugins/Developer/Concert/ConcertApp/MultiUserClient/Source/MultiUserClient/Private/Widgets/ActiveSession/Replication/Client/Multi/Columns/MultiStreamColumns_AssignedClients.cpp
deleted file mode 100644
diff --git a/Engine/Plugins/Developer/Concert/ConcertApp/MultiUserClient/Source/MultiUserClient/Private/Widgets/ActiveSession/Replication/Client/Multi/SMultiClientView.cpp b/Engine/Plugins/Developer/Concert/ConcertApp/MultiUserClient/Source/MultiUserClient/Private/Widgets/ActiveSession/Replication/Client/Multi/SMultiClientView.cpp
--- a/Engine/Plugins/Developer/Concert/ConcertApp/MultiUserClient/Source/MultiUserClient/Private/Widgets/ActiveSession/Replication/Client/Multi/SMultiClientView.cpp
+++ b/Engine/Plugins/Developer/Concert/ConcertApp/MultiUserClient/Source/MultiUserClient/Private/Widgets/ActiveSession/Replication/Client/Multi/SMultiClientView.cpp
@@ -146,7 +146,7 @@ namespace UE::MultiUserClient::Replication
 			.ObjectColumns =
 			{
 				MultiStreamColumns::MuteToggleColumn(MuteManager.GetChangeTracker()),
-				MultiStreamColumns::AssignedClientsColumn(InConcertClient, MultiStreamEditorAttribute, *ObjectHierarchy, ClientManager->GetReassignmentLogic(), *ClientManager)
+				MultiStreamColumns::AssignedClientsColumn(InConcertClient, MultiStreamEditorAttribute, *ObjectHierarchy, *ClientManager)
 			},
 			.ShouldDisplayObjectDelegate = FShouldDisplayObject::CreateSP(this, &SMultiClientView::ShouldDisplayObject),
 		};
