Extract property assignment logic from column widget to new FAssignPropertyModell.

diff --git a/Engine/Plugins/Developer/Concert/ConcertApp/MultiUserClient/Source/MultiUserClient/Private/Replication/Submission/MultiEdit/ReassignObjectPropertiesLogic.cpp b/Engine/Plugins/Developer/Concert/ConcertApp/MultiUserClient/Source/MultiUserClient/Private/Replication/Submission/MultiEdit/ReassignObjectPropertiesLogic.cpp
--- a/Engine/Plugins/Developer/Concert/ConcertApp/MultiUserClient/Source/MultiUserClient/Private/Replication/Submission/MultiEdit/ReassignObjectPropertiesLogic.cpp
+++ b/Engine/Plugins/Developer/Concert/ConcertApp/MultiUserClient/Source/MultiUserClient/Private/Replication/Submission/MultiEdit/ReassignObjectPropertiesLogic.cpp
@@ -11,6 +11,8 @@
 #include "Algo/AnyOf.h"
 #include "Async/Async.h"
 
+#include <type_traits>
+
 #define LOCTEXT_NAMESPACE "FReassignObjectPropertiesLogic"
 
 namespace UE::MultiUserClient::Replication::ReassignObjectProperties
diff --git a/Engine/Plugins/Developer/Concert/ConcertApp/MultiUserClient/Source/MultiUserClient/Private/Widgets/ActiveSession/Replication/Client/Multi/Columns/AssignProperty/AssignPropertyModel.cpp b/Engine/Plugins/Developer/Concert/ConcertApp/MultiUserClient/Source/MultiUserClient/Private/Widgets/ActiveSession/Replication/Client/Multi/Columns/AssignProperty/AssignPropertyModel.cpp
new file mode 100644
--- /dev/null
+++ b/Engine/Plugins/Developer/Concert/ConcertApp/MultiUserClient/Source/MultiUserClient/Private/Widgets/ActiveSession/Replication/Client/Multi/Columns/AssignProperty/AssignPropertyModel.cpp
@@ -0,0 +1,221 @@
+﻿// Copyright Epic Games, Inc. All Rights Reserved.
+
+#include "AssignPropertyModel.h"
+
+#include "ConcertLogGlobal.h"
+#include "Replication/Client/Online/OnlineClient.h"
+#include "Replication/Client/Online/OnlineClientManager.h"
+#include "Replication/Editor/Model/PropertyUtils.h"
+
+#include "GameFramework/Actor.h"
+#include "ScopedTransaction.h"
+
+#include <type_traits>
+
+#define LOCTEXT_NAMESPACE "FAssignPropertyModel"
+
+namespace UE::MultiUserClient::Replication::MultiStreamColumns::AssignPropertyModel
+{
+	static void RemovePropertiesFromClient(
+		const FOnlineClient& ClientToRemoveFrom,
+		TConstArrayView<TSoftObjectPtr<>> Objects,
+		const FConcertPropertyChain& Property
+		)
+	{
+		const TSharedRef<ConcertSharedSlate::IEditableReplicationStreamModel> EditModel = ClientToRemoveFrom.GetClientEditModel();
+		for (const TSoftObjectPtr<>& Object : Objects)
+		{
+			const FSoftObjectPath& ObjectPath = Object.GetUniqueID();
+			const FSoftClassPath ClassPath = EditModel->GetObjectClass(ObjectPath);
+			EditModel->RemoveProperties(ObjectPath, { Property });
+					
+			if (EditModel->HasAnyPropertyAssigned(ObjectPath))
+			{
+				continue;
+			}
+
+			// We want to remove subobjects that have no properties. Retain actors because they cause their entire component / subobject hierarchy to be displayed.
+			// Skipping this check would close the entire property tree view and remove the actor hierarchy from the view.
+			// That would feel very unnatural / unexpected for the user. 
+			// If the user does not want the actor anymore, they should click it and delete it.
+			const UClass* ObjectClass = ClassPath.IsValid() ? ClassPath.TryLoadClass<UObject>() : nullptr;
+			UE_CLOG(ClassPath.IsValid() && !ObjectClass, LogConcert, Warning, TEXT("SAssignPropertyComboBox: Failed to resolve class %s"), *ClassPath.ToString());
+			const bool bIsTopLevelObject = ObjectClass && !ObjectClass->IsChildOf<AActor>();
+			if (bIsTopLevelObject)
+			{
+				EditModel->RemoveObjects({ ObjectPath });
+			}
+		}
+	}
+	
+	template<typename TShouldRemove>
+	requires std::is_invocable_r_v<bool, TShouldRemove, const FOnlineClient&>
+	static void UnassignPropertyFromClients(
+		const FOnlineClientManager& ClientManager,
+		TConstArrayView<TSoftObjectPtr<>> Objects,
+		const FConcertPropertyChain& Property,
+		TShouldRemove&& ShouldRemoveFromClient
+		)
+	{
+		ClientManager.ForEachClient([Objects, Property, &ShouldRemoveFromClient](const FOnlineClient& ClientToRemoveFrom)
+		{
+			if (ShouldRemoveFromClient(ClientToRemoveFrom))
+			{
+				RemovePropertiesFromClient(ClientToRemoveFrom, Objects, Property);
+			}
+			return EBreakBehavior::Continue;
+		});
+	}
+}
+
+namespace UE::MultiUserClient::Replication::MultiStreamColumns
+{
+	FAssignPropertyModel::FAssignPropertyModel(FOnlineClientManager& InClientManager)
+		: ClientManager(InClientManager)
+	{
+		ClientManager.OnRemoteClientsChanged().AddRaw(this, &FAssignPropertyModel::BroadcastOnOwnershipChanged);
+	}
+
+	FAssignPropertyModel::~FAssignPropertyModel()
+	{
+		ClientManager.OnRemoteClientsChanged().RemoveAll(this);
+	}
+
+#define SET_REASON(Text) if (Reason) { *Reason = Text; }
+	bool FAssignPropertyModel::CanChangePropertyFor(const FGuid& ClientId, FText* Reason) const
+	{
+		const FOnlineClient* Client = ClientManager.FindClient(ClientId);
+		// Remote clients can disconnect after the combo-box is opened.
+		if (!Client)
+		{
+			SET_REASON(LOCTEXT("ClientDisconnected", "Client disconnected."));
+			return false;
+		}
+		
+		return true;
+	}
+#undef SET_REASON
+
+	bool FAssignPropertyModel::CanClear(TConstArrayView<TSoftObjectPtr<>> Objects, const FConcertPropertyChain& Property) const
+	{
+		bool bIsAssignedToAnyClient = false;
+		ClientManager.ForEachClient([this, &Objects, &Property, &bIsAssignedToAnyClient](const FOnlineClient& Client)
+		{
+			for (const TSoftObjectPtr<>& EditedObject : Objects)
+			{
+				if (bIsAssignedToAnyClient)
+				{
+					break;
+				}
+				
+				const TSharedRef<ConcertSharedSlate::IEditableReplicationStreamModel> Model = Client.GetClientEditModel();
+				const bool bHasProperty = Model->HasProperty(EditedObject.GetUniqueID(), Property);
+				bIsAssignedToAnyClient |= bHasProperty;
+			}
+			
+			return bIsAssignedToAnyClient ? EBreakBehavior::Break : EBreakBehavior::Continue;
+		});
+		return bIsAssignedToAnyClient;
+	}
+
+	EPropertyOnObjectsOwnershipState FAssignPropertyModel::GetPropertyOwnershipState(const FGuid& ClientId, TConstArrayView<TSoftObjectPtr<>> Objects, const FConcertPropertyChain& Property) const
+	{
+		const FOnlineClient* Client = ClientManager.FindClient(ClientId);
+		// Remote clients can disconnect after the combo-box is opened.
+		if (!Client)
+		{
+			return {};
+		}
+
+		const TSharedRef<ConcertSharedSlate::IEditableReplicationStreamModel> Model = Client->GetClientEditModel();
+		EPropertyOnObjectsOwnershipState Result = EPropertyOnObjectsOwnershipState::Mixed;
+		for (const TSoftObjectPtr<>& ObjectPath : Objects)
+		{
+			const bool bHasProperty = Model->HasProperty(ObjectPath.GetUniqueID(), Property);
+			const EPropertyOnObjectsOwnershipState ExpectedState = bHasProperty
+				? EPropertyOnObjectsOwnershipState::OwnedOnAllObjects
+				: EPropertyOnObjectsOwnershipState::NotOwnedOnAllObjects;
+			
+			if (Result == EPropertyOnObjectsOwnershipState::Mixed)
+			{
+				Result = ExpectedState;
+				continue;
+			}
+
+			if (Result != ExpectedState)
+			{
+				return EPropertyOnObjectsOwnershipState::Mixed;
+			}
+		}
+		return Result;
+	}
+
+	void FAssignPropertyModel::AssignPropertyTo(const FOnlineClient* Client, TConstArrayView<TSoftObjectPtr<>> Objects, const FConcertPropertyChain& Property)
+	{
+		const TSharedRef<ConcertSharedSlate::IEditableReplicationStreamModel> EditModel = Client->GetClientEditModel();
+		for (const TSoftObjectPtr<>& Object : Objects)
+		{
+			const FSoftObjectPath& ObjectPath = Object.GetUniqueID();
+			if (!EditModel->ContainsObjects({ ObjectPath }))
+			{
+				EditModel->AddObjects({ Object.Get() });
+			}
+
+			const FSoftClassPath ClassPath = EditModel->GetObjectClass(ObjectPath);
+			TArray AddedProperties { Property };
+			ConcertClientSharedSlate::PropertyUtils::AppendAdditionalPropertiesToAdd(ClassPath, AddedProperties);
+			EditModel->AddProperties(ObjectPath, AddedProperties);
+		}
+	}
+
+	void FAssignPropertyModel::TogglePropertyFor(const FGuid& ClientId, TConstArrayView<TSoftObjectPtr<>> Objects, const FConcertPropertyChain& Property)
+	{
+		// Remote clients can disconnect after the combo-box is opened.
+		const FOnlineClient* Client = ClientManager.FindClient(ClientId);
+		if (!Client)
+		{
+			return;
+		}
+		const FText TransactionText = FText::Format(
+			LOCTEXT("AllClientsAssignFmt", "Assign {0} property"),
+			FText::FromString(Property.ToString(FConcertPropertyChain::EToStringMethod::LeafProperty))
+			);
+		FScopedTransaction Transaction(TransactionText);
+
+		const EPropertyOnObjectsOwnershipState OwnershipState = GetPropertyOwnershipState(ClientId, Objects, Property);
+		const bool bRemovePropertyFromEditedClient = OwnershipState == EPropertyOnObjectsOwnershipState::OwnedOnAllObjects;
+		
+		// To make it simpler for the user, at most one client is supposed to be assigned to the object at any given time so ...
+		if (bRemovePropertyFromEditedClient)
+		{
+			// ... remove property from all clients
+			ClearProperty(Objects, Property);
+		}
+		else
+		{
+			// ... remove the property from all clients but the one we'll assign to ...
+			AssignPropertyModel::UnassignPropertyFromClients(ClientManager, Objects, Property,
+				[Client](const FOnlineClient& ClientToRemoveFrom){ return *Client != ClientToRemoveFrom; }
+				);
+
+			// ... and then assign the property
+			AssignPropertyTo(Client, Objects, Property);
+		}
+	}
+
+	void FAssignPropertyModel::ClearProperty(TConstArrayView<TSoftObjectPtr<>> Objects, const FConcertPropertyChain& Property) const
+	{
+		if (CanClear(Objects, Property))
+		{
+			const FText TransactionText = FText::Format(
+				LOCTEXT("ClearAllClientsFmt", "Clear {0} property"),
+				FText::FromString(Property.ToString(FConcertPropertyChain::EToStringMethod::LeafProperty))
+				);
+			FScopedTransaction Transaction(TransactionText);
+			
+			AssignPropertyModel::UnassignPropertyFromClients(ClientManager, Objects, Property, [](auto&){ return true; });
+		}
+	}
+}
+
+#undef LOCTEXT_NAMESPACE
\ No newline at end of file
diff --git a/Engine/Plugins/Developer/Concert/ConcertApp/MultiUserClient/Source/MultiUserClient/Private/Widgets/ActiveSession/Replication/Client/Multi/Columns/AssignProperty/AssignPropertyModel.h b/Engine/Plugins/Developer/Concert/ConcertApp/MultiUserClient/Source/MultiUserClient/Private/Widgets/ActiveSession/Replication/Client/Multi/Columns/AssignProperty/AssignPropertyModel.h
new file mode 100644
--- /dev/null
+++ b/Engine/Plugins/Developer/Concert/ConcertApp/MultiUserClient/Source/MultiUserClient/Private/Widgets/ActiveSession/Replication/Client/Multi/Columns/AssignProperty/AssignPropertyModel.h
@@ -0,0 +1,64 @@
+﻿// Copyright Epic Games, Inc. All Rights Reserved.
+
+#pragma once
+
+#include "Delegates/Delegate.h"
+#include "HAL/Platform.h"
+#include "Templates/UnrealTemplate.h"
+#include "UObject/SoftObjectPtr.h"
+
+struct FConcertPropertyChain;
+
+namespace UE::MultiUserClient::Replication
+{
+	class FOnlineClient;
+	class FOnlineClientManager;
+}
+
+namespace UE::MultiUserClient::Replication::MultiStreamColumns
+{
+	enum class EPropertyOnObjectsOwnershipState : uint8
+	{
+		/** Clients owns the property on all objects */
+		OwnedOnAllObjects,
+		/** does not own the property on any object */
+		NotOwnedOnAllObjects,
+		
+		Mixed
+	};
+
+	/** Implements the model in the MVC pattern for the property assignment column. */
+	class FAssignPropertyModel : public FNoncopyable
+	{
+	public:
+		
+		FAssignPropertyModel(FOnlineClientManager& InClientManager UE_LIFETIMEBOUND);
+		~FAssignPropertyModel();
+
+		/** @return Whether property ownership can be changed for the given client. */
+		bool CanChangePropertyFor(const FGuid& ClientId, FText* Reason = nullptr) const;
+		/** @return Whether it is valid to call ClearProperty with this parameters. */
+		bool CanClear(TConstArrayView<TSoftObjectPtr<>> Objects, const FConcertPropertyChain& Property) const;
+		/** @return How the property is owned by ClientId */
+		EPropertyOnObjectsOwnershipState GetPropertyOwnershipState(const FGuid& ClientId, TConstArrayView<TSoftObjectPtr<>> Objects, const FConcertPropertyChain& Property) const;
+		void AssignPropertyTo(const FOnlineClient* Client, TConstArrayView<TSoftObjectPtr<>, signed int> Objects, const FConcertPropertyChain& Property);
+
+		/** Assigns the property to ClientId or unassigns it. Removes the property from all other clients in both cases. */
+		void TogglePropertyFor(const FGuid& ClientId, TConstArrayView<TSoftObjectPtr<>> Objects, const FConcertPropertyChain& Property);
+		/** Removes the property from all clients. */
+		void ClearProperty(TConstArrayView<TSoftObjectPtr<>> Objects, const FConcertPropertyChain& Property) const;
+		
+		/** Broadcasts when property ownership may have changed. */
+		FSimpleMulticastDelegate& OnOwnershipChanged() { return OnOwnershipChangedDelegate; }
+
+	private:
+
+		/** Used to get all clients that can be assigned / reassigned to. */
+		FOnlineClientManager& ClientManager;
+
+		/** Broadcasts when property ownership may have changed. */
+		FSimpleMulticastDelegate OnOwnershipChangedDelegate;
+		
+		void BroadcastOnOwnershipChanged() const { OnOwnershipChangedDelegate.Broadcast(); }
+	};
+}
diff --git a/Engine/Plugins/Developer/Concert/ConcertApp/MultiUserClient/Source/MultiUserClient/Private/Widgets/ActiveSession/Replication/Client/Multi/Columns/AssignProperty/MultiStreamColumns_AssignProperties.cpp b/Engine/Plugins/Developer/Concert/ConcertApp/MultiUserClient/Source/MultiUserClient/Private/Widgets/ActiveSession/Replication/Client/Multi/Columns/AssignProperty/MultiStreamColumns_AssignProperties.cpp
new file mode 100644
--- /dev/null
+++ b/Engine/Plugins/Developer/Concert/ConcertApp/MultiUserClient/Source/MultiUserClient/Private/Widgets/ActiveSession/Replication/Client/Multi/Columns/AssignProperty/MultiStreamColumns_AssignProperties.cpp
@@ -0,0 +1,175 @@
+// Copyright Epic Games, Inc. All Rights Reserved.
+
+#include "AssignPropertyModel.h"
+
+#include "IConcertClient.h"
+#include "MultiUserReplicationStyle.h"
+#include "Replication/Client/Online/OnlineClient.h"
+#include "Replication/Client/Online/OnlineClientManager.h"
+#include "Replication/Editor/Model/IEditableMultiReplicationStreamModel.h"
+#include "Replication/Editor/Model/IEditableReplicationStreamModel.h"
+#include "Replication/Editor/View/IMultiReplicationStreamEditor.h"
+#include "Replication/Editor/View/IReplicationStreamEditor.h"
+#include "SAssignPropertyComboBox.h"
+#include "Replication/Editor/View/Column/IPropertyTreeColumn.h"
+#include "Replication/Editor/View/Column/ReplicationColumnDelegates.h"
+#include "Widgets/ActiveSession/Replication/Client/ClientUtils.h"
+
+#define LOCTEXT_NAMESPACE "AssignPropertyColumn"
+
+namespace UE::MultiUserClient::Replication::MultiStreamColumns
+{
+	const FName AssignPropertyColumnId(TEXT("AssignPropertyColumn"));
+
+	namespace AssignPropertyColumnUtils
+	{
+		static void ForEachStreamAssignedTo(
+			const ConcertSharedSlate::IMultiReplicationStreamEditor& MultiEditor,
+			const ConcertSharedSlate::FPropertyTreeRowContext& InItem,
+			TFunctionRef<void(const TSharedRef<ConcertSharedSlate::IEditableReplicationStreamModel>& Stream)> Consume
+		)
+		{
+			const FConcertPropertyChain& Property = InItem.RowData.GetProperty();
+		
+			for (const TSharedRef<ConcertSharedSlate::IEditableReplicationStreamModel>& Stream : MultiEditor.GetMultiStreamModel().GetEditableStreams())
+			{
+				for (const TSoftObjectPtr<>& ContextObject : InItem.RowData.GetContextObjects())
+				{
+					const bool bStreamAssignedToProperty = Stream->HasProperty(ContextObject.GetUniqueID(), Property);
+					if (bStreamAssignedToProperty)
+					{
+						Consume(Stream);
+						// Each IEditableReplicationStreamModel should only be visited at most once.
+						break;
+					}
+				}
+			}
+		}
+
+		const FOnlineClient* FindClientByStream(const FOnlineClientManager& ClientManager, const ConcertSharedSlate::IReplicationStreamModel& StreamModel)
+		{
+			if (&ClientManager.GetLocalClient().GetClientEditModel().Get() == &StreamModel)
+			{
+				return &ClientManager.GetLocalClient();
+			}
+
+			for (const TNonNullPtr<const FRemoteClient> Client : ClientManager.GetRemoteClients())
+			{
+				if (&Client->GetClientEditModel().Get() == &StreamModel)
+				{
+					return Client;
+				}
+			}
+			
+			return nullptr;
+		}
+		
+		static FString GetClientDisplayText(const IConcertClient& InConcertClient, const FOnlineClientManager& ClientManager, const ConcertSharedSlate::IReplicationStreamModel& StreamModel)
+		{
+			if (const FOnlineClient* Client = FindClientByStream(ClientManager, StreamModel))
+			{
+				return ClientUtils::GetClientDisplayName(InConcertClient, Client->GetEndpointId());
+			}
+
+			ensure(false);
+			return {};
+		}
+	}
+	
+	ConcertSharedSlate::FPropertyColumnEntry AssignPropertyColumn(
+		TAttribute<TSharedPtr<ConcertSharedSlate::IMultiReplicationStreamEditor>> MultiStreamEditor,
+		TSharedRef<IConcertClient> ConcertClient,
+		FOnlineClientManager& ClientManager,
+		const int32 ColumnsSortPriority
+		)
+	{
+		using namespace ConcertSharedSlate;
+		class FPropertyColumn_AssignPropertyColumn : public IPropertyTreeColumn
+		{
+		public:
+
+			FPropertyColumn_AssignPropertyColumn(
+				TAttribute<TSharedPtr<IMultiReplicationStreamEditor>> MultiStreamEditor,
+				TSharedRef<IConcertClient> ConcertClient,
+				FOnlineClientManager& ClientManager
+				)
+				: MultiStreamEditor(MoveTemp(MultiStreamEditor))
+				, ConcertClient(MoveTemp(ConcertClient))
+				, ClientManager(ClientManager)
+				, Model(ClientManager)
+			{}
+			
+			virtual SHeaderRow::FColumn::FArguments CreateHeaderRowArgs() const override
+			{
+				return SHeaderRow::Column(AssignPropertyColumnId)
+					.DefaultLabel(LOCTEXT("Owner.Label", "Author"))
+					.ToolTipText(LOCTEXT("Owner.ToolTip", "Client that should replicate this property"))
+					.FillSized(FMultiUserReplicationStyle::Get()->GetFloat(TEXT("AllClients.Property.OwnerSize")));
+			}
+			
+			virtual TSharedRef<SWidget> GenerateColumnWidget(const FBuildArgs& InArgs) override
+			{
+				const TArray<TSoftObjectPtr<>> DisplayedObjects = InArgs.RowItem.RowData.GetContextObjects();
+				return SNew(SAssignPropertyComboBox, MultiStreamEditor.Get().ToSharedRef(), ConcertClient, ClientManager, Model)
+					.DisplayedProperty(InArgs.RowItem.RowData.GetProperty())
+					.EditedObjects(DisplayedObjects)
+					.HighlightText(InArgs.HighlightText)
+					.OnPropertyAssignmentChanged_Lambda([this]()
+					{
+						if (const TSharedPtr<IMultiReplicationStreamEditor> Editor = MultiStreamEditor.Get())
+						{
+							Editor->GetEditorBase().RequestPropertyColumnResort(AssignPropertyColumnId);
+						}
+					});
+			}
+			
+			virtual void PopulateSearchString(const FPropertyTreeRowContext& InItem, TArray<FString>& InOutSearchStrings) const override
+			{
+				AssignPropertyColumnUtils::ForEachStreamAssignedTo(
+					*MultiStreamEditor.Get(),
+					InItem,
+					[this, &InOutSearchStrings](const TSharedRef<IEditableReplicationStreamModel>& Stream)
+					{
+						InOutSearchStrings.Add(AssignPropertyColumnUtils::GetClientDisplayText(*ConcertClient, ClientManager, *Stream));
+					});
+			}
+
+			virtual bool CanBeSorted() const override { return true; }
+			virtual bool IsLessThan(const FPropertyTreeRowContext& Left, const FPropertyTreeRowContext& Right) const override
+			{
+				const TArray<TSoftObjectPtr<>> DisplayedObjects = MultiStreamEditor.Get()->GetEditorBase().GetSelectedObjects();
+				const TOptional<FString> LeftClientDisplayString = SAssignPropertyComboBox::GetDisplayString(ConcertClient, ClientManager, Left.RowData.GetProperty(), DisplayedObjects);
+				const TOptional<FString> RightClientDisplayString = SAssignPropertyComboBox::GetDisplayString(ConcertClient, ClientManager, Right.RowData.GetProperty(), DisplayedObjects);
+			
+				if (LeftClientDisplayString && RightClientDisplayString)
+				{
+					return *LeftClientDisplayString < *RightClientDisplayString;
+				}
+				// Our rule: set < unset. This way unassigned appears last.
+				return LeftClientDisplayString.IsSet() && !RightClientDisplayString.IsSet();
+			}
+
+		private:
+
+			const TAttribute<TSharedPtr<IMultiReplicationStreamEditor>> MultiStreamEditor;
+			const TSharedRef<IConcertClient> ConcertClient;
+			FOnlineClientManager& ClientManager;
+
+			/** The model the view displays. */
+			FAssignPropertyModel Model;
+		};
+		
+		check(MultiStreamEditor.IsBound() || MultiStreamEditor.IsSet());
+		return {
+			TReplicationColumnDelegates<FPropertyTreeRowContext>::FCreateColumn::CreateLambda(
+				[MultiStreamEditor = MoveTemp(MultiStreamEditor), ConcertClient = MoveTemp(ConcertClient), &ClientManager]()
+				{
+					return MakeShared<FPropertyColumn_AssignPropertyColumn>(MultiStreamEditor, ConcertClient, ClientManager);
+				}),
+			AssignPropertyColumnId,
+			{ ColumnsSortPriority }
+		};
+	}
+}
+
+#undef LOCTEXT_NAMESPACE
\ No newline at end of file
diff --git a/Engine/Plugins/Developer/Concert/ConcertApp/MultiUserClient/Source/MultiUserClient/Private/Widgets/ActiveSession/Replication/Client/Multi/Columns/AssignProperty/SAssignPropertyComboBox.cpp b/Engine/Plugins/Developer/Concert/ConcertApp/MultiUserClient/Source/MultiUserClient/Private/Widgets/ActiveSession/Replication/Client/Multi/Columns/AssignProperty/SAssignPropertyComboBox.cpp
new file mode 100644
--- /dev/null
+++ b/Engine/Plugins/Developer/Concert/ConcertApp/MultiUserClient/Source/MultiUserClient/Private/Widgets/ActiveSession/Replication/Client/Multi/Columns/AssignProperty/SAssignPropertyComboBox.cpp
@@ -0,0 +1,235 @@
+// Copyright Epic Games, Inc. All Rights Reserved.
+
+#include "SAssignPropertyComboBox.h"
+
+#include "AssignPropertyModel.h"
+#include "IConcertClient.h"
+#include "Replication/Client/Online/OnlineClient.h"
+#include "Replication/Client/Online/OnlineClientManager.h"
+#include "Replication/Editor/Model/PropertyUtils.h"
+#include "Widgets/ActiveSession/Replication/Misc/SNoClients.h"
+#include "Widgets/ActiveSession/Replication/Client/ClientUtils.h"
+#include "Widgets/Client/ClientInfoHelpers.h"
+#include "Widgets/Client/SHorizontalClientList.h"
+#include "Widgets/Client/SLocalClientName.h"
+#include "Widgets/Client/SRemoteClientName.h"
+
+#include "Framework/MultiBox/MultiBoxBuilder.h"
+#include "Styling/AppStyle.h"
+#include "Styling/SlateTypes.h"
+#include "Widgets/Input/SComboButton.h"
+
+#define LOCTEXT_NAMESPACE "SAssignPropertyComboBox"
+
+namespace UE::MultiUserClient::Replication::MultiStreamColumns
+{
+	namespace AssignPropertyComboBox
+	{
+		TArray<FGuid> GetDisplayedClients(const FOnlineClientManager& ClientManager, const FConcertPropertyChain& DisplayedProperty, const TArray<TSoftObjectPtr<>>& EditedObjects)
+		{
+			TArray<FGuid> Clients;
+			ClientManager.ForEachClient([&DisplayedProperty, &EditedObjects, &Clients](const FOnlineClient& Client)
+			{
+				const TMap<FSoftObjectPath, FConcertReplicatedObjectInfo>& ObjectInfoMap = Client.GetStreamSynchronizer().GetServerState().ReplicatedObjects;
+				for (const TSoftObjectPtr<>& ObjectPath : EditedObjects)
+				{
+					if (const FConcertReplicatedObjectInfo* ObjectInfo = ObjectInfoMap.Find(ObjectPath.GetUniqueID())
+						; ObjectInfo && ObjectInfo->PropertySelection.ReplicatedProperties.Contains(DisplayedProperty))
+					{
+						Clients.Add(Client.GetEndpointId());
+						return EBreakBehavior::Continue;
+					}
+				}
+				return EBreakBehavior::Continue;
+			});
+			return Clients;
+		}
+	}
+	
+	TOptional<FString> SAssignPropertyComboBox::GetDisplayString(
+		const TSharedRef<IConcertClient>& LocalConcertClient,
+		const FOnlineClientManager& ClientManager,
+		const FConcertPropertyChain& DisplayedProperty,
+		const TArray<TSoftObjectPtr<>>& EditedObjects)
+	{
+		using SWidgetType = ConcertSharedSlate::SHorizontalClientList;
+		const TArray<FGuid> Clients = AssignPropertyComboBox::GetDisplayedClients(ClientManager, DisplayedProperty, EditedObjects);
+		
+		const ConcertSharedSlate::FGetClientParenthesesContent GetParenthesesContent =
+			ConcertClientSharedSlate::MakeGetLocalClientParenthesesContent(LocalConcertClient);
+		const auto SortPredicate = [&GetParenthesesContent](const FConcertSessionClientInfo& Left, const FConcertSessionClientInfo& Right)
+		{
+			return ConcertSharedSlate::SortLocalClientParenthesesFirstThenThenAlphabetical(Left, Right, GetParenthesesContent);
+		};
+		
+		return SWidgetType::GetDisplayString(
+			Clients,
+			ConcertClientSharedSlate::MakeClientInfoGetter(LocalConcertClient),
+			ConcertSharedSlate::FClientSortPredicate::CreateLambda(SortPredicate),
+			GetParenthesesContent
+			);
+	}
+
+	void SAssignPropertyComboBox::Construct(const FArguments& InArgs,
+	    TSharedRef<ConcertSharedSlate::IMultiReplicationStreamEditor> InEditor,
+	    TSharedRef<IConcertClient> InConcertClient,
+	    FOnlineClientManager& InClientManager,
+	    FAssignPropertyModel& InModel
+	)
+	{
+		Editor = MoveTemp(InEditor);
+		ConcertClient = MoveTemp(InConcertClient);
+		ClientManager = &InClientManager;
+		Model = &InModel;
+		
+		Property = InArgs._DisplayedProperty;
+		EditedObjects = InArgs._EditedObjects;
+		HighlightText = InArgs._HighlightText;
+		check(!EditedObjects.IsEmpty());
+
+		OnOptionClickedDelegate = InArgs._OnPropertyAssignmentChanged;
+		
+		ChildSlot
+		[
+			SNew(SComboButton)
+			.HasDownArrow(true)
+			.ButtonContent()
+			[
+				SAssignNew(ClientListWidget, ConcertSharedSlate::SHorizontalClientList)
+				.GetClientParenthesesContent(ConcertClientSharedSlate::MakeGetLocalClientParenthesesContent(ConcertClient.ToSharedRef()))
+				.GetClientInfo(ConcertClientSharedSlate::MakeClientInfoGetter(ConcertClient.ToSharedRef()))
+				.Font(FAppStyle::GetFontStyle("PropertyWindow.NormalFont"))
+				.HighlightText_Lambda([this](){ return HighlightText ? *HighlightText : FText::GetEmpty(); })
+				.EmptyListSlot() [ SNew(SNoClients) ]
+			]
+			.OnGetMenuContent(this, &SAssignPropertyComboBox::GetMenuContent)
+		];
+
+		Model->OnOwnershipChanged().AddSP(this, &SAssignPropertyComboBox::RebuildSubscriptionsAndRefresh);
+		RebuildSubscriptions();
+		RefreshContentBoxContent();
+	}
+	
+	void SAssignPropertyComboBox::RefreshContentBoxContent() const
+	{
+		ClientListWidget->RefreshList(
+			AssignPropertyComboBox::GetDisplayedClients(*ClientManager, Property, EditedObjects)
+			);
+	}
+
+	TSharedRef<SWidget> SAssignPropertyComboBox::GetMenuContent()
+	{
+		using namespace ConcertSharedSlate;
+
+		const auto MakeWidget = [this](const FGuid& EndpointId) -> TSharedRef<SWidget>
+		{
+			const bool bIsLocalClient = EndpointId == ConcertClient->GetCurrentSession()->GetSessionClientEndpointId();
+			if (bIsLocalClient)
+			{
+				return SNew(SLocalClientName)
+					.DisplayInfo(ConcertClientSharedSlate::MakeLocalClientInfoAttribute(ConcertClient.ToSharedRef()))
+					.Font(FAppStyle::GetFontStyle("PropertyWindow.NormalFont"))
+					.HighlightText_Lambda([this](){ return HighlightText ? *HighlightText : FText::GetEmpty(); });
+			}
+			return SNew(SRemoteClientName)
+				.DisplayInfo(ConcertClientSharedSlate::MakeClientInfoAttribute(ConcertClient.ToSharedRef(), EndpointId))
+				.Font(FAppStyle::GetFontStyle("PropertyWindow.NormalFont"))
+				.HighlightText_Lambda([this](){ return HighlightText ? *HighlightText : FText::GetEmpty(); });
+		};
+		
+		FMenuBuilder MenuBuilder(true, nullptr);
+		MenuBuilder.AddMenuEntry(
+			LOCTEXT("Clear.Label", "Clear"),
+			LOCTEXT("Clear.Tooltip", "Stop this property from being replicated"),
+			FSlateIcon(),
+			FUIAction(
+				FExecuteAction::CreateSP(this, &SAssignPropertyComboBox::OnClickClear),
+				FCanExecuteAction::CreateSP(this, &SAssignPropertyComboBox::CanClickClear)
+				),
+			NAME_None,
+			EUserInterfaceActionType::Button
+		);
+		
+		MenuBuilder.BeginSection(TEXT("AssignTo"), LOCTEXT("AssignTo", "Assign to"));
+		for (const FOnlineClient* Client : ClientUtils::GetSortedClientList(*ConcertClient, *ClientManager))
+		{
+			TAttribute<FText> Tooltip = TAttribute<FText>::CreateLambda([this, EndpointId = Client->GetEndpointId()]()
+			{
+				FText Reason;
+				const bool bCanClick = CanClickOptionWithReason(EndpointId, &Reason);
+				if (!bCanClick)
+				{
+					return Reason;
+				}
+
+				switch (GetOptionCheckState(EndpointId))
+				{
+				case ECheckBoxState::Unchecked: return LOCTEXT("Action.Unchecked", "Assign property to client and remove it from all others.");
+				case ECheckBoxState::Undetermined:  return LOCTEXT("Action.Undetermined", "Assign property to client for all selected objects.");
+				case ECheckBoxState::Checked: return LOCTEXT("Action.Checked", "Remove property from client and remove it from all others.");
+				default: return FText::GetEmpty();
+				}
+			});
+			
+			MenuBuilder.AddMenuEntry(
+				FUIAction(
+					FExecuteAction::CreateSP(this, &SAssignPropertyComboBox::OnClickOption, Client->GetEndpointId()),
+					FCanExecuteAction::CreateSP(this, &SAssignPropertyComboBox::CanClickOption, Client->GetEndpointId()),
+					FGetActionCheckState::CreateSP(this, &SAssignPropertyComboBox::GetOptionCheckState, Client->GetEndpointId())
+					),
+				MakeWidget(Client->GetEndpointId()),
+				NAME_None,
+				Tooltip,
+				EUserInterfaceActionType::Check
+				);
+		}
+		MenuBuilder.EndSection();
+		
+		return MenuBuilder.MakeWidget();
+	}
+	
+	void SAssignPropertyComboBox::OnClickOption(const FGuid EndpointId) const
+	{
+		Model->TogglePropertyFor(EndpointId, EditedObjects, Property);
+		OnOptionClickedDelegate.ExecuteIfBound();
+	}
+
+	bool SAssignPropertyComboBox::CanClickOptionWithReason(const FGuid& EndpointId, FText* Reason) const
+	{
+		return Model->CanChangePropertyFor(EndpointId, Reason);
+	}
+	
+	ECheckBoxState SAssignPropertyComboBox::GetOptionCheckState(const FGuid EndpointId) const
+	{
+		switch (Model->GetPropertyOwnershipState(EndpointId,EditedObjects, Property))
+		{
+		case EPropertyOnObjectsOwnershipState::OwnedOnAllObjects: return ECheckBoxState::Checked;
+		case EPropertyOnObjectsOwnershipState::NotOwnedOnAllObjects: return ECheckBoxState::Unchecked;
+		case EPropertyOnObjectsOwnershipState::Mixed: return ECheckBoxState::Undetermined;
+		default: return ECheckBoxState::Undetermined;
+		}
+	}
+
+	void SAssignPropertyComboBox::OnClickClear()
+	{
+		Model->ClearProperty(EditedObjects, Property);
+		OnOptionClickedDelegate.ExecuteIfBound();
+	}
+
+	bool SAssignPropertyComboBox::CanClickClear() const
+	{
+		return Model->CanClear(EditedObjects, Property);
+	}
+
+	void SAssignPropertyComboBox::RebuildSubscriptions()
+	{
+		ClientManager->ForEachClient([this](FOnlineClient& Client)
+		{
+			Client.OnModelChanged().RemoveAll(this);
+			Client.OnModelChanged().AddSP(this, &SAssignPropertyComboBox::RefreshContentBoxContent);
+			return EBreakBehavior::Continue;
+		});
+	}
+}
+
+#undef LOCTEXT_NAMESPACE
\ No newline at end of file
diff --git a/Engine/Plugins/Developer/Concert/ConcertApp/MultiUserClient/Source/MultiUserClient/Private/Widgets/ActiveSession/Replication/Client/Multi/Columns/AssignProperty/SAssignPropertyComboBox.h b/Engine/Plugins/Developer/Concert/ConcertApp/MultiUserClient/Source/MultiUserClient/Private/Widgets/ActiveSession/Replication/Client/Multi/Columns/AssignProperty/SAssignPropertyComboBox.h
new file mode 100644
--- /dev/null
+++ b/Engine/Plugins/Developer/Concert/ConcertApp/MultiUserClient/Source/MultiUserClient/Private/Widgets/ActiveSession/Replication/Client/Multi/Columns/AssignProperty/SAssignPropertyComboBox.h
@@ -0,0 +1,115 @@
+// Copyright Epic Games, Inc. All Rights Reserved.
+
+#pragma once
+
+#include "Replication/Data/ConcertPropertySelection.h"
+
+#include "Containers/Array.h"
+#include "Templates/Function.h"
+#include "Widgets/DeclarativeSyntaxSupport.h"
+#include "Widgets/SCompoundWidget.h"
+
+class IConcertClient;
+enum class ECheckBoxState : uint8;
+struct FGuid;
+
+namespace UE::ConcertSharedSlate
+{
+	class IMultiReplicationStreamEditor;
+	class SHorizontalClientList;
+}
+
+namespace UE::MultiUserClient::Replication
+{
+	class FOnlineClient;
+	class FOnlineClientManager;
+}
+
+namespace UE::MultiUserClient::Replication::MultiStreamColumns
+{
+	class FAssignPropertyModel;
+	
+	/**
+	 * Placed in every property row column to assign properties to clients.
+	 * 
+	 * Can be used to assign a property to at most one client: multiple clients cannot be assigned because it is an advanced workflow
+	 * which might confuse basic users.
+	 */
+	class SAssignPropertyComboBox : public SCompoundWidget
+	{
+	public:
+
+		DECLARE_DELEGATE(FOnPropertyAssignmentChanged);
+
+		/** @return The display string this widget would have with the given state. If unset, no clients are displayed in the combobox.*/
+		static TOptional<FString> GetDisplayString(
+			const TSharedRef<IConcertClient>& LocalConcertClient,
+			const FOnlineClientManager& ClientManager,
+			const FConcertPropertyChain& DisplayedProperty,
+			const TArray<TSoftObjectPtr<>>& EditedObjects
+			);
+		
+		SLATE_BEGIN_ARGS(SAssignPropertyComboBox)
+		{}
+			SLATE_ARGUMENT(FConcertPropertyChain, DisplayedProperty)
+			SLATE_ARGUMENT(TArray<TSoftObjectPtr<>>, EditedObjects)
+			SLATE_ARGUMENT(TSharedPtr<FText>, HighlightText)
+			/** Called when the property assignment of any client(s) is changed by this widget. */
+			SLATE_EVENT(FOnPropertyAssignmentChanged, OnPropertyAssignmentChanged)
+		SLATE_END_ARGS()
+
+		void Construct(const FArguments& InArgs,
+           TSharedRef<ConcertSharedSlate::IMultiReplicationStreamEditor> InEditor,
+           TSharedRef<IConcertClient> InConcertClient,
+           FOnlineClientManager& InClientManager,
+           FAssignPropertyModel& InModel UE_LIFETIMEBOUND
+		);
+
+	private:
+
+		FOnlineClientManager* ClientManager = nullptr;
+		/** The model this view is displaying. */
+		FAssignPropertyModel* Model = nullptr;
+		
+		/** Used to obtain info about the streams */
+		TSharedPtr<ConcertSharedSlate::IMultiReplicationStreamEditor> Editor;
+		TSharedPtr<IConcertClient> ConcertClient;
+
+		/** The objects for which the property is being displayed. */
+		TArray<TSoftObjectPtr<>> EditedObjects;
+		/** The property assigned to this column. */
+		FConcertPropertyChain Property;
+
+		/** The static menu content (when there is no drop-down). */
+		TSharedPtr<ConcertSharedSlate::SHorizontalClientList> ClientListWidget;
+		/** Passed to MakeListWidgetDelegate */
+		TSharedPtr<FText> HighlightText;
+
+		/** Called when the property assignment of any client(s) is changed by this widget. */
+		FOnPropertyAssignmentChanged OnOptionClickedDelegate;
+
+		/** Updates the content of the combo box */
+		void RefreshContentBoxContent() const;
+		
+		/** Builds the drop-down menu */
+		TSharedRef<SWidget> GetMenuContent();
+
+		/** Called when a client is clicked in the drop-down menu */
+		void OnClickOption(const FGuid EndpointId) const;
+		/** @return Whether an item in the drop-down menu should be clickable (returns true if the stream is writable) */
+		bool CanClickOption(const FGuid EndpointId) const { return CanClickOptionWithReason(EndpointId, nullptr); }
+		bool CanClickOptionWithReason(const FGuid& EndpointId, FText* Reason = nullptr) const;
+		/** @return Checkbox state for the item. Handles cases of multiple objects being property edited. */
+		ECheckBoxState GetOptionCheckState(const FGuid EndpointId) const;
+
+		/** Called when a client clicks the clear button */
+		void OnClickClear();
+		/** Whether the property is assigned to any client */
+		bool CanClickClear() const;
+
+		/** Resubscribes to all clients changing. */
+		void RebuildSubscriptions();
+		void RebuildSubscriptionsAndRefresh() { RebuildSubscriptions(); RefreshContentBoxContent(); }
+	};
+}
+
diff --git a/Engine/Plugins/Developer/Concert/ConcertApp/MultiUserClient/Source/MultiUserClient/Private/Widgets/ActiveSession/Replication/Client/Multi/Columns/AssignedClients/AssignedClientsModel.cpp b/Engine/Plugins/Developer/Concert/ConcertApp/MultiUserClient/Source/MultiUserClient/Private/Widgets/ActiveSession/Replication/Client/Multi/Columns/AssignedClients/AssignedClientsModel.cpp
--- a/Engine/Plugins/Developer/Concert/ConcertApp/MultiUserClient/Source/MultiUserClient/Private/Widgets/ActiveSession/Replication/Client/Multi/Columns/AssignedClients/AssignedClientsModel.cpp
+++ b/Engine/Plugins/Developer/Concert/ConcertApp/MultiUserClient/Source/MultiUserClient/Private/Widgets/ActiveSession/Replication/Client/Multi/Columns/AssignedClients/AssignedClientsModel.cpp
@@ -9,6 +9,8 @@
 #include "Containers/Array.h"
 #include "Misc/Guid.h"
 
+#include <type_traits>
+
 namespace UE::MultiUserClient::Replication::MultiStreamColumns::Private
 {
 	template<typename TCallback> requires std::is_invocable_v<TCallback, const FGuid&>
diff --git a/Engine/Plugins/Developer/Concert/ConcertApp/MultiUserClient/Source/MultiUserClient/Private/Widgets/ActiveSession/Replication/Client/Multi/Columns/MultiStreamColumns_AssignProperties.cpp b/Engine/Plugins/Developer/Concert/ConcertApp/MultiUserClient/Source/MultiUserClient/Private/Widgets/ActiveSession/Replication/Client/Multi/Columns/MultiStreamColumns_AssignProperties.cpp
deleted file mode 100644
diff --git a/Engine/Plugins/Developer/Concert/ConcertApp/MultiUserClient/Source/MultiUserClient/Private/Widgets/ActiveSession/Replication/Client/Multi/Columns/SAssignPropertyComboBox.cpp b/Engine/Plugins/Developer/Concert/ConcertApp/MultiUserClient/Source/MultiUserClient/Private/Widgets/ActiveSession/Replication/Client/Multi/Columns/SAssignPropertyComboBox.cpp
deleted file mode 100644
diff --git a/Engine/Plugins/Developer/Concert/ConcertApp/MultiUserClient/Source/MultiUserClient/Private/Widgets/ActiveSession/Replication/Client/Multi/Columns/SAssignPropertyComboBox.h b/Engine/Plugins/Developer/Concert/ConcertApp/MultiUserClient/Source/MultiUserClient/Private/Widgets/ActiveSession/Replication/Client/Multi/Columns/SAssignPropertyComboBox.h
deleted file mode 100644
