Move ClientUtils to new folder

diff --git a/Engine/Plugins/Developer/Concert/ConcertApp/MultiUserClient/Source/MultiUserClient/Private/Replication/Client/ClientUtils.cpp b/Engine/Plugins/Developer/Concert/ConcertApp/MultiUserClient/Source/MultiUserClient/Private/Replication/Client/ClientUtils.cpp
new file mode 100644
--- /dev/null
+++ b/Engine/Plugins/Developer/Concert/ConcertApp/MultiUserClient/Source/MultiUserClient/Private/Replication/Client/ClientUtils.cpp
@@ -0,0 +1,92 @@
+// Copyright Epic Games, Inc. All Rights Reserved.
+
+#include "ClientUtils.h"
+
+#include "ConcertMessageData.h"
+#include "IConcertClient.h"
+#include "Replication/Client/Online/OnlineClientManager.h"
+#include "Widgets/Client/SClientName.h"
+
+namespace UE::MultiUserClient::Replication::ClientUtils
+{
+	FString GetClientDisplayName(const IConcertClient& InLocalClientInstance, const FGuid& InClientEndpointId)
+	{
+		const TSharedPtr<IConcertClientSession> Session = InLocalClientInstance.GetCurrentSession();
+		return ensure(Session) ? GetClientDisplayName(*Session, InClientEndpointId) : FString{};
+	}
+	
+	FString GetClientDisplayName(const IConcertClientSession& InSession, const FGuid& InClientEndpointId)
+	{
+		const bool bIsLocalClient = InSession.GetSessionClientEndpointId() == InClientEndpointId;
+		if (bIsLocalClient)
+		{
+			return ConcertSharedSlate::SClientName::GetDisplayText(InSession.GetLocalClientInfo(), bIsLocalClient).ToString();
+		}
+
+		FConcertSessionClientInfo ClientInfo;
+		if (InSession.FindSessionClient(InClientEndpointId, ClientInfo))
+		{
+			return ConcertSharedSlate::SClientName::GetDisplayText(ClientInfo.ClientInfo, bIsLocalClient).ToString();
+		}
+
+		ensureMsgf(false, TEXT("Bad args"));
+		return InClientEndpointId.ToString(EGuidFormats::DigitsWithHyphens);
+	}
+	
+	bool GetClientDisplayInfo(const IConcertClient& InLocalClientInstance, const FGuid& InClientEndpointId, FConcertClientInfo& OutClientInfo)
+	{
+		const TSharedPtr<IConcertClientSession> Session = InLocalClientInstance.GetCurrentSession();
+		return GetClientDisplayInfo(*Session, InClientEndpointId, OutClientInfo);
+	}
+	
+	bool GetClientDisplayInfo(const IConcertClientSession& InSession, const FGuid& InClientEndpointId, FConcertClientInfo& OutClientInfo)
+	{
+		const bool bIsLocalClient = InSession.GetSessionClientEndpointId() == InClientEndpointId;
+		if (bIsLocalClient)
+		{
+			OutClientInfo = InSession.GetLocalClientInfo();
+			return true;
+		}
+
+		FConcertSessionClientInfo ClientInfo;
+		if (InSession.FindSessionClient(InClientEndpointId, ClientInfo))
+		{
+			OutClientInfo = MoveTemp(ClientInfo.ClientInfo);
+			return true;
+		}
+
+		return false;
+	}
+	
+	TArray<const FOnlineClient*> GetSortedClientList(const IConcertClient& InLocalClientInstance, const FOnlineClientManager& InReplicationManager)
+	{
+		return GetSortedClientList(*InLocalClientInstance.GetCurrentSession(), InReplicationManager);
+	}
+
+	TArray<const FOnlineClient*> GetSortedClientList(const IConcertClientSession& InSession, const FOnlineClientManager& InReplicationManager)
+	{
+		TArray<const FOnlineClient*> Result;
+		TMap<const FOnlineClient*, FConcertClientInfo> ClientToDisplayInfo;
+		for (const TNonNullPtr<const FRemoteClient> RemoteClient : InReplicationManager.GetRemoteClients())
+		{
+			FConcertClientInfo Info;
+			if (GetClientDisplayInfo(InSession, RemoteClient->GetEndpointId(), Info))
+			{
+				Result.Add(RemoteClient);
+				ClientToDisplayInfo.Add(RemoteClient, MoveTemp(Info));
+			}
+		}
+
+		Result.Sort([&ClientToDisplayInfo](const FOnlineClient& Left, const FOnlineClient& Right)
+		{
+			const FConcertClientInfo& LeftInfo = ClientToDisplayInfo[&Left];
+			const FConcertClientInfo& RightInfo = ClientToDisplayInfo[&Right];
+			const FText LeftDisplayName = ConcertSharedSlate::SClientName::GetDisplayText(LeftInfo);
+			const FText RightDisplayName = ConcertSharedSlate::SClientName::GetDisplayText(RightInfo);
+			return LeftDisplayName.ToString() <= RightDisplayName.ToString();
+		});
+		
+		Result.Insert(&InReplicationManager.GetLocalClient(), 0);
+		return Result;
+	}
+}
diff --git a/Engine/Plugins/Developer/Concert/ConcertApp/MultiUserClient/Source/MultiUserClient/Private/Replication/Client/ClientUtils.h b/Engine/Plugins/Developer/Concert/ConcertApp/MultiUserClient/Source/MultiUserClient/Private/Replication/Client/ClientUtils.h
new file mode 100644
--- /dev/null
+++ b/Engine/Plugins/Developer/Concert/ConcertApp/MultiUserClient/Source/MultiUserClient/Private/Replication/Client/ClientUtils.h
@@ -0,0 +1,51 @@
+// Copyright Epic Games, Inc. All Rights Reserved.
+
+#pragma once
+
+#include "Containers/Array.h"
+
+class FString;
+class IConcertClient;
+class IConcertClientSession;
+
+struct FConcertClientInfo;
+struct FGuid;
+
+namespace UE::MultiUserClient::Replication
+{
+	class FOnlineClient;
+	class FOnlineClientManager;
+}
+
+namespace UE::MultiUserClient::Replication::ClientUtils
+{
+	/**
+	 * Gets the display name for a client. Appends (me) if the client is local.
+	 * 
+	 * @param InLocalClientInstance Used to look up client display info
+	 * @param InClientEndpointId The endpoint ID of the client whose name to get
+	 * @return The display name or empty
+	 */
+	FString GetClientDisplayName(const IConcertClient& InLocalClientInstance, const FGuid& InClientEndpointId);
+	FString GetClientDisplayName(const IConcertClientSession& InSession, const FGuid& InClientEndpointId);
+
+	/**
+	 * Gets the display info for a given client
+	 *
+	 * @param InLocalClientInstance Used to look up client display info
+	 * @param InClientEndpointId The endpoint ID of the client whose name to get
+	 * @param OutClientInfo The client display info to get
+	 * @return Whether OutClientInfo holds a value
+	 */
+	bool GetClientDisplayInfo(const IConcertClient& InLocalClientInstance, const FGuid& InClientEndpointId, FConcertClientInfo& OutClientInfo);
+	bool GetClientDisplayInfo(const IConcertClientSession& InSession, const FGuid& InClientEndpointId, FConcertClientInfo& OutClientInfo);
+
+	/**
+	 * Gets all replication clients in a sorted array. The local client will always be first and then come all remote clients sorted alphabetically.
+	 * @param InLocalClientInstance Used to look up client display info
+	 * @param InReplicationManager Used to obtain all clients
+	 * @return Sorted client array
+	 */
+	TArray<const FOnlineClient*> GetSortedClientList(const IConcertClient& InLocalClientInstance, const FOnlineClientManager& InReplicationManager);
+	TArray<const FOnlineClient*> GetSortedClientList(const IConcertClientSession& InSession, const FOnlineClientManager& InReplicationManager);
+}
diff --git a/Engine/Plugins/Developer/Concert/ConcertApp/MultiUserClient/Source/MultiUserClient/Private/Replication/Misc/Util/StreamRequestUtils.cpp b/Engine/Plugins/Developer/Concert/ConcertApp/MultiUserClient/Source/MultiUserClient/Private/Replication/Misc/Util/StreamRequestUtils.cpp
--- a/Engine/Plugins/Developer/Concert/ConcertApp/MultiUserClient/Source/MultiUserClient/Private/Replication/Misc/Util/StreamRequestUtils.cpp
+++ b/Engine/Plugins/Developer/Concert/ConcertApp/MultiUserClient/Source/MultiUserClient/Private/Replication/Misc/Util/StreamRequestUtils.cpp
@@ -2,10 +2,7 @@
 
 #include "StreamRequestUtils.h"
 
-#include "ConcertLogGlobal.h"
-#include "JsonObjectConverter.h"
 #include "Replication/Stream/StreamChangeTracker.h"
-#include "Widgets/ActiveSession/Replication/Client/ClientUtils.h"
 
 #include "Containers/UnrealString.h"
 
diff --git a/Engine/Plugins/Developer/Concert/ConcertApp/MultiUserClient/Source/MultiUserClient/Private/Replication/Preset/PresetManager.cpp b/Engine/Plugins/Developer/Concert/ConcertApp/MultiUserClient/Source/MultiUserClient/Private/Replication/Preset/PresetManager.cpp
--- a/Engine/Plugins/Developer/Concert/ConcertApp/MultiUserClient/Source/MultiUserClient/Private/Replication/Preset/PresetManager.cpp
+++ b/Engine/Plugins/Developer/Concert/ConcertApp/MultiUserClient/Source/MultiUserClient/Private/Replication/Preset/PresetManager.cpp
@@ -7,12 +7,12 @@
 #include "Assets/MultiUserReplicationStream.h"
 #include "ConcertLogGlobal.h"
 #include "IConcertSyncClient.h"
+#include "Replication/Client/ClientUtils.h"
 #include "Replication/Client/Online/OnlineClientManager.h"
 #include "Replication/Misc/ActorLabelRemappingEditor.h"
 #include "Replication/Misc/ReplicationStreamUtils.h"
 #include "Replication/Muting/MuteStateManager.h"
 #include "Replication/Stream/MultiUserStreamId.h"
-#include "Widgets/ActiveSession/Replication/Client/ClientUtils.h"
 
 #include "Engine/World.h"
 #include "FileHelpers.h"
diff --git a/Engine/Plugins/Developer/Concert/ConcertApp/MultiUserClient/Source/MultiUserClient/Private/Replication/Submission/Remote/RemoteSubmissionListener.cpp b/Engine/Plugins/Developer/Concert/ConcertApp/MultiUserClient/Source/MultiUserClient/Private/Replication/Submission/Remote/RemoteSubmissionListener.cpp
--- a/Engine/Plugins/Developer/Concert/ConcertApp/MultiUserClient/Source/MultiUserClient/Private/Replication/Submission/Remote/RemoteSubmissionListener.cpp
+++ b/Engine/Plugins/Developer/Concert/ConcertApp/MultiUserClient/Source/MultiUserClient/Private/Replication/Submission/Remote/RemoteSubmissionListener.cpp
@@ -5,11 +5,11 @@
 #include "ConcertLogGlobal.h"
 #include "IConcertSession.h"
 #include "RemoteSubmissionMessages.h"
+#include "Replication/Client/ClientUtils.h"
 #include "Replication/Stream/IClientStreamSynchronizer.h"
 #include "Replication/Submission/ISubmissionOperation.h"
 #include "Replication/Submission/ISubmissionWorkflow.h"
 #include "Replication/Submission/Queue/DeferredSubmitter.h"
-#include "Widgets/ActiveSession/Replication/Client/ClientUtils.h"
 
 namespace UE::MultiUserClient::Replication
 {
diff --git a/Engine/Plugins/Developer/Concert/ConcertApp/MultiUserClient/Source/MultiUserClient/Private/Replication/Submission/Remote/SubmissionWorkflow_RemoteClient.cpp b/Engine/Plugins/Developer/Concert/ConcertApp/MultiUserClient/Source/MultiUserClient/Private/Replication/Submission/Remote/SubmissionWorkflow_RemoteClient.cpp
--- a/Engine/Plugins/Developer/Concert/ConcertApp/MultiUserClient/Source/MultiUserClient/Private/Replication/Submission/Remote/SubmissionWorkflow_RemoteClient.cpp
+++ b/Engine/Plugins/Developer/Concert/ConcertApp/MultiUserClient/Source/MultiUserClient/Private/Replication/Submission/Remote/SubmissionWorkflow_RemoteClient.cpp
@@ -5,9 +5,9 @@
 #include "ConcertLogGlobal.h"
 #include "IConcertSession.h"
 #include "RemoteSubmissionMessages.h"
-#include "Replication/Submission/SingleClientSubmissionOperation.h"
+#include "Replication/Client/ClientUtils.h"
 #include "Replication/Misc/Util/StreamRequestUtils.h"
-#include "Widgets/ActiveSession/Replication/Client/ClientUtils.h"
+#include "Replication/Submission/SingleClientSubmissionOperation.h"
 
 namespace UE::MultiUserClient::Replication
 {
diff --git a/Engine/Plugins/Developer/Concert/ConcertApp/MultiUserClient/Source/MultiUserClient/Private/Replication/Submission/SubmissionWorkflow_LocalClient.cpp b/Engine/Plugins/Developer/Concert/ConcertApp/MultiUserClient/Source/MultiUserClient/Private/Replication/Submission/SubmissionWorkflow_LocalClient.cpp
--- a/Engine/Plugins/Developer/Concert/ConcertApp/MultiUserClient/Source/MultiUserClient/Private/Replication/Submission/SubmissionWorkflow_LocalClient.cpp
+++ b/Engine/Plugins/Developer/Concert/ConcertApp/MultiUserClient/Source/MultiUserClient/Private/Replication/Submission/SubmissionWorkflow_LocalClient.cpp
@@ -5,7 +5,7 @@
 #include "IConcertSyncClient.h"
 #include "Replication/IConcertClientReplicationManager.h"
 #include "Replication/Misc/Util/StreamRequestUtils.h"
-#include "Widgets/ActiveSession/Replication/Client/ClientUtils.h"
+#include "Replication/Client/ClientUtils.h"
 
 namespace UE::MultiUserClient::Replication
 {
diff --git a/Engine/Plugins/Developer/Concert/ConcertApp/MultiUserClient/Source/MultiUserClient/Private/Widgets/ActiveSession/Replication/Client/ClientUtils.cpp b/Engine/Plugins/Developer/Concert/ConcertApp/MultiUserClient/Source/MultiUserClient/Private/Widgets/ActiveSession/Replication/Client/ClientUtils.cpp
deleted file mode 100644
diff --git a/Engine/Plugins/Developer/Concert/ConcertApp/MultiUserClient/Source/MultiUserClient/Private/Widgets/ActiveSession/Replication/Client/ClientUtils.h b/Engine/Plugins/Developer/Concert/ConcertApp/MultiUserClient/Source/MultiUserClient/Private/Widgets/ActiveSession/Replication/Client/ClientUtils.h
deleted file mode 100644
diff --git a/Engine/Plugins/Developer/Concert/ConcertApp/MultiUserClient/Source/MultiUserClient/Private/Widgets/ActiveSession/Replication/Client/Context/ContextMenuUtils_Reassignment.cpp b/Engine/Plugins/Developer/Concert/ConcertApp/MultiUserClient/Source/MultiUserClient/Private/Widgets/ActiveSession/Replication/Client/Context/ContextMenuUtils_Reassignment.cpp
--- a/Engine/Plugins/Developer/Concert/ConcertApp/MultiUserClient/Source/MultiUserClient/Private/Widgets/ActiveSession/Replication/Client/Context/ContextMenuUtils_Reassignment.cpp
+++ b/Engine/Plugins/Developer/Concert/ConcertApp/MultiUserClient/Source/MultiUserClient/Private/Widgets/ActiveSession/Replication/Client/Context/ContextMenuUtils_Reassignment.cpp
@@ -2,12 +2,12 @@
 
 #include "ContextMenuUtils.h"
 
+#include "Replication/Client/ClientUtils.h"
 #include "Replication/Client/Online/OnlineClient.h"
 #include "Replication/Editor/Model/Object/IObjectHierarchyModel.h"
 #include "Replication/Editor/View/IMultiReplicationStreamEditor.h"
 #include "Replication/Editor/View/IReplicationStreamEditor.h"
 #include "Replication/Submission/MultiEdit/ReassignObjectPropertiesLogic.h"
-#include "Widgets/ActiveSession/Replication/Client/ClientUtils.h"
 #include "Widgets/ActiveSession/Replication/Client/Multi/Columns/MultiStreamColumns.h"
 
 #include "Framework/Commands/UIAction.h"
diff --git a/Engine/Plugins/Developer/Concert/ConcertApp/MultiUserClient/Source/MultiUserClient/Private/Widgets/ActiveSession/Replication/Client/Multi/Columns/AssignProperty/MultiStreamColumns_AssignProperties.cpp b/Engine/Plugins/Developer/Concert/ConcertApp/MultiUserClient/Source/MultiUserClient/Private/Widgets/ActiveSession/Replication/Client/Multi/Columns/AssignProperty/MultiStreamColumns_AssignProperties.cpp
--- a/Engine/Plugins/Developer/Concert/ConcertApp/MultiUserClient/Source/MultiUserClient/Private/Widgets/ActiveSession/Replication/Client/Multi/Columns/AssignProperty/MultiStreamColumns_AssignProperties.cpp
+++ b/Engine/Plugins/Developer/Concert/ConcertApp/MultiUserClient/Source/MultiUserClient/Private/Widgets/ActiveSession/Replication/Client/Multi/Columns/AssignProperty/MultiStreamColumns_AssignProperties.cpp
@@ -13,7 +13,7 @@
 #include "SAssignPropertyComboBox.h"
 #include "Replication/Editor/View/Column/IPropertyTreeColumn.h"
 #include "Replication/Editor/View/Column/ReplicationColumnDelegates.h"
-#include "Widgets/ActiveSession/Replication/Client/ClientUtils.h"
+#include "Replication/Client/ClientUtils.h"
 
 #define LOCTEXT_NAMESPACE "AssignPropertyColumn"
 
diff --git a/Engine/Plugins/Developer/Concert/ConcertApp/MultiUserClient/Source/MultiUserClient/Private/Widgets/ActiveSession/Replication/Client/Multi/Columns/AssignProperty/SAssignPropertyComboBox.cpp b/Engine/Plugins/Developer/Concert/ConcertApp/MultiUserClient/Source/MultiUserClient/Private/Widgets/ActiveSession/Replication/Client/Multi/Columns/AssignProperty/SAssignPropertyComboBox.cpp
--- a/Engine/Plugins/Developer/Concert/ConcertApp/MultiUserClient/Source/MultiUserClient/Private/Widgets/ActiveSession/Replication/Client/Multi/Columns/AssignProperty/SAssignPropertyComboBox.cpp
+++ b/Engine/Plugins/Developer/Concert/ConcertApp/MultiUserClient/Source/MultiUserClient/Private/Widgets/ActiveSession/Replication/Client/Multi/Columns/AssignProperty/SAssignPropertyComboBox.cpp
@@ -4,11 +4,11 @@
 
 #include "AssignPropertyModel.h"
 #include "IConcertClient.h"
+#include "Replication/Client/ClientUtils.h"
 #include "Replication/Client/Online/OnlineClient.h"
 #include "Replication/Client/Online/OnlineClientManager.h"
 #include "Replication/Editor/Model/PropertyUtils.h"
 #include "Widgets/ActiveSession/Replication/Misc/SNoClients.h"
-#include "Widgets/ActiveSession/Replication/Client/ClientUtils.h"
 #include "Widgets/Client/ClientInfoHelpers.h"
 #include "Widgets/Client/SHorizontalClientList.h"
 #include "Widgets/Client/SLocalClientName.h"
diff --git a/Engine/Plugins/Developer/Concert/ConcertApp/MultiUserClient/Source/MultiUserClient/Private/Widgets/ActiveSession/Replication/Client/Multi/Columns/AssignedClients/MultiStreamColumns_AssignedClients.cpp b/Engine/Plugins/Developer/Concert/ConcertApp/MultiUserClient/Source/MultiUserClient/Private/Widgets/ActiveSession/Replication/Client/Multi/Columns/AssignedClients/MultiStreamColumns_AssignedClients.cpp
--- a/Engine/Plugins/Developer/Concert/ConcertApp/MultiUserClient/Source/MultiUserClient/Private/Widgets/ActiveSession/Replication/Client/Multi/Columns/AssignedClients/MultiStreamColumns_AssignedClients.cpp
+++ b/Engine/Plugins/Developer/Concert/ConcertApp/MultiUserClient/Source/MultiUserClient/Private/Widgets/ActiveSession/Replication/Client/Multi/Columns/AssignedClients/MultiStreamColumns_AssignedClients.cpp
@@ -4,9 +4,9 @@
 
 #include "IConcertClient.h"
 #include "MultiUserReplicationStyle.h"
+#include "Replication/Client/ClientUtils.h"
 #include "Replication/Editor/View/IMultiReplicationStreamEditor.h"
 #include "SAssignedClientsWidget.h"
-#include "Widgets/ActiveSession/Replication/Client/ClientUtils.h"
 #include "Widgets/ActiveSession/Replication/Client/Multi/Columns/MultiStreamColumns.h"
 #include "Widgets/Client/ClientInfoHelpers.h"
 #include "Widgets/Client/SHorizontalClientList.h"
