MU replication now shows content that offline clients will regain upon rejoining a session.

A FOfflineClient  is created for every client that was in the session in the past;
it is used to display the replication content disconnected clients will receive when they rejoin.

FUnifiedClientView combines online and offlien client managers. This interface is now used by the UI for displaying the replication content.

diff --git a/Engine/Plugins/Developer/Concert/ConcertApp/MultiUserClient/Source/MultiUserClient/Private/Replication/Authority/IClientAuthoritySynchronizer.h b/Engine/Plugins/Developer/Concert/ConcertApp/MultiUserClient/Source/MultiUserClient/Private/Replication/Authority/IClientAuthoritySynchronizer.h
--- a/Engine/Plugins/Developer/Concert/ConcertApp/MultiUserClient/Source/MultiUserClient/Private/Replication/Authority/IClientAuthoritySynchronizer.h
+++ b/Engine/Plugins/Developer/Concert/ConcertApp/MultiUserClient/Source/MultiUserClient/Private/Replication/Authority/IClientAuthoritySynchronizer.h
@@ -3,7 +3,6 @@
 #pragma once
 
 #include "Delegates/Delegate.h"
-#include "Templates/UnrealTemplate.h"
 
 struct FSoftObjectPath;
 
@@ -30,7 +29,6 @@ namespace UE::MultiUserClient::Replication
 	/** Util base class for implementing the events */
 	class FAuthoritySynchronizer_Base
 		: public IClientAuthoritySynchronizer
-		, public FNoncopyable
 	{
 	public:
 
diff --git a/Engine/Plugins/Developer/Concert/ConcertApp/MultiUserClient/Source/MultiUserClient/Private/Replication/Client/Offline/EndpointCache.cpp b/Engine/Plugins/Developer/Concert/ConcertApp/MultiUserClient/Source/MultiUserClient/Private/Replication/Client/Offline/EndpointCache.cpp
new file mode 100644
--- /dev/null
+++ b/Engine/Plugins/Developer/Concert/ConcertApp/MultiUserClient/Source/MultiUserClient/Private/Replication/Client/Offline/EndpointCache.cpp
@@ -0,0 +1,69 @@
+﻿// Copyright Epic Games, Inc. All Rights Reserved.
+
+#include "EndpointCache.h"
+
+#include "IConcertClientWorkspace.h"
+#include "Replication/Misc/StreamAndAuthorityPredictionUtils.h"
+
+namespace UE::MultiUserClient::Replication
+{
+	void FEndpointCache::UpdateEndpoints(const IConcertClientWorkspace& Workspace)
+	{
+		while (NextFirstActivityToFetch < Workspace.GetLastActivityId())
+		{
+			constexpr int64 MaxToFetch = 1000;
+			TArray<FConcertSessionActivity> DummyActivities;
+			TMap<FGuid, FConcertClientInfo> NewEndpoints;
+			Workspace.GetActivities(NextFirstActivityToFetch, MaxToFetch, NewEndpoints, DummyActivities);
+			NextFirstActivityToFetch += DummyActivities.Num();
+
+			MergeEndpointsWith(NewEndpoints);
+		}
+	}
+
+	int32 FEndpointCache::FindClientIndexByEndpointId(const FGuid& EndpointId) const
+	{
+		return EndpointMetaData.IndexOfByPredicate([&EndpointId](const FClientMetaData& MetaData)
+		{
+			return MetaData.AssociatedEndpoints.Contains(EndpointId);
+		});
+	}
+
+	void FEndpointCache::MergeEndpointsWith(const TMap<FGuid, FConcertClientInfo>& NewEndpoints)
+	{
+		for (const TPair<FGuid, FConcertClientInfo>& NewEndpointPair : NewEndpoints)
+		{
+			const FGuid& EndpointId = NewEndpointPair.Key;
+			const FConcertClientInfo& EncounteredClientInfo = NewEndpointPair.Value;
+				
+			// We'll consider a client equal if its DisplayName and DeviceName are equal...
+			// This will not work if you have two UE instances with the same name on the same machine (unsupported).
+			const int32 EndpointIndex = KnownClients.IndexOfByPredicate([&EncounteredClientInfo](const FConcertClientInfo& ExistingClientInfo)
+			{
+				return ConcertSyncCore::Replication::AreLogicallySameClients(EncounteredClientInfo, ExistingClientInfo);
+			});
+			const bool bAlreadyKnown = KnownClients.IsValidIndex(EndpointIndex);
+			if (bAlreadyKnown)
+			{
+				EndpointMetaData[EndpointIndex].AssociatedEndpoints.AddUnique(EndpointId);
+				continue;
+			}
+
+			const int32 MetaDataIndex = EndpointMetaData.IndexOfByPredicate([&](const FClientMetaData& MetaData)
+			{
+				return MetaData.GetLastKnownEndpointId() == EndpointId;
+			});
+			const bool bIsNameChange = EndpointMetaData.IsValidIndex(MetaDataIndex);
+			if (bIsNameChange)
+			{
+				KnownClients[EndpointIndex] = EncounteredClientInfo;
+				EndpointMetaData[EndpointIndex].AssociatedEndpoints.AddUnique(EndpointId);
+			}
+			else
+			{
+				KnownClients.Add(EncounteredClientInfo);
+				EndpointMetaData.Emplace(NewEndpointPair.Key);
+			}
+		}
+	}
+}
diff --git a/Engine/Plugins/Developer/Concert/ConcertApp/MultiUserClient/Source/MultiUserClient/Private/Replication/Client/Offline/EndpointCache.h b/Engine/Plugins/Developer/Concert/ConcertApp/MultiUserClient/Source/MultiUserClient/Private/Replication/Client/Offline/EndpointCache.h
new file mode 100644
--- /dev/null
+++ b/Engine/Plugins/Developer/Concert/ConcertApp/MultiUserClient/Source/MultiUserClient/Private/Replication/Client/Offline/EndpointCache.h
@@ -0,0 +1,65 @@
+﻿// Copyright Epic Games, Inc. All Rights Reserved.
+
+#pragma once
+
+#include "ConcertMessageData.h"
+
+#include "Containers/Map.h"
+#include "HAL/Platform.h"
+
+class IConcertClientWorkspace;
+
+namespace UE::MultiUserClient::Replication
+{
+	/** Goes through activity history and saves all endpoints is encounters. */
+	class FEndpointCache
+	{
+	public:
+
+		/** Updates the list of endpoints, if needed. */
+		void UpdateEndpoints(const IConcertClientWorkspace& Workspace);
+
+		/** @return An index to GetKnownClients() if a client with EndpointId is known. */
+		int32 FindClientIndexByEndpointId(const FGuid& EndpointId) const;
+
+		/** @return The last endpoint that the client at GetKnownClients()[Index] was associated with. */
+		const FGuid& GetLastAssociatedEndpoint(int32 Index) const { return EndpointMetaData[Index].GetLastKnownEndpointId(); }
+		
+		/** @return The clients that are known */
+		const TArray<FConcertClientInfo>& GetKnownClients() const { return KnownClients; }
+
+	private:
+
+		/** The next activity to start fetching from.*/
+		int64 NextFirstActivityToFetch = 1; // Activity IDs start with 1
+
+		/**
+		 * The clients encountered thus far. Multiple endpoint IDs are associated with a client.
+		 * 
+		 * Client info is grouped by DeviceName and DisplayName even though they are technically new endpoints.
+		 * Example: Client1 joins, leaves, then joins again.
+		 * On the server, this would generate two distinct endpoints with different endpoint IDs, but we'll merge them into one assuming it's the
+		 * same machine. This will not work if you have two UE instances with the same name on the same machine, which is unsupported by this class.
+		 */
+		TArray<FConcertClientInfo> KnownClients;
+
+		struct FClientMetaData
+		{
+			/** All endpoint IDs that are associated with this client. */
+			TArray<FGuid> AssociatedEndpoints;
+
+			const FGuid& GetLastKnownEndpointId() const { return AssociatedEndpoints[AssociatedEndpoints.Num() - 1]; }
+
+			explicit FClientMetaData(const FGuid& LastKnownEndpointId)
+				: AssociatedEndpoints({ LastKnownEndpointId })
+			{}
+		};
+		/** Additional data about. Each index maps to KnownEndpoints. Array based to make GetKnownEndpoints interface easier. */
+		TArray<FClientMetaData> EndpointMetaData;
+
+		/** Adds NewEndpoints to EndpointMetaData. */
+		void MergeEndpointsWith(const TMap<FGuid, FConcertClientInfo>& NewEndpoints);
+	};
+}
+
+
diff --git a/Engine/Plugins/Developer/Concert/ConcertApp/MultiUserClient/Source/MultiUserClient/Private/Replication/Client/Offline/OfflineClient.cpp b/Engine/Plugins/Developer/Concert/ConcertApp/MultiUserClient/Source/MultiUserClient/Private/Replication/Client/Offline/OfflineClient.cpp
new file mode 100644
--- /dev/null
+++ b/Engine/Plugins/Developer/Concert/ConcertApp/MultiUserClient/Source/MultiUserClient/Private/Replication/Client/Offline/OfflineClient.cpp
@@ -0,0 +1,25 @@
+﻿// Copyright Epic Games, Inc. All Rights Reserved.
+
+#include "OfflineClient.h"
+
+#include "Replication/ReplicationWidgetFactories.h"
+
+namespace UE::MultiUserClient::Replication
+{
+	FOfflineClient::FOfflineClient(
+		IConcertClientWorkspace& InWorkspace,
+		FConcertClientInfo InClientInfo,
+		const FGuid& LastAssociatedEndpoint
+		)
+		: ClientInfo(MoveTemp(InClientInfo))
+		, LastAssociatedEndpoint(LastAssociatedEndpoint)
+		, ContentPredictor(InWorkspace, ClientInfo)
+		, StreamModel(
+			ConcertSharedSlate::CreateReadOnlyStreamModel(
+				TAttribute<const FConcertObjectReplicationMap*>::CreateLambda(
+					[this]{ return &ContentPredictor.GetServerState(); }
+					)
+				)
+			)
+	{}
+}
diff --git a/Engine/Plugins/Developer/Concert/ConcertApp/MultiUserClient/Source/MultiUserClient/Private/Replication/Client/Offline/OfflineClient.h b/Engine/Plugins/Developer/Concert/ConcertApp/MultiUserClient/Source/MultiUserClient/Private/Replication/Client/Offline/OfflineClient.h
new file mode 100644
--- /dev/null
+++ b/Engine/Plugins/Developer/Concert/ConcertApp/MultiUserClient/Source/MultiUserClient/Private/Replication/Client/Offline/OfflineClient.h
@@ -0,0 +1,67 @@
+﻿// Copyright Epic Games, Inc. All Rights Reserved.
+
+#pragma once
+
+#include "ConcertMessageData.h"
+#include "RejoinStreamAndAuthorityPredictor.h"
+
+#include "Templates/UnrealTemplate.h"
+#include "Templates/SharedPointer.h"
+
+namespace UE::ConcertSharedSlate { class IReplicationStreamModel; }
+
+namespace UE::MultiUserClient::Replication
+{
+	/**
+	 * Info about a client that had joined the MU session but no are longer is present.
+	 * This is used to display info about the properties the client would receive if they joined back into the session. 
+	 */
+	class FOfflineClient : public FNoncopyable
+	{
+	public:
+
+		FOfflineClient(
+			IConcertClientWorkspace& InWorkspace UE_LIFETIMEBOUND,
+			FConcertClientInfo InClientInfo,
+			const FGuid& LastAssociatedEndpoint
+			);
+
+		FORCEINLINE const FConcertClientInfo& GetClientInfo() const { return ClientInfo; }
+		FORCEINLINE const FGuid& GetLastAssociatedEndpoint() const { return LastAssociatedEndpoint; }
+		FORCEINLINE const FConcertBaseStreamInfo& GetPredictedStream() const { return ContentPredictor.GetPredictedStream(); }
+		
+		/**
+		 * This is used so the UI can construct the IReplicationStreamEditor.
+		 * @see CreateBaseStreamEditor and FCreateEditorParams.
+		 * 
+		 * @note You must make sure to release this object when this FReplicationClient is destroyed.
+		 * Listen for events on the owning FOfflineClientManager::OnDisconnectedClientsChanged.
+		 * @see FReplicationClientManager and FMultiUserReplicationManager
+		 */
+		TSharedRef<ConcertSharedSlate::IReplicationStreamModel> GetStreamModel() const { return StreamModel; }
+
+		DECLARE_MULTICAST_DELEGATE(FStreamPredictionChanged);
+		/** Broadcasts when the stream, which the offline client will get upon rejoining, has changed. */
+		FStreamPredictionChanged& OnStreamPredictionChanged() { return OnStreamPredictionChangedDelegate; }
+
+	private:
+
+		/** The client's last registered display data. */
+		const FConcertClientInfo ClientInfo;
+		/** The last endpoint ID this client was associated with- */
+		const FGuid LastAssociatedEndpoint;
+
+		/** Predicts the stream as activities are produced. */
+		FRejoinStreamAndAuthorityPredictor ContentPredictor;
+		/**
+		 * Read-only model of the state predicted by StreamSynchronizer.
+		 * TODO DP: In the future, this may be made editable to allow the user to edit the future state in the UI.
+		 */
+		const TSharedRef<ConcertSharedSlate::IReplicationStreamModel> StreamModel;
+		
+		/** Broadcasts when the stream, which the offline client will get upon rejoining, has changed. */
+		FStreamPredictionChanged OnStreamPredictionChangedDelegate;
+	};
+}
+
+
diff --git a/Engine/Plugins/Developer/Concert/ConcertApp/MultiUserClient/Source/MultiUserClient/Private/Replication/Client/Offline/OfflineClientManager.cpp b/Engine/Plugins/Developer/Concert/ConcertApp/MultiUserClient/Source/MultiUserClient/Private/Replication/Client/Offline/OfflineClientManager.cpp
new file mode 100644
--- /dev/null
+++ b/Engine/Plugins/Developer/Concert/ConcertApp/MultiUserClient/Source/MultiUserClient/Private/Replication/Client/Offline/OfflineClientManager.cpp
@@ -0,0 +1,115 @@
+﻿// Copyright Epic Games, Inc. All Rights Reserved.
+
+#include "OfflineClientManager.h"
+
+#include "IConcertClientWorkspace.h"
+#include "IConcertSyncClient.h"
+#include "Replication/Client/ClientUtils.h"
+#include "Replication/Client/Online/OnlineClientManager.h"
+#include "Replication/Misc/StreamAndAuthorityPredictionUtils.h"
+
+namespace UE::MultiUserClient::Replication
+{
+	FOfflineClientManager::FOfflineClientManager(
+		IConcertSyncClient& ClientInstance,
+		FOnlineClientManager& OnlineClientManager
+		)
+		: ClientInstance(ClientInstance)
+		, OnlineClientManager(OnlineClientManager)
+	{
+		OnlineClientManager.OnRemoteClientsChanged().AddRaw(this, &FOfflineClientManager::RefreshOfflineClients);
+		RefreshOfflineClients();
+	}
+
+	FOfflineClientManager::~FOfflineClientManager()
+	{
+		OnlineClientManager.OnRemoteClientsChanged().RemoveAll(this);
+	}
+
+	const FOfflineClient* FOfflineClientManager::FindClient(const FGuid& EndpointId) const
+	{
+		const int32 CacheIndex = EndpointCache.FindClientIndexByEndpointId(EndpointId);
+		const TArray<FConcertClientInfo>& KnownClients = EndpointCache.GetKnownClients();
+		if (!KnownClients.IsValidIndex(CacheIndex))
+		{
+			return nullptr;
+		}
+
+		const int32 ClientIndex = FindClientIndexByInfo(KnownClients[CacheIndex]);
+		return Clients.IsValidIndex(ClientIndex) ? Clients[ClientIndex].Get() : nullptr;
+	}
+
+	FOfflineClient* FOfflineClientManager::FindClient(const FGuid& EndpointId)
+	{
+		return const_cast<FOfflineClient*>(
+			const_cast<const FOfflineClientManager*>(this)->FindClient(EndpointId)
+			);
+	}
+
+	void FOfflineClientManager::RefreshOfflineClients()
+	{
+		const TSharedPtr<IConcertClientWorkspace> Workspace = ClientInstance.GetWorkspace();
+		check(Workspace);
+		EndpointCache.UpdateEndpoints(*Workspace);
+
+		const bool bChanged = UpdateClientList(*Workspace);
+		if (bChanged)
+		{
+			OnClientsChangedDelegate.Broadcast();
+		}
+	}
+
+	bool FOfflineClientManager::UpdateClientList(IConcertClientWorkspace& Workspace)
+	{
+		bool bChanged = false;
+		const TArray<FConcertClientInfo>& KnownClients = EndpointCache.GetKnownClients();
+		for (int32 i = 0; i < KnownClients.Num(); ++i)
+		{
+			const FConcertClientInfo& KnownEndpointInfo = KnownClients[i];
+			if (IsClientOnline(KnownEndpointInfo))
+			{
+				const int32 Index = FindClientIndexByInfo(KnownEndpointInfo);
+				if (Clients.IsValidIndex(Index))
+				{
+					bChanged = true;
+					
+					TUniquePtr<FOfflineClient>& OfflineClient = Clients[Index];
+					OnPreClientRemovedDelegate.Broadcast(*OfflineClient.Get());
+					
+					Clients.RemoveAt(Index);
+				}
+			}
+			else
+			{
+				bChanged = true;
+				
+				const int32 AddedIndex = Clients.Emplace(
+					MakeUnique<FOfflineClient>(Workspace, KnownEndpointInfo, EndpointCache.GetLastAssociatedEndpoint(i))
+					);
+				OnPostClientAddedDelegate.Broadcast(*Clients[AddedIndex].Get());
+			}
+		}
+		return bChanged;
+	}
+
+	bool FOfflineClientManager::IsClientOnline(const FConcertClientInfo& QueryClientInfo) const
+	{
+		bool bIsOnline = false;
+		OnlineClientManager.ForEachClient([this, &QueryClientInfo, &bIsOnline](const FOnlineClient& Client)
+		{
+			FConcertClientInfo ClientInfo;
+			const bool bGotInfo = ClientUtils::GetClientDisplayInfo(*ClientInstance.GetConcertClient(), Client.GetEndpointId(), ClientInfo);
+			bIsOnline = bGotInfo && ConcertSyncCore::Replication::AreLogicallySameClients(QueryClientInfo, ClientInfo);
+			return bIsOnline ? EBreakBehavior::Break : EBreakBehavior::Continue;
+		});
+		return bIsOnline;
+	}
+	
+	int32 FOfflineClientManager::FindClientIndexByInfo(const FConcertClientInfo& KnownEndpointInfo) const
+	{
+		return Clients.IndexOfByPredicate([&KnownEndpointInfo](const TUniquePtr<FOfflineClient>& OfflineClient)
+		{
+			return ConcertSyncCore::Replication::AreLogicallySameClients(OfflineClient->GetClientInfo(), KnownEndpointInfo);
+		});
+	}
+}
diff --git a/Engine/Plugins/Developer/Concert/ConcertApp/MultiUserClient/Source/MultiUserClient/Private/Replication/Client/Offline/OfflineClientManager.h b/Engine/Plugins/Developer/Concert/ConcertApp/MultiUserClient/Source/MultiUserClient/Private/Replication/Client/Offline/OfflineClientManager.h
new file mode 100644
--- /dev/null
+++ b/Engine/Plugins/Developer/Concert/ConcertApp/MultiUserClient/Source/MultiUserClient/Private/Replication/Client/Offline/OfflineClientManager.h
@@ -0,0 +1,106 @@
+﻿// Copyright Epic Games, Inc. All Rights Reserved.
+
+#pragma once
+
+#include "EndpointCache.h"
+#include "OfflineClient.h"
+
+#include "Delegates/Delegate.h"
+#include "HAL/Platform.h"
+#include "Misc/EBreakBehavior.h"
+#include "Templates/UniquePtr.h"
+#include "Templates/UnrealTemplate.h"
+
+#include <type_traits>
+
+class IConcertClientWorkspace;
+class IConcertSyncClient;
+
+namespace UE::MultiUserClient::Replication
+{
+	class FOnlineClientManager;
+	
+	/**
+	 * Manages clients that had joined the MU session but no are longer is present.
+	 *
+	 * Endpoint IDs are saved in the database. Endpoint IDs with the same DisplayName and DeviceName are associated with single FOfflineClient.
+	 * While FEndpointCache detects name changes, the rest of the system is not currently set up to handle them, e.g. FindClient.
+	 */
+	class FOfflineClientManager : public FNoncopyable
+	{
+	public:
+		
+		/**
+		 * @param ClientInstance Used to obtain disconnected endpoint info from the server. The caller ensures it outlives the new object.
+		 * @param OnlineClientManager Used to determine whether a client should be considered offline. The caller ensures it outlives the new object.
+		 */
+		FOfflineClientManager(
+			IConcertSyncClient& ClientInstance UE_LIFETIMEBOUND,
+			FOnlineClientManager& OnlineClientManager UE_LIFETIMEBOUND
+			);
+		~FOfflineClientManager();
+
+		/** @return The offline client that was associated with this endpoint ID. */
+		const FOfflineClient* FindClient(const FGuid& EndpointId) const;
+		/** @return The offline client that was associated with this endpoint ID. */
+		FOfflineClient* FindClient(const FGuid& EndpointId);
+		
+		/** Iterates through every offline client. */
+		template<typename TLambda> requires std::is_invocable_r_v<EBreakBehavior, TLambda, FOfflineClient&>
+		void ForEachClient(TLambda&& Consumer) const;
+
+		DECLARE_MULTICAST_DELEGATE(FOfflineClientsChanged);
+		/** Broadcasts when Clients is changed. */
+		FOfflineClientsChanged& OnClientsChanged() { return OnClientsChangedDelegate; }
+		
+		DECLARE_MULTICAST_DELEGATE_OneParam(FOfflineClientDelegate, FOfflineClient&);
+		/** Called just after an offline client is has been added. Called before OnClientsChanged. */
+		FOfflineClientDelegate& OnPostClientAdded() { return OnPostClientAddedDelegate; }
+		/** Called just before an offline client is about to be removed. Called before OnClientsChanged. */
+		FOfflineClientDelegate& OnPreClientRemoved() { return OnPreClientRemovedDelegate; }
+
+	private:
+
+		/** The local client connected to the MU session. Used to obtain disconnected endpoint info from the server. */
+		IConcertSyncClient& ClientInstance;
+		/** Used to determine whether a client should be considered offline. */
+		FOnlineClientManager& OnlineClientManager;
+
+		/** Info about clients that were once connected but are no longer. */
+		TArray<TUniquePtr<FOfflineClient>> Clients;
+
+		/** Keeps track of all known endpoints. */
+		FEndpointCache EndpointCache;
+
+		/** Broadcasts when Clients is changed. */
+		FOfflineClientsChanged OnClientsChangedDelegate;
+		/** Called just after an offline client is has been added. Called before OnClientsChanged. */
+		FOfflineClientDelegate OnPostClientAddedDelegate;
+		/** Called just before an offline client is about to be removed. Called before OnClientsChanged. */
+		FOfflineClientDelegate OnPreClientRemovedDelegate;
+		
+		/** Refreshes Clients based on which clients are online. */
+		void RefreshOfflineClients();
+		/** Iterates through EndpointCache and updates Clients.*/
+		bool UpdateClientList(IConcertClientWorkspace& Workspace);
+		/** @return Whether a client with the given info is online. */
+		bool IsClientOnline(const FConcertClientInfo& QueryClientInfo) const;
+
+		/** @return Index to Clients that matches ClientInfo. */
+		int32 FindClientIndexByInfo(const FConcertClientInfo& ClientInfo) const;
+	};
+
+	template <typename TLambda> requires std::is_invocable_r_v<EBreakBehavior, TLambda, FOfflineClient&>
+	void FOfflineClientManager::ForEachClient(TLambda&& Consumer) const
+	{
+		for (const TUniquePtr<FOfflineClient>& ClientPtr : Clients)
+		{
+			if (Consumer(*ClientPtr) == EBreakBehavior::Break)
+			{
+				break;
+			}
+		}
+	}
+}
+
+
diff --git a/Engine/Plugins/Developer/Concert/ConcertApp/MultiUserClient/Source/MultiUserClient/Private/Replication/Client/Offline/RejoinStreamAndAuthorityPredictor.cpp b/Engine/Plugins/Developer/Concert/ConcertApp/MultiUserClient/Source/MultiUserClient/Private/Replication/Client/Offline/RejoinStreamAndAuthorityPredictor.cpp
new file mode 100644
--- /dev/null
+++ b/Engine/Plugins/Developer/Concert/ConcertApp/MultiUserClient/Source/MultiUserClient/Private/Replication/Client/Offline/RejoinStreamAndAuthorityPredictor.cpp
@@ -0,0 +1,67 @@
+﻿// Copyright Epic Games, Inc. All Rights Reserved.
+
+#include "RejoinStreamAndAuthorityPredictor.h"
+
+#include "EndpointCache.h"
+#include "Replication/Messages/ReplicationActivity.h"
+
+#include "Replication/Misc/ClientPredictionUtils.h"
+
+namespace UE::MultiUserClient::Replication
+{
+	namespace Private
+	{
+		static void ExtractMultiUserContent(
+			const TArray<FConcertBaseStreamInfo>& Streams,
+			const TArray<FConcertObjectInStreamID>& Authority,
+			FConcertBaseStreamInfo& OutPredictedStream,
+			TSet<FSoftObjectPath>& OutPredictedAuthority
+			)
+		{
+			// A user may have created their custom streams, so find the MultiUserStreamID one.
+			const FConcertBaseStreamInfo* MultiUserStream = Streams.FindByPredicate([](const FConcertBaseStreamInfo& Stream)
+			{
+				return Stream.Identifier == MultiUserStreamID;
+			});
+			if (!MultiUserStream)
+			{
+				return;
+			}
+			
+			OutPredictedStream = *MultiUserStream;
+			// Again... a user may have created their custom streams, so only return the authority of the MultiUserStreamID one.
+			Algo::TransformIf(Authority, OutPredictedAuthority,
+				[](const FConcertObjectInStreamID& ObjectInStreamID) { return ObjectInStreamID.StreamId == MultiUserStreamID; },
+				[](const FConcertObjectInStreamID& ObjectInStreamID) { return ObjectInStreamID.Object; }
+			);
+		}
+		
+		static void AnalyzeActivityHistory(
+			const IConcertClientWorkspace& Workspace,
+			const FConcertClientInfo& ClientInfo,
+			FConcertBaseStreamInfo& OutPredictedStream,
+			TSet<FSoftObjectPath>& OutPredictedAuthority
+			)
+		{
+			TArray<FConcertBaseStreamInfo> Streams;
+			TArray<FConcertObjectInStreamID> Authority;
+			const TOptional<int64> FoundActivity = ConcertSyncClient::Replication::IncrementalBacktrackActivityHistoryForActivityThatSetsContent(
+				Workspace, ClientInfo, Streams, Authority
+				);
+			if (FoundActivity.IsSet())
+			{
+				ExtractMultiUserContent(Streams, Authority, OutPredictedStream, OutPredictedAuthority);
+			}
+		}
+	}
+	
+	FRejoinStreamAndAuthorityPredictor::FRejoinStreamAndAuthorityPredictor(const IConcertClientWorkspace& InWorkspace, FConcertClientInfo InClientInfo)
+		: ClientInfo(MoveTemp(InClientInfo))
+		, PredictedStream({ .Identifier = MultiUserStreamID })
+	{
+		// For now, we only look at the FConcertSyncReplicationPayload_LeaveReplication activity. It can be looked up once.
+		// In the future, we'll add an event for FConcertReplication_PutState_Requests, which can affect offline clients, invoked by presets.
+		// Then, we'll have to listen for activity changes.
+		Private::AnalyzeActivityHistory(InWorkspace, ClientInfo, PredictedStream, PredictedAuthority);
+	}
+}
diff --git a/Engine/Plugins/Developer/Concert/ConcertApp/MultiUserClient/Source/MultiUserClient/Private/Replication/Client/Offline/RejoinStreamAndAuthorityPredictor.h b/Engine/Plugins/Developer/Concert/ConcertApp/MultiUserClient/Source/MultiUserClient/Private/Replication/Client/Offline/RejoinStreamAndAuthorityPredictor.h
new file mode 100644
--- /dev/null
+++ b/Engine/Plugins/Developer/Concert/ConcertApp/MultiUserClient/Source/MultiUserClient/Private/Replication/Client/Offline/RejoinStreamAndAuthorityPredictor.h
@@ -0,0 +1,55 @@
+﻿// Copyright Epic Games, Inc. All Rights Reserved.
+
+#pragma once
+
+#include "ConcertMessageData.h"
+#include "Replication/Authority/IClientAuthoritySynchronizer.h"
+#include "Replication/Data/ReplicationStream.h"
+#include "Replication/Stream/IClientStreamSynchronizer.h"
+#include "Replication/Stream/MultiUserStreamId.h"
+
+#include "Containers/Set.h"
+#include "HAL/Platform.h"
+
+class IConcertClientWorkspace;
+
+namespace UE::MultiUserClient::Replication
+{
+	/**
+	 * Predicts the stream and authority state a client will have when they rejoin the session.
+	 * The prediction returns the stream content in the latest EConcertSyncReplicationActivityType::LeaveReplication pertaining to ClientInfo.
+	 */
+	class FRejoinStreamAndAuthorityPredictor
+		: public FStreamSynchronizer_Base
+		, public FAuthoritySynchronizer_Base
+		, public FNoncopyable
+	{
+	public:
+		
+		FRejoinStreamAndAuthorityPredictor(const IConcertClientWorkspace& InWorkspace UE_LIFETIMEBOUND, FConcertClientInfo InClientInfo);
+
+		FORCEINLINE const FConcertBaseStreamInfo& GetPredictedStream() const { return PredictedStream; }
+		
+		//~ Begin IClientStreamSynchronizer Interface
+		virtual FGuid GetStreamId() const override { return MultiUserStreamID; }
+		virtual const FConcertObjectReplicationMap& GetServerState() const override { return PredictedStream.ReplicationMap; }
+		virtual const FConcertStreamFrequencySettings& GetFrequencySettings() const override { return PredictedStream.FrequencySettings; }
+		//~ End IClientStreamSynchronizer Interface
+
+		//~ Begin IClientAuthoritySynchronizer Interface
+		virtual bool HasAnyAuthority() const override { return !PredictedAuthority.IsEmpty(); }
+		virtual bool HasAuthorityOver(const FSoftObjectPath& ObjectPath) const override { return PredictedAuthority.Contains(ObjectPath); }
+		//~ End IClientAuthoritySynchronizer Interface
+
+	private:
+
+		/** Client for which we're predicting the state. */
+		const FConcertClientInfo ClientInfo;
+
+		/** The stream content the local client thinks the offline client will have. */
+		FConcertBaseStreamInfo PredictedStream;
+		/** The most up to date server state of the remote client's authority. */
+		TSet<FSoftObjectPath> PredictedAuthority; 
+	};
+}
+
diff --git a/Engine/Plugins/Developer/Concert/ConcertApp/MultiUserClient/Source/MultiUserClient/Private/Replication/Client/UnifiedClientView.cpp b/Engine/Plugins/Developer/Concert/ConcertApp/MultiUserClient/Source/MultiUserClient/Private/Replication/Client/UnifiedClientView.cpp
new file mode 100644
--- /dev/null
+++ b/Engine/Plugins/Developer/Concert/ConcertApp/MultiUserClient/Source/MultiUserClient/Private/Replication/Client/UnifiedClientView.cpp
@@ -0,0 +1,110 @@
+﻿// Copyright Epic Games, Inc. All Rights Reserved.
+
+#include "UnifiedClientView.h"
+
+#include "IConcertSyncClient.h"
+#include "Offline/OfflineClient.h"
+#include "Offline/OfflineClientManager.h"
+#include "Online/OnlineClientManager.h"
+#include "Replication/Client/ClientUtils.h"
+
+#include "Misc/Optional.h"
+
+namespace UE::MultiUserClient::Replication
+{
+	FUnifiedClientView::FUnifiedClientView(
+		IConcertSyncClient& SyncClient,
+		FOnlineClientManager& OnlineClientManager,
+		FOfflineClientManager& OfflineClientManager
+		)
+		: SyncClient(SyncClient)
+		, OnlineClientManager(OnlineClientManager)
+		, OfflineClientManager(OfflineClientManager)
+		, StreamCache(*this, OnlineClientManager, OfflineClientManager)
+	{
+		// Technically it's enough to subscribe only to OfflineClientManager because that changes when OnlineClientManager changes...
+		// but we'll just be safe for now.
+		OnlineClientManager.OnRemoteClientsChanged().AddRaw(this, &FUnifiedClientView::BroadcastOnClientsChanged);
+		OfflineClientManager.OnClientsChanged().AddRaw(this, &FUnifiedClientView::BroadcastOnClientsChanged);
+	}
+
+	FUnifiedClientView::~FUnifiedClientView()
+	{
+		OnlineClientManager.OnRemoteClientsChanged().RemoveAll(this);
+		OfflineClientManager.OnClientsChanged().RemoveAll(this);
+	}
+
+	TArray<FGuid> FUnifiedClientView::GetClients() const
+	{
+		TArray<FGuid> Result;
+		ForEachClient([&Result](const FGuid& ClientId)
+		{
+			Result.Add(ClientId);
+			return EBreakBehavior::Continue;
+		});
+		return Result;
+	}
+
+	TArray<FGuid> FUnifiedClientView::GetOnlineClients() const
+	{
+		TArray<FGuid> Result;
+		ForEachOnlineClient([&Result](const FGuid& ClientId)
+		{
+			Result.Add(ClientId);
+			return EBreakBehavior::Continue;
+		});
+		return Result;
+	}
+
+	TOptional<FConcertClientInfo> FUnifiedClientView::GetClientInfoByEndpoint(const FGuid& EndpointId) const
+	{
+		FConcertClientInfo ClientInfo;
+		const bool bIsOnline = ClientUtils::GetClientDisplayInfo(*SyncClient.GetConcertClient(), EndpointId, ClientInfo);
+		if (bIsOnline)
+		{
+			return ClientInfo;
+		}
+		
+		if (const FOfflineClient* OfflineClient = OfflineClientManager.FindClient(EndpointId))
+		{
+			return OfflineClient->GetClientInfo();
+		}
+		return {};
+	}
+
+	TOptional<EClientType> FUnifiedClientView::GetClientType(const FGuid& EndpointId) const
+	{
+		if (OnlineClientManager.GetLocalClient().GetEndpointId() == EndpointId)
+		{
+			return EClientType::Local;
+		}
+		if (OnlineClientManager.FindClient(EndpointId))
+		{
+			return EClientType::Remote;
+		}
+		if (OfflineClientManager.FindClient(EndpointId))
+		{
+			return EClientType::Offline;
+		}
+		return {};
+	}
+
+	TSharedPtr<const ConcertSharedSlate::IReplicationStreamModel> FUnifiedClientView::GetClientStreamById(const FGuid& EndpointId) const
+	{
+		if (const FOnlineClient* Client = OnlineClientManager.FindClient(EndpointId))
+		{
+			return Client->GetClientEditModel();
+		}
+		if (const FOfflineClient* Client = OfflineClientManager.FindClient(EndpointId))
+		{
+			return Client->GetStreamModel();
+		}
+		return nullptr;
+	}
+
+	TSharedPtr<ConcertSharedSlate::IEditableReplicationStreamModel> FUnifiedClientView::GetEditableClientStreamById(const FGuid& EndpointId) const
+	{
+		const FOnlineClient* Client = OnlineClientManager.FindClient(EndpointId);
+		return Client ? Client->GetClientEditModel().ToSharedPtr() : nullptr;
+	}
+}
diff --git a/Engine/Plugins/Developer/Concert/ConcertApp/MultiUserClient/Source/MultiUserClient/Private/Replication/Client/UnifiedClientView.h b/Engine/Plugins/Developer/Concert/ConcertApp/MultiUserClient/Source/MultiUserClient/Private/Replication/Client/UnifiedClientView.h
new file mode 100644
--- /dev/null
+++ b/Engine/Plugins/Developer/Concert/ConcertApp/MultiUserClient/Source/MultiUserClient/Private/Replication/Client/UnifiedClientView.h
@@ -0,0 +1,127 @@
+﻿// Copyright Epic Games, Inc. All Rights Reserved.
+
+#pragma once
+
+#include "Online/OnlineClient.h"
+#include "Online/OnlineClientManager.h"
+#include "UnifiedStreamCache.h"
+
+#include "HAL/Platform.h"
+#include "Offline/OfflineClientManager.h"
+#include "Templates/UnrealTemplate.h"
+
+class IConcertSyncClient;
+
+struct FConcertClientInfo;
+template<typename OptionalType> struct TOptional;
+
+namespace UE::ConcertSharedSlate
+{
+	class IEditableReplicationStreamModel;
+	class IReplicationStreamModel;
+}
+
+namespace UE::MultiUserClient::Replication
+{
+	enum class EClientType : uint8
+	{
+		Local,
+		Remote,
+		Offline
+	};
+	inline bool IsOnlineClient(EClientType Type) { return Type == EClientType::Local || Type == EClientType::Remote; }
+	inline bool IsOfflineClient(EClientType Type) { return Type == EClientType::Offline; }
+	
+	/**
+	 * Uses FOnlineClientManager and FOfflineClientManager to preset a unified interface of querying client info to the UI.
+	 * This abstraction effectively implements the Adapter design pattern.
+	 *
+	 * This task of this class is not to hide the fact that there are online and offline clients: the goal is to unify the interface and put shared
+	 * logic that needs to handle online and offline client interactions here.
+	 */
+	class FUnifiedClientView : public FNoncopyable
+	{
+	public:
+		
+		FUnifiedClientView(
+			IConcertSyncClient& SyncClient UE_LIFETIMEBOUND,
+			FOnlineClientManager& OnlineClientManager UE_LIFETIMEBOUND,
+			FOfflineClientManager& OfflineClientManager UE_LIFETIMEBOUND
+			);
+		~FUnifiedClientView();
+
+		/** Enumerates the endpoint ID of every user; in the case of offline clients, only the latest associated ID is listed.. */
+		template<typename TLambda> requires std::is_invocable_r_v<EBreakBehavior, TLambda, const FGuid&>
+		void ForEachClient(TLambda&& Callback) const;
+		/** Enumerates the endpoint ID of every online client. */
+		template<typename TLambda> requires std::is_invocable_r_v<EBreakBehavior, TLambda, const FGuid&>
+		void ForEachOnlineClient(TLambda&& Callback) const;
+
+		TArray<FGuid> GetClients() const;
+		TArray<FGuid> GetOnlineClients() const;
+		
+		/** @return Client info associated with EndpointId. */
+		TOptional<FConcertClientInfo> GetClientInfoByEndpoint(const FGuid& EndpointId) const;
+		/** @return Gets the type of client */
+		TOptional<EClientType> GetClientType(const FGuid& EndpointId) const;
+		
+		/** @return The stream model of the client associated with EndpointId. */
+		TSharedPtr<const ConcertSharedSlate::IReplicationStreamModel> GetClientStreamById(const FGuid& EndpointId) const;
+		/** @return The editable stream model of the client associated with EndpointId, if EndpointId refers to an online client. */
+		TSharedPtr<ConcertSharedSlate::IEditableReplicationStreamModel> GetEditableClientStreamById(const FGuid& EndpointId) const;
+
+		const FUnifiedStreamCache& GetStreamCache() const { return StreamCache; }
+		FUnifiedStreamCache& GetStreamCache() { return StreamCache; }
+
+		/** Broadcasts when the list of online of offline clients has changed. */
+		FSimpleMulticastDelegate OnClientsChanged() { return OnClientsChangedDelegate; }
+
+	private:
+
+		/** Used to look up client info from session. */
+		IConcertSyncClient& SyncClient;
+
+		FOnlineClientManager& OnlineClientManager;
+		FOfflineClientManager& OfflineClientManager;
+		
+		/** Broadcasts when the list of online of offline clients has changed. */
+		FSimpleMulticastDelegate OnClientsChangedDelegate;
+		
+		/**
+		 * Unified interface for querying stream content efficiently.
+		 * Important: Has dependency on OnClientsChangedDelegate hence this must be constructed after.
+		 */
+		FUnifiedStreamCache StreamCache;
+
+		void BroadcastOnClientsChanged() const { OnClientsChangedDelegate.Broadcast(); }
+	};
+
+	template <typename TLambda> requires std::is_invocable_r_v<EBreakBehavior, TLambda, const FGuid&>
+	void FUnifiedClientView::ForEachClient(TLambda&& Callback) const
+	{
+		EBreakBehavior BreakBehavior = EBreakBehavior::Continue;
+		OnlineClientManager.ForEachClient([&Callback, &BreakBehavior](const FOnlineClient& Client)
+		{
+			BreakBehavior = Callback(Client.GetEndpointId());
+			return BreakBehavior;
+		});
+		if (BreakBehavior == EBreakBehavior::Break)
+		{
+			return;
+		}
+		OfflineClientManager.ForEachClient([&Callback](const FOfflineClient& Client)
+		{
+			return Callback(Client.GetLastAssociatedEndpoint());
+		});
+	}
+
+	template <typename TLambda> requires std::is_invocable_r_v<EBreakBehavior, TLambda, const FGuid&>
+	void FUnifiedClientView::ForEachOnlineClient(TLambda&& Callback) const
+	{
+		OnlineClientManager.ForEachClient([&Callback](const FOnlineClient& Client)
+		{
+			return Callback(Client.GetEndpointId());
+		});
+	}
+}
+
diff --git a/Engine/Plugins/Developer/Concert/ConcertApp/MultiUserClient/Source/MultiUserClient/Private/Replication/Client/UnifiedClientViewExtensions.cpp b/Engine/Plugins/Developer/Concert/ConcertApp/MultiUserClient/Source/MultiUserClient/Private/Replication/Client/UnifiedClientViewExtensions.cpp
new file mode 100644
--- /dev/null
+++ b/Engine/Plugins/Developer/Concert/ConcertApp/MultiUserClient/Source/MultiUserClient/Private/Replication/Client/UnifiedClientViewExtensions.cpp
@@ -0,0 +1,82 @@
+﻿// Copyright Epic Games, Inc. All Rights Reserved.
+
+#include "UnifiedClientViewExtensions.h"
+
+#include "ConcertMessageData.h"
+#include "Widgets/Client/SClientName.h"
+
+#include "Containers/UnrealString.h"
+#include "Misc/Optional.h"
+
+namespace UE::MultiUserClient::Replication
+{
+	FString GetClientDisplayString(const FUnifiedClientView& ClientView, const FGuid& ClientEndpointId, EDisplayStringOption Option)
+	{
+		return GetClientDisplayText(ClientView, ClientEndpointId, Option).ToString();
+	}
+
+	FText GetClientDisplayText(const FUnifiedClientView& ClientView, const FGuid& ClientEndpointId, EDisplayStringOption Option)
+	{
+		const TOptional<FConcertClientInfo> ClientInfo = ClientView.GetClientInfoByEndpoint(ClientEndpointId);
+		if (!ClientInfo)
+		{
+			return FText::GetEmpty();
+		}
+
+
+		const FText ParenthesesContent = Option == EDisplayStringOption::NameAndParentheses
+			? GetParenthesesContent(ClientView, ClientEndpointId)
+			: FText::GetEmpty();
+		return ConcertSharedSlate::SClientName::GetDisplayText(*ClientInfo, ParenthesesContent);
+	}
+
+	FText GetParenthesesContent(const FUnifiedClientView& ClientView, const FGuid& ClientEndpointId)
+	{
+		const TOptional<EClientType> ClientType = ClientView.GetClientType(ClientEndpointId);
+		if (!ClientType)
+		{
+			return FText::GetEmpty();
+		}
+
+		if (ClientType == EClientType::Local)
+		{
+			return ConcertSharedSlate::ParenthesesClientNameContent::LocalClient;
+		}
+			
+		return IsOfflineClient(*ClientType)
+			? ConcertSharedSlate::ParenthesesClientNameContent::OfflineClient
+			: FText::GetEmpty();
+	}
+
+	TOptional<FGuid> FindClientIdByStream(const FUnifiedClientView& ClientView, const ConcertSharedSlate::IReplicationStreamModel& SearchedStream)
+	{
+		TOptional<FGuid> Result;
+		ClientView.ForEachClient([&ClientView, &SearchedStream, &Result](const FGuid& EndpointId)
+		{
+			const TSharedPtr<const ConcertSharedSlate::IReplicationStreamModel> ResolvedStream = ClientView.GetClientStreamById(EndpointId);
+			if (ResolvedStream.Get() == &SearchedStream)
+			{
+				Result = EndpointId;
+				return EBreakBehavior::Break;
+			}
+			return EBreakBehavior::Continue;
+		});
+		return Result;
+	}
+
+	TArray<FGuid> GetSortedOnlineClients(const FUnifiedClientView& ClientView, EDisplayStringOption Option)
+	{
+		TMap<FGuid, FString> ClientToDisplayString;
+		TArray<FGuid> SortedClients = ClientView.GetOnlineClients();
+		for (const FGuid& EndpointId : SortedClients)
+		{
+			ClientToDisplayString.Add(EndpointId, GetClientDisplayString(ClientView, EndpointId, Option));
+		}
+
+		SortedClients.Sort([&ClientToDisplayString](const FGuid& Left, const FGuid& Right)
+		{
+			return ClientToDisplayString[Left] < ClientToDisplayString[Right];
+		});
+		return SortedClients;
+	}
+}
diff --git a/Engine/Plugins/Developer/Concert/ConcertApp/MultiUserClient/Source/MultiUserClient/Private/Replication/Client/UnifiedClientViewExtensions.h b/Engine/Plugins/Developer/Concert/ConcertApp/MultiUserClient/Source/MultiUserClient/Private/Replication/Client/UnifiedClientViewExtensions.h
new file mode 100644
--- /dev/null
+++ b/Engine/Plugins/Developer/Concert/ConcertApp/MultiUserClient/Source/MultiUserClient/Private/Replication/Client/UnifiedClientViewExtensions.h
@@ -0,0 +1,65 @@
+﻿// Copyright Epic Games, Inc. All Rights Reserved.
+
+#pragma once
+
+#include "UnifiedClientView.h"
+#include "Widgets/Client/ClientInfoDelegate.h"
+
+#include "HAL/Platform.h"
+#include "Internationalization/Text.h"
+
+class FString;
+struct FGuid;
+
+namespace UE::ConcertSharedSlate { class IReplicationStreamModel; }
+
+namespace UE::MultiUserClient::Replication
+{
+	enum class EDisplayStringOption : uint8
+	{
+		/** E.g. "OfflineUserName", "LocalUserName", and "RemoteUserName". */
+		NameOnly,
+		/** E.g. "OfflineUserName(Offline)", "LocalUserName(Me)", and "RemoteUserName". */
+		NameAndParentheses
+	};
+
+	/** @return Gets the display string to use for the given client; use it e.g. for search. Returns empty FString if the client is not found.*/
+	FString GetClientDisplayString(
+		const FUnifiedClientView& ClientView,
+		const FGuid& ClientEndpointId,
+		EDisplayStringOption Option = EDisplayStringOption::NameAndParentheses
+		);
+	/** @return Gets the display text to use for displaying the given client. Returns empty FText if the client is not found. */
+	FText GetClientDisplayText(
+		const FUnifiedClientView& ClientView,
+		const FGuid& ClientEndpointId,
+		EDisplayStringOption Option = EDisplayStringOption::NameAndParentheses
+		);
+	
+	/** @return "You" or "Offline" depending on ClientEndpointId. */
+	FText GetParenthesesContent(const FUnifiedClientView& ClientView, const FGuid& ClientEndpointId);
+
+	/** @return Endpoint ID of the client that has the given Stream. */
+	TOptional<FGuid> FindClientIdByStream(const FUnifiedClientView& ClientView, const ConcertSharedSlate::IReplicationStreamModel& SearchedStream);
+
+	/** @return Gets all online clients sorted by display name. */
+	TArray<FGuid> GetSortedOnlineClients(const FUnifiedClientView& ClientView, EDisplayStringOption Option = EDisplayStringOption::NameAndParentheses);
+	
+	/** @return Delegate that will get client info from online clients if possible and falls back to offline clients. */
+	inline ConcertSharedSlate::FGetOptionalClientInfo MakeOnlineThenOfflineClientInfoGetter(const FUnifiedClientView& ClientView UE_LIFETIMEBOUND)
+	{
+		return ConcertSharedSlate::FGetOptionalClientInfo::CreateLambda([&ClientView](const FGuid& ClientEndpointId)
+		{
+			return ClientView.GetClientInfoByEndpoint(ClientEndpointId);
+		});
+	}
+
+	/** @return Delegate that will return (You) or (Offline) for a client. */
+	inline ConcertSharedSlate::FGetClientParenthesesContent MakeLocalAndOfflineParenthesesContentGetter(const FUnifiedClientView& ClientView UE_LIFETIMEBOUND)
+	{
+		return ConcertSharedSlate::FGetClientParenthesesContent::CreateLambda([&ClientView](const FGuid& ClientEndpointId)
+		{
+			return GetParenthesesContent(ClientView, ClientEndpointId);
+		});
+	}
+}
diff --git a/Engine/Plugins/Developer/Concert/ConcertApp/MultiUserClient/Source/MultiUserClient/Private/Replication/Client/UnifiedStreamCache.cpp b/Engine/Plugins/Developer/Concert/ConcertApp/MultiUserClient/Source/MultiUserClient/Private/Replication/Client/UnifiedStreamCache.cpp
new file mode 100644
--- /dev/null
+++ b/Engine/Plugins/Developer/Concert/ConcertApp/MultiUserClient/Source/MultiUserClient/Private/Replication/Client/UnifiedStreamCache.cpp
@@ -0,0 +1,155 @@
+﻿// Copyright Epic Games, Inc. All Rights Reserved.
+
+#include "UnifiedStreamCache.h"
+
+#include "Replication/Data/ConcertPropertySelection.h"
+
+#include "Offline/OfflineClientManager.h"
+#include "Online/OnlineClientManager.h"
+#include "UnifiedClientView.h"
+
+namespace UE::MultiUserClient::Replication::UnifiedStreamCache
+{
+	static bool HasNonOverlappingProperties(
+		const FSoftObjectPath& ObjectPath,
+		const FConcertReplicatedObjectInfo& ObjectInfo,
+		const FGlobalAuthorityCache& OnlineCache,
+		FOnlineClientManager& OnlineClientManager
+		)
+	{
+		TSet<FConcertPropertyChain> NonOverlappingProperties = ObjectInfo.PropertySelection.ReplicatedProperties;
+		
+		OnlineCache.ForEachClientWithObjectInStream(ObjectPath,
+			[&ObjectPath, &OnlineClientManager, &NonOverlappingProperties](const FGuid& ClientId)
+			{
+				const FOnlineClient* Client = OnlineClientManager.FindClient(ClientId);
+				const FConcertReplicatedObjectInfo* OnlineObjectInfo = Client
+					? Client->GetStreamSynchronizer().GetServerState().ReplicatedObjects.Find(ObjectPath)
+					: nullptr;
+				if (!ensure(OnlineObjectInfo))
+				{
+					return EBreakBehavior::Continue;
+				}
+
+				NonOverlappingProperties = NonOverlappingProperties.Difference(OnlineObjectInfo->PropertySelection.ReplicatedProperties);
+				return NonOverlappingProperties.IsEmpty() ? EBreakBehavior::Break : EBreakBehavior::Continue;
+			});
+		
+		return !NonOverlappingProperties.IsEmpty();
+	}
+}
+
+namespace UE::MultiUserClient::Replication
+{
+	FUnifiedStreamCache::FUnifiedStreamCache(
+		FUnifiedClientView& InOwner,
+		FOnlineClientManager& InOnlineClientManager,
+		FOfflineClientManager& InOfflineClientManager
+		)
+		: Owner(InOwner)
+		, OnlineClientManager(InOnlineClientManager)
+		, OfflineClientManager(InOfflineClientManager)
+	{
+		Owner.OnClientsChanged().AddRaw(this, &FUnifiedStreamCache::BroadcastOnCacheChanged);
+		OnlineClientManager.GetAuthorityCache().OnCacheChanged().AddRaw(this, &FUnifiedStreamCache::OnOnlineCacheChanged);
+	}
+
+	FUnifiedStreamCache::~FUnifiedStreamCache()
+	{
+		Owner.OnClientsChanged().RemoveAll(this);
+		OnlineClientManager.GetAuthorityCache().OnCacheChanged().RemoveAll(this);
+	}
+
+	const FConcertObjectReplicationMap* FUnifiedStreamCache::GetReplicationMapFor(const FGuid& ClientId) const
+	{
+		if (const FOnlineClient* Client = OnlineClientManager.FindClient(ClientId))
+		{
+			return &Client->GetStreamSynchronizer().GetServerState();
+		}
+		if (const FOfflineClient* Client = OfflineClientManager.FindClient(ClientId))
+		{
+			return &Client->GetPredictedStream().ReplicationMap;
+		}
+		return nullptr;
+	}
+
+	void FUnifiedStreamCache::EnumerateClientsWithObject(
+		const FSoftObjectPath& ObjectPath,
+		const TFunctionRef<EBreakBehavior(const FGuid& ClientId)>& Callback,
+		EClientEnumerationFlags Flags
+		) const
+	{
+		EBreakBehavior BreakBehavior = EBreakBehavior::Continue;
+		const FGlobalAuthorityCache& OnlineCache = OnlineClientManager.GetAuthorityCache();
+		OnlineCache.ForEachClientWithObjectInStream(ObjectPath, [&BreakBehavior, &Callback](const FGuid& ClientId)
+		{
+			BreakBehavior = Callback(ClientId);
+			return BreakBehavior;
+		});
+
+		if (BreakBehavior == EBreakBehavior::Break)
+		{
+			return;
+		}
+
+		const bool bShouldCheckNonOverlapping = EnumHasAnyFlags(Flags, EClientEnumerationFlags::SkipOfflineClientsThatFullyOverlapWithOnline);
+		OfflineClientManager.ForEachClient(
+			[this, &ObjectPath, &Callback, bShouldCheckNonOverlapping, &OnlineCache](const FOfflineClient& Client)
+			{
+				const FConcertReplicatedObjectInfo* ObjectInfo = Client.GetPredictedStream().ReplicationMap.ReplicatedObjects.Find(ObjectPath);
+				const bool bContainsObject = ObjectInfo != nullptr;
+				const bool bShouldListObject = bContainsObject
+					&& (!bShouldCheckNonOverlapping
+						|| UnifiedStreamCache::HasNonOverlappingProperties(ObjectPath, *ObjectInfo, OnlineCache, OnlineClientManager));
+				return bShouldListObject ? Callback(Client.GetLastAssociatedEndpoint()) : EBreakBehavior::Continue;
+			});
+	}
+
+	void FUnifiedStreamCache::EnumerateClientsWithObjectAndProperty(
+		const FSoftObjectPath& ObjectPath,
+		const FConcertPropertyChain& Property,
+		const TFunctionRef<EBreakBehavior(const FGuid& ClientId)>& Callback,
+		EClientEnumerationFlags Option
+		) const
+	{
+		EBreakBehavior BreakBehavior = EBreakBehavior::Continue;
+		bool bEncounteredOnlineClient = false;
+		const FGlobalAuthorityCache& OnlineCache = OnlineClientManager.GetAuthorityCache();
+		OnlineCache.ForEachClientWithObjectInStream(ObjectPath,
+			[this, &ObjectPath, &Property, &Callback, &BreakBehavior, &bEncounteredOnlineClient](const FGuid& ClientId)
+			{
+				const FOnlineClient* Client = OnlineClientManager.FindClient(ClientId);
+				if (!Client)
+				{
+					return EBreakBehavior::Continue;
+				}
+
+				const FConcertReplicatedObjectInfo* ObjectInfo = Client->GetStreamSynchronizer().GetServerState().ReplicatedObjects.Find(ObjectPath);
+				const bool bShouldInvoke = ObjectInfo && ObjectInfo->PropertySelection.ReplicatedProperties.Contains(Property);
+				if (!bShouldInvoke)
+				{
+					return EBreakBehavior::Continue;
+				}
+
+				bEncounteredOnlineClient = true;
+				BreakBehavior = Callback(ClientId);
+				return BreakBehavior;
+			}
+		);
+
+		if (BreakBehavior == EBreakBehavior::Break
+			// If there is an online client, the offline clients are not supposed to be displayed because the caller requested no property overlaps.
+			|| (Option == EClientEnumerationFlags::SkipOfflineClientsThatFullyOverlapWithOnline && bEncounteredOnlineClient))
+		{
+			return;
+		}
+		
+		OfflineClientManager.ForEachClient(
+			[this, &ObjectPath, &Callback](const FOfflineClient& Client)
+			{
+				const FConcertReplicatedObjectInfo* ObjectInfo = Client.GetPredictedStream().ReplicationMap.ReplicatedObjects.Find(ObjectPath);
+				const bool bContainsObject = ObjectInfo != nullptr;
+				return bContainsObject ? Callback(Client.GetLastAssociatedEndpoint()) : EBreakBehavior::Continue;
+			});
+	}
+}
diff --git a/Engine/Plugins/Developer/Concert/ConcertApp/MultiUserClient/Source/MultiUserClient/Private/Replication/Client/UnifiedStreamCache.h b/Engine/Plugins/Developer/Concert/ConcertApp/MultiUserClient/Source/MultiUserClient/Private/Replication/Client/UnifiedStreamCache.h
new file mode 100644
--- /dev/null
+++ b/Engine/Plugins/Developer/Concert/ConcertApp/MultiUserClient/Source/MultiUserClient/Private/Replication/Client/UnifiedStreamCache.h
@@ -0,0 +1,86 @@
+﻿// Copyright Epic Games, Inc. All Rights Reserved.
+
+#pragma once
+
+#include "Delegates/Delegate.h"
+#include "HAL/Platform.h"
+#include "Misc/EnumClassFlags.h"
+#include "Templates/FunctionFwd.h"
+
+struct FConcertObjectReplicationMap;
+struct FConcertPropertyChain;
+
+namespace UE::MultiUserClient::Replication
+{
+	class FUnifiedClientView;
+}
+
+enum class EBreakBehavior : uint8;
+struct FGuid;
+struct FSoftObjectPath;
+
+namespace UE::MultiUserClient::Replication
+{
+	class FOfflineClientManager;
+	class FOnlineClientManager;
+
+	enum class EClientEnumerationFlags : uint8
+	{
+		None = 0,
+		SkipOfflineClientsThatFullyOverlapWithOnline = 1 << 0
+	};
+	ENUM_CLASS_FLAGS(EClientEnumerationFlags);
+	
+	/** Access point for querying stream content of both online and offline clients through a single interface. */
+	class FUnifiedStreamCache : public FNoncopyable
+	{
+	public:
+		
+		FUnifiedStreamCache(
+			FUnifiedClientView& InOwner UE_LIFETIMEBOUND,
+			FOnlineClientManager& InOnlineClientManager UE_LIFETIMEBOUND,
+			FOfflineClientManager& InOfflineClientManager UE_LIFETIMEBOUND
+			);
+		~FUnifiedStreamCache();
+
+		/** @return Gets the replication map for the specified client. */
+		const FConcertObjectReplicationMap* GetReplicationMapFor(const FGuid& ClientId) const;
+		
+		/**
+		 * Lists out all clients that have ObjectPath registered in their stream.
+		 * Does not mean that they have authority over it.
+		 */
+		void EnumerateClientsWithObject(
+			const FSoftObjectPath& ObjectPath,
+			const TFunctionRef<EBreakBehavior(const FGuid& ClientId)>& Callback,
+			EClientEnumerationFlags Flags = EClientEnumerationFlags::SkipOfflineClientsThatFullyOverlapWithOnline
+			) const;
+
+		/**
+		 * Lists out all clients that have ObjectPath with Property assigned to it registered in their stream.
+		 * Does not mean that they have authority over it.
+		 */
+		void EnumerateClientsWithObjectAndProperty(
+			const FSoftObjectPath& ObjectPath,
+			const FConcertPropertyChain& Property,
+			const TFunctionRef<EBreakBehavior(const FGuid& ClientId)>& Callback,
+			EClientEnumerationFlags Option = EClientEnumerationFlags::SkipOfflineClientsThatFullyOverlapWithOnline
+			) const;
+
+		/** Broadcasts when the content of a client has changed. */
+		FSimpleMulticastDelegate& OnCacheChanged() { return OnCacheChangedDelegate; }
+
+	private:
+		
+		FUnifiedClientView& Owner;
+		FOnlineClientManager& OnlineClientManager;
+		FOfflineClientManager& OfflineClientManager;
+
+		/** Broadcasts when the content of a client has changed. */
+		FSimpleMulticastDelegate OnCacheChangedDelegate;
+		
+		void OnOnlineCacheChanged(const FGuid&) const { BroadcastOnCacheChanged(); }
+		void BroadcastOnCacheChanged() const { OnCacheChangedDelegate.Broadcast(); }
+	};
+}
+
diff --git a/Engine/Plugins/Developer/Concert/ConcertApp/MultiUserClient/Source/MultiUserClient/Private/Replication/Misc/PropertySelection/UserPropertySelectionSource.cpp b/Engine/Plugins/Developer/Concert/ConcertApp/MultiUserClient/Source/MultiUserClient/Private/Replication/Misc/PropertySelection/UserPropertySelectionSource.cpp
--- a/Engine/Plugins/Developer/Concert/ConcertApp/MultiUserClient/Source/MultiUserClient/Private/Replication/Misc/PropertySelection/UserPropertySelectionSource.cpp
+++ b/Engine/Plugins/Developer/Concert/ConcertApp/MultiUserClient/Source/MultiUserClient/Private/Replication/Misc/PropertySelection/UserPropertySelectionSource.cpp
@@ -18,7 +18,7 @@ namespace UE::MultiUserClient::Replication
 			
 			FUserPropertySource(
 				FSoftObjectPath Object,
-				const ConcertSharedSlate::IEditableReplicationStreamModel& InUserSelection UE_LIFETIMEBOUND,
+				const ConcertSharedSlate::IReplicationStreamModel& InUserSelection UE_LIFETIMEBOUND,
 				const FOnlineClientManager& InClientManager UE_LIFETIMEBOUND
 				)
 				: Object(MoveTemp(Object))
@@ -89,14 +89,14 @@ namespace UE::MultiUserClient::Replication
 			const FSoftObjectPath Object;
 
 			/** Used to get the properties the user has selected. */
-			const ConcertSharedSlate::IEditableReplicationStreamModel& UserSelection;
+			const ConcertSharedSlate::IReplicationStreamModel& UserSelection;
 			/** Used to get client stream content. */
 			const FOnlineClientManager& ClientManager;
 		};
 	}
 	
 	FUserPropertySelectionSource::FUserPropertySelectionSource(
-		const ConcertSharedSlate::IEditableReplicationStreamModel& InUserSelection,
+		const ConcertSharedSlate::IReplicationStreamModel& InUserSelection,
 		const FOnlineClientManager& InClientManager
 		)
 		: UserSelection(InUserSelection)
diff --git a/Engine/Plugins/Developer/Concert/ConcertApp/MultiUserClient/Source/MultiUserClient/Private/Replication/Misc/PropertySelection/UserPropertySelectionSource.h b/Engine/Plugins/Developer/Concert/ConcertApp/MultiUserClient/Source/MultiUserClient/Private/Replication/Misc/PropertySelection/UserPropertySelectionSource.h
--- a/Engine/Plugins/Developer/Concert/ConcertApp/MultiUserClient/Source/MultiUserClient/Private/Replication/Misc/PropertySelection/UserPropertySelectionSource.h
+++ b/Engine/Plugins/Developer/Concert/ConcertApp/MultiUserClient/Source/MultiUserClient/Private/Replication/Misc/PropertySelection/UserPropertySelectionSource.h
@@ -7,7 +7,7 @@
 #include "HAL//Platform.h"
 #include "Templates/UnrealTemplate.h"
 
-namespace UE::ConcertSharedSlate { class IEditableReplicationStreamModel; }
+namespace UE::ConcertSharedSlate { class IReplicationStreamModel; }
 
 namespace UE::MultiUserClient::Replication
 {
@@ -25,7 +25,7 @@ namespace UE::MultiUserClient::Replication
 	public:
 
 		FUserPropertySelectionSource(
-			const ConcertSharedSlate::IEditableReplicationStreamModel& InUserSelection UE_LIFETIMEBOUND,
+			const ConcertSharedSlate::IReplicationStreamModel& InUserSelection UE_LIFETIMEBOUND,
 			const FOnlineClientManager& InClientManager UE_LIFETIMEBOUND
 			);
 
@@ -38,8 +38,8 @@ namespace UE::MultiUserClient::Replication
 
 	private:
 
-		/** Used to get the properties the user has selected. */
-		const ConcertSharedSlate::IEditableReplicationStreamModel& UserSelection;
+		/** This is used to read the properties the user has selected, which is represented by a stream. */
+		const ConcertSharedSlate::IReplicationStreamModel& UserSelection;
 		/** Used to get client stream content and subscribe to changes. */
 		const FOnlineClientManager& ClientManager;
 	};
diff --git a/Engine/Plugins/Developer/Concert/ConcertApp/MultiUserClient/Source/MultiUserClient/Private/Replication/Misc/PropertySelection/UserPropertySelector.cpp b/Engine/Plugins/Developer/Concert/ConcertApp/MultiUserClient/Source/MultiUserClient/Private/Replication/Misc/PropertySelection/UserPropertySelector.cpp
--- a/Engine/Plugins/Developer/Concert/ConcertApp/MultiUserClient/Source/MultiUserClient/Private/Replication/Misc/PropertySelection/UserPropertySelector.cpp
+++ b/Engine/Plugins/Developer/Concert/ConcertApp/MultiUserClient/Source/MultiUserClient/Private/Replication/Misc/PropertySelection/UserPropertySelector.cpp
@@ -8,6 +8,7 @@
 #include "Replication/Editor/Model/IEditableReplicationStreamModel.h"
 
 #include "ScopedTransaction.h"
+#include "Replication/Client/Offline/OfflineClientManager.h"
 #include "UObject/Package.h"
 #include "UObject/UObjectGlobals.h"
 
@@ -15,25 +16,49 @@
 
 namespace UE::MultiUserClient::Replication
 {
-	FUserPropertySelector::FUserPropertySelector(FOnlineClientManager& InClientManager)
-		: ClientManager(InClientManager)
+	FUserPropertySelector::FUserPropertySelector(
+		FOnlineClientManager& InOnlineClientManager,
+		FOfflineClientManager& InOfflineClientManger
+		)
+		: OnlineClientManager(InOnlineClientManager)
+		, OfflineClientManager(InOfflineClientManger)
 		, PropertySelection(NewObject<UMultiUserReplicationStream>(GetTransientPackage(), NAME_None, RF_Transient | RF_Transactional))
 		, SelectionEditModel(ConcertSharedSlate::CreateBaseStreamModel(PropertySelection->MakeReplicationMapGetterAttribute()))
-		, PropertyProcessor(MakeShared<FUserPropertySelectionSource>(*SelectionEditModel, InClientManager))
+		, PropertyProcessor(MakeShared<FUserPropertySelectionSource>(*SelectionEditModel, InOnlineClientManager))
 	{
-		RegisterClient(ClientManager.GetLocalClient());
-		ClientManager.OnPostRemoteClientAdded().AddRaw(this, &FUserPropertySelector::OnClientAdded);
+		OnlineClientManager.OnPostRemoteClientAdded().AddRaw(this, &FUserPropertySelector::OnClientAdded);
+		OnlineClientManager.ForEachClient([this](FOnlineClient& Client)
+		{
+			RegisterOnlineClient(Client);
+			return EBreakBehavior::Continue;
+		});
+
+		OfflineClientManager.OnPostClientAdded().AddRaw(this, &FUserPropertySelector::RegisterOfflineClient);
+		OfflineClientManager.ForEachClient([this](FOfflineClient& Client)
+		{
+			RegisterOfflineClient(Client);
+			return EBreakBehavior::Continue;
+		});
+		
 		FCoreUObjectDelegates::OnObjectTransacted.AddRaw(this, &FUserPropertySelector::OnObjectTransacted);
 	}
 
 	FUserPropertySelector::~FUserPropertySelector()
 	{
-		ClientManager.OnPostRemoteClientAdded().RemoveAll(this);
-		ClientManager.ForEachClient([this](FOnlineClient& Client)
+		OnlineClientManager.OnPostRemoteClientAdded().RemoveAll(this);
+		OnlineClientManager.ForEachClient([this](FOnlineClient& Client)
 		{
 			Client.GetStreamSynchronizer().OnServerStreamChanged().RemoveAll(this);
 			return EBreakBehavior::Continue;
 		});
+
+		OfflineClientManager.OnPostClientAdded().RemoveAll(this);
+		OfflineClientManager.ForEachClient([this](FOfflineClient& Client)
+		{
+			Client.OnStreamPredictionChanged().RemoveAll(this);
+			return EBreakBehavior::Continue;
+		});
+		
 		FCoreUObjectDelegates::OnObjectTransacted.RemoveAll(this);
 	}
 
@@ -55,7 +80,7 @@ namespace UE::MultiUserClient::Replication
 
 	bool FUserPropertySelector::IsPropertySelected(const FSoftObjectPath& Object, const FConcertPropertyChain& Property) const
 	{
-		return ClientManager.GetAuthorityCache().IsPropertyReferencedByAnyClientStream(Object, Property)
+		return OnlineClientManager.GetAuthorityCache().IsPropertyReferencedByAnyClientStream(Object, Property)
 			|| SelectionEditModel->HasProperty(Object, Property);
 	}
 
@@ -69,22 +94,41 @@ namespace UE::MultiUserClient::Replication
 		Collector.AddReferencedObject(PropertySelection);
 	}
 
-	void FUserPropertySelector::RegisterClient(FOnlineClient& Client)
+	void FUserPropertySelector::OnClientAdded(FRemoteClient& Client)
+	{
+		RegisterOnlineClient(Client);
+	}
+
+	void FUserPropertySelector::RegisterOnlineClient(FOnlineClient& Client)
 	{
 		IClientStreamSynchronizer& StreamSynchronizer = Client.GetStreamSynchronizer();
 		TrackProperties(StreamSynchronizer.GetServerState());
-		StreamSynchronizer.OnServerStreamChanged().AddRaw(this, &FUserPropertySelector::OnServerStateChanged, Client.GetEndpointId());
+		
+		StreamSynchronizer.OnServerStreamChanged().AddRaw(this, &FUserPropertySelector::OnOnlineClientContentChanged, Client.GetEndpointId());
+	}
+
+	void FUserPropertySelector::RegisterOfflineClient(FOfflineClient& Client)
+	{
+		TrackProperties(Client.GetPredictedStream().ReplicationMap);
+		
+		const TNonNullPtr<const FOfflineClient>& NonNullClient = &Client;
+		Client.OnStreamPredictionChanged().AddRaw(this, &FUserPropertySelector::OnOfflineClientContentChanged, NonNullClient);
 	}
 
-	void FUserPropertySelector::OnServerStateChanged(const FGuid ClientId)
+	void FUserPropertySelector::OnOnlineClientContentChanged(const FGuid ClientId)
 	{
-		const FOnlineClient* Client = ClientManager.FindClient(ClientId);
+		const FOnlineClient* Client = OnlineClientManager.FindClient(ClientId);
 		if (ensure(Client))
 		{
 			TrackProperties(Client->GetStreamSynchronizer().GetServerState());
 		}
 	}
 
+	void FUserPropertySelector::OnOfflineClientContentChanged(const TNonNullPtr<const FOfflineClient> Client)
+	{
+		TrackProperties(Client->GetPredictedStream().ReplicationMap);
+	}
+
 	void FUserPropertySelector::TrackProperties(const FConcertObjectReplicationMap& ReplicationMap)
 	{
 		for (const TPair<FSoftObjectPath, FConcertReplicatedObjectInfo>& Pair : ReplicationMap.ReplicatedObjects)
@@ -129,7 +173,7 @@ namespace UE::MultiUserClient::Replication
 			SelectionEditModel->RemoveObjects({ Object });
 		}
 
-		const TSharedRef<ConcertSharedSlate::IEditableReplicationStreamModel> ClientEditModel = ClientManager.GetLocalClient().GetClientEditModel();
+		const TSharedRef<ConcertSharedSlate::IEditableReplicationStreamModel> ClientEditModel = OnlineClientManager.GetLocalClient().GetClientEditModel();
 		ClientEditModel->RemoveProperties({ Object }, Properties);
 		
 		OnPropertySelectionChangedDelegate.Broadcast();
diff --git a/Engine/Plugins/Developer/Concert/ConcertApp/MultiUserClient/Source/MultiUserClient/Private/Replication/Misc/PropertySelection/UserPropertySelector.h b/Engine/Plugins/Developer/Concert/ConcertApp/MultiUserClient/Source/MultiUserClient/Private/Replication/Misc/PropertySelection/UserPropertySelector.h
--- a/Engine/Plugins/Developer/Concert/ConcertApp/MultiUserClient/Source/MultiUserClient/Private/Replication/Misc/PropertySelection/UserPropertySelector.h
+++ b/Engine/Plugins/Developer/Concert/ConcertApp/MultiUserClient/Source/MultiUserClient/Private/Replication/Misc/PropertySelection/UserPropertySelector.h
@@ -4,13 +4,13 @@
 
 #include "Assets/MultiUserReplicationStream.h"
 
-#include "Replication/Client/Online/RemoteClient.h"
-
 #include "HAL/Platform.h"
 #include "Templates/UnrealTemplate.h"
 #include "UObject/GCObject.h"
 #include "UObject/ObjectPtr.h"
 
+
+
 class FReferenceCollector;
 class FTransactionObjectEvent;
 class UObject;
@@ -23,8 +23,12 @@ namespace UE::ConcertSharedSlate
 
 namespace UE::MultiUserClient::Replication
 {
-	class FUserPropertySelectionSource;
+	class FOfflineClient;
+	class FOfflineClientManager;
 	class FOnlineClientManager;
+	class FOnlineClient;
+	class FRemoteClient;
+	class FUserPropertySelectionSource;
 	
 	/**
 	 * Manages the properties the user is iterating on in the replication session.
@@ -39,8 +43,11 @@ namespace UE::MultiUserClient::Replication
 	{
 	public:
 		
-		FUserPropertySelector(FOnlineClientManager& InClientManager UE_LIFETIMEBOUND);
-		~FUserPropertySelector();
+		FUserPropertySelector(
+			FOnlineClientManager& InOnlineClientManager UE_LIFETIMEBOUND,
+			FOfflineClientManager& InOfflineClientManger UE_LIFETIMEBOUND
+			);
+		virtual ~FUserPropertySelector() override;
 
 		/** Add Properties from the user's selection for Object. */
 		void AddSelectedProperties(UObject* Object, TConstArrayView<FConcertPropertyChain> Properties);
@@ -62,8 +69,13 @@ namespace UE::MultiUserClient::Replication
 
 	private:
 
-		/** Used to remove deselected properties from local client's stream. */
-		FOnlineClientManager& ClientManager;
+		/** Used to remove deselected properties from local client's stream and auto-add properties from remote clients to the user selection. */
+		FOnlineClientManager& OnlineClientManager;
+		/**
+		 * When an online client turns into an offline client, we need to subscribe to changes made to that client,
+		 * e.g. preset could change what client gets when it rejoins.
+		 */
+		FOfflineClientManager& OfflineClientManager;
 
 		/** This underlying object saves the properties that user has selected. It allows for transactions. */
 		TObjectPtr<UMultiUserReplicationStream> PropertySelection;
@@ -77,12 +89,15 @@ namespace UE::MultiUserClient::Replication
 		FOnPropertySelectionChanged OnPropertySelectionChangedDelegate;
 
 		/** Called when a remote client joins. */
-		void OnClientAdded(FRemoteClient& Client) { RegisterClient(Client); }
+		void OnClientAdded(FRemoteClient& Client);
 		/** Ensures that whenever the client's server state changes, its properties are tracked as user selected. */
-		void RegisterClient(FOnlineClient& Client);
+		void RegisterOnlineClient(FOnlineClient& Client);
+		void RegisterOfflineClient(FOfflineClient& Client);
 
 		/** Tracks all properties of the client as user selected. */
-		void OnServerStateChanged(const FGuid ClientId);
+		void OnOnlineClientContentChanged(const FGuid ClientId);
+		/** Tracks all properties the offline client will get upon rejoining as user selected. */
+		void OnOfflineClientContentChanged(const TNonNullPtr<const FOfflineClient> Client);
 		/** Adds all properties in the replication map as user selected. */
 		void TrackProperties(const FConcertObjectReplicationMap& ReplicationMap);
 
diff --git a/Engine/Plugins/Developer/Concert/ConcertApp/MultiUserClient/Source/MultiUserClient/Private/Replication/MultiUserReplicationManager.cpp b/Engine/Plugins/Developer/Concert/ConcertApp/MultiUserClient/Source/MultiUserClient/Private/Replication/MultiUserReplicationManager.cpp
--- a/Engine/Plugins/Developer/Concert/ConcertApp/MultiUserClient/Source/MultiUserClient/Private/Replication/MultiUserReplicationManager.cpp
+++ b/Engine/Plugins/Developer/Concert/ConcertApp/MultiUserClient/Source/MultiUserClient/Private/Replication/MultiUserReplicationManager.cpp
@@ -142,7 +142,7 @@ namespace UE::MultiUserClient::Replication
 
 	void FMultiUserReplicationManager::SetupClientConnectionEvents()
 	{
-		FOnlineClientManager& ClientManager = ConnectedState->ClientManager;
+		FOnlineClientManager& ClientManager = ConnectedState->OnlineClientManager;
 		ClientManager.ForEachClient([this](FOnlineClient& InClient){ SetupClientDelegates(InClient); return EBreakBehavior::Continue; });
 		ClientManager.OnPostRemoteClientAdded().AddRaw(this, &FMultiUserReplicationManager::OnReplicationClientConnected);
 	}
@@ -169,7 +169,7 @@ namespace UE::MultiUserClient::Replication
 	{
 		if (ConnectedState && ensureMsgf(IsInGameThread(), TEXT("To simplify implementation, only calls from game thread are allowed.")))
 		{
-			const FOnlineClient* ReplicationClient = ConnectedState->ClientManager.FindClient(ClientId);
+			const FOnlineClient* ReplicationClient = ConnectedState->OnlineClientManager.FindClient(ClientId);
 			return ReplicationClient
 				? &ReplicationClient->GetStreamSynchronizer().GetServerState()
 				: nullptr;
@@ -181,7 +181,7 @@ namespace UE::MultiUserClient::Replication
 	{
 		if (ConnectedState && ensureMsgf(IsInGameThread(), TEXT("To simplify implementation, only calls from game thread are allowed.")))
 		{
-			const FOnlineClient* ReplicationClient = ConnectedState->ClientManager.FindClient(ClientId);
+			const FOnlineClient* ReplicationClient = ConnectedState->OnlineClientManager.FindClient(ClientId);
 			return ReplicationClient
 				? &ReplicationClient->GetStreamSynchronizer().GetFrequencySettings()
 				: nullptr;
@@ -193,7 +193,7 @@ namespace UE::MultiUserClient::Replication
 	{
 		if (ConnectedState && ensureMsgf(IsInGameThread(), TEXT("To simplify implementation, only calls from game thread are allowed.")))
 		{
-			const FOnlineClient* ReplicationClient = ConnectedState->ClientManager.FindClient(ClientId);
+			const FOnlineClient* ReplicationClient = ConnectedState->OnlineClientManager.FindClient(ClientId);
 			return ReplicationClient && ReplicationClient->GetAuthoritySynchronizer().HasAuthorityOver(ObjectPath);
 		}
 		return false;
@@ -224,7 +224,7 @@ namespace UE::MultiUserClient::Replication
 		
 		if (ConnectedState)
 		{
-			FOnlineClient* ReplicationClient = ConnectedState->ClientManager.FindClient(ClientId);
+			FOnlineClient* ReplicationClient = ConnectedState->OnlineClientManager.FindClient(ClientId);
 			return ReplicationClient
 				? ReplicationClient->GetExternalRequestHandler().HandleRequest(MoveTemp(SubmissionParams))
 				: FExternalClientChangeRequestHandler::MakeFailedOperation(EChangeStreamOperationResult::UnknownClient, EChangeAuthorityOperationResult::UnknownClient);
@@ -235,13 +235,15 @@ namespace UE::MultiUserClient::Replication
 	FMultiUserReplicationManager::FConnectedState::FConnectedState(TSharedRef<IConcertSyncClient> InClient, FReplicationDiscoveryContainer& InDiscoveryContainer)
 		: Client(InClient)
 		, QueryService(*InClient)
-		, ClientManager(InClient, InClient->GetConcertClient()->GetCurrentSession().ToSharedRef(), InDiscoveryContainer, QueryService.GetStreamAndAuthorityQueryService())
-		, MuteManager(*InClient, QueryService.GetMuteStateQueryService(), ClientManager.GetAuthorityCache())
-		, PresetManager(*InClient, ClientManager, MuteManager.GetSynchronizer())
-		, PropertySelector(ClientManager)
-		, ChangeLevelHandler(ClientManager.GetLocalClient().GetClientEditModel().Get())
-		, PreventReplicatedPropertyTransaction(*InClient, ClientManager, MuteManager)
-		, UserNotifier(ClientManager, MuteManager)
+		, OnlineClientManager(InClient, InClient->GetConcertClient()->GetCurrentSession().ToSharedRef(), InDiscoveryContainer, QueryService.GetStreamAndAuthorityQueryService())
+		, OfflineClientManager(*InClient, OnlineClientManager)
+		, UnifiedClientView(*InClient, OnlineClientManager, OfflineClientManager)
+		, MuteManager(*InClient, QueryService.GetMuteStateQueryService(), OnlineClientManager.GetAuthorityCache())
+		, PresetManager(*InClient, OnlineClientManager, MuteManager.GetSynchronizer())
+		, PropertySelector(OnlineClientManager, OfflineClientManager)
+		, ChangeLevelHandler(OnlineClientManager.GetLocalClient().GetClientEditModel().Get())
+		, PreventReplicatedPropertyTransaction(*InClient, OnlineClientManager, MuteManager)
+		, UserNotifier(OnlineClientManager, MuteManager)
 	{}
 }
 
diff --git a/Engine/Plugins/Developer/Concert/ConcertApp/MultiUserClient/Source/MultiUserClient/Private/Replication/MultiUserReplicationManager.h b/Engine/Plugins/Developer/Concert/ConcertApp/MultiUserClient/Source/MultiUserClient/Private/Replication/MultiUserReplicationManager.h
--- a/Engine/Plugins/Developer/Concert/ConcertApp/MultiUserClient/Source/MultiUserClient/Private/Replication/MultiUserReplicationManager.h
+++ b/Engine/Plugins/Developer/Concert/ConcertApp/MultiUserClient/Source/MultiUserClient/Private/Replication/MultiUserReplicationManager.h
@@ -4,6 +4,8 @@
 
 #include "Client/Online/OnlineClientManager.h"
 #include "IConcertSession.h"
+#include "Client/UnifiedClientView.h"
+#include "Client/Offline/OfflineClientManager.h"
 #include "Misc/ChangeLevelHandler.h"
 #include "Misc/Notification/ReplicationUserNotifier.h"
 #include "Misc/PreventReplicatedPropertyTransaction.h"
@@ -61,19 +63,27 @@ namespace UE::MultiUserClient::Replication
 		 */
 		void JoinReplicationSession();
 
-		/** @note You're not supposed to keep any reference to the ClientManager since it can become invalid depending on connection state. */
-		FOnlineClientManager* GetOnlineClientManager() { return ConnectedState ? &ConnectedState->ClientManager : nullptr; }
-		const FOnlineClientManager* GetOnlineClientManager() const { return ConnectedState ? &ConnectedState->ClientManager : nullptr; }
+		/** @note This pointer is only valid while connected and can become stale. Listen for OnReplicationConnectionStateChanged. */
+		FOnlineClientManager* GetOnlineClientManager() { return ConnectedState ? &ConnectedState->OnlineClientManager : nullptr; }
+		const FOnlineClientManager* GetOnlineClientManager() const { return ConnectedState ? &ConnectedState->OnlineClientManager : nullptr; }
 		
-		/** @note You're not supposed to keep any reference to the MuteManager since it can become invalid depending on connection state. */
+		/** @note This pointer is only valid while connected and can become stale. Listen for OnReplicationConnectionStateChanged. */
+		FOfflineClientManager* GetOfflineClientManager() { return ConnectedState ? &ConnectedState->OfflineClientManager : nullptr; }
+		const FOfflineClientManager* GetOfflineClientManager() const { return ConnectedState ? &ConnectedState->OfflineClientManager : nullptr; }
+		
+		/** @note This pointer is only valid while connected and can become stale. Listen for OnReplicationConnectionStateChanged. */
+		FUnifiedClientView* GetUnifiedClientView() { return ConnectedState ? &ConnectedState->UnifiedClientView : nullptr; }
+		const FUnifiedClientView* GetUnifiedClientView() const { return ConnectedState ? &ConnectedState->UnifiedClientView : nullptr; }
+		
+		/** @note This pointer is only valid while connected and can become stale. Listen for OnReplicationConnectionStateChanged. */
 		FMuteStateManager* GetMuteManager() { return ConnectedState ? &ConnectedState->MuteManager : nullptr; }
 		const FMuteStateManager* GetMuteManager() const { return ConnectedState ? &ConnectedState->MuteManager : nullptr; }
 
-		/** @note You're not supposed to keep any reference to the PresetManager since it can become invalid depending on connection state. */
+		/** @note This pointer is only valid while connected and can become stale. Listen for OnReplicationConnectionStateChanged. */
 		FPresetManager* GetPresetManager() { return ConnectedState ? &ConnectedState->PresetManager : nullptr; }
 		const FPresetManager* GetPresetManager() const { return ConnectedState ? &ConnectedState->PresetManager : nullptr; }
 
-		/** @note You're not supposed to keep any reference to the PropertySelector since it can become invalid depending on connection state. */
+		/** @note This pointer is only valid while connected and can become stale. Listen for OnReplicationConnectionStateChanged. */
 		FUserPropertySelector* GetUserPropertySelector() { return ConnectedState ? &ConnectedState->PropertySelector : nullptr; }
 		const FUserPropertySelector* GetUserPropertySelector() const { return ConnectedState ? &ConnectedState->PropertySelector : nullptr; }
 
@@ -116,12 +126,18 @@ namespace UE::MultiUserClient::Replication
 			FRegularQueryService QueryService;
 			
 			/**
-			 * Creates UMultiUserReplicationSessionPreset which is displayed by UI.
+			 * Synchronizes replication settings of clients connected to the session.
+			 * Creates UMultiUserReplicationSessionPreset, which is displayed by UI.
 			 * Keeps the preset in sync with the state on the server.
 			 *
 			 * Only valid when ConnectionState == EMultiUserReplicationConnectionState::Connected.
 			 */
-			FOnlineClientManager ClientManager;
+			FOnlineClientManager OnlineClientManager;
+			/** Keeps track of clients that were once in the session but are no longer to be able to display their old settings in the UI. */
+			FOfflineClientManager OfflineClientManager;
+			/** Adapter abstraction that allows systems to query online and offline clients with an unified interface. */
+			FUnifiedClientView UnifiedClientView;
+			
 			/** Interacts with the mute global server mute system. */
 			FMuteStateManager MuteManager;
 			/** Saves and loads presets for the session. Accessed by UI. */
diff --git a/Engine/Plugins/Developer/Concert/ConcertApp/MultiUserClient/Source/MultiUserClient/Private/Replication/Stream/IClientStreamSynchronizer.h b/Engine/Plugins/Developer/Concert/ConcertApp/MultiUserClient/Source/MultiUserClient/Private/Replication/Stream/IClientStreamSynchronizer.h
--- a/Engine/Plugins/Developer/Concert/ConcertApp/MultiUserClient/Source/MultiUserClient/Private/Replication/Stream/IClientStreamSynchronizer.h
+++ b/Engine/Plugins/Developer/Concert/ConcertApp/MultiUserClient/Source/MultiUserClient/Private/Replication/Stream/IClientStreamSynchronizer.h
@@ -39,7 +39,6 @@ namespace UE::MultiUserClient::Replication
 	/** Util base class for implementing the events */
 	class FStreamSynchronizer_Base
 		: public IClientStreamSynchronizer
-		, public FNoncopyable
 	{
 	public:
 
diff --git a/Engine/Plugins/Developer/Concert/ConcertApp/MultiUserClient/Source/MultiUserClient/Private/Widgets/ActiveSession/Replication/Client/Multi/Columns/AssignProperty/AssignPropertyModel.cpp b/Engine/Plugins/Developer/Concert/ConcertApp/MultiUserClient/Source/MultiUserClient/Private/Widgets/ActiveSession/Replication/Client/Multi/Columns/AssignProperty/AssignPropertyModel.cpp
--- a/Engine/Plugins/Developer/Concert/ConcertApp/MultiUserClient/Source/MultiUserClient/Private/Widgets/ActiveSession/Replication/Client/Multi/Columns/AssignProperty/AssignPropertyModel.cpp
+++ b/Engine/Plugins/Developer/Concert/ConcertApp/MultiUserClient/Source/MultiUserClient/Private/Widgets/ActiveSession/Replication/Client/Multi/Columns/AssignProperty/AssignPropertyModel.cpp
@@ -3,8 +3,6 @@
 #include "AssignPropertyModel.h"
 
 #include "ConcertLogGlobal.h"
-#include "Replication/Client/Online/OnlineClient.h"
-#include "Replication/Client/Online/OnlineClientManager.h"
 #include "Replication/Editor/Model/PropertyUtils.h"
 
 #include "GameFramework/Actor.h"
@@ -12,24 +10,46 @@
 
 #include <type_traits>
 
+#include "Replication/Client/UnifiedClientView.h"
+
 #define LOCTEXT_NAMESPACE "FAssignPropertyModel"
 
 namespace UE::MultiUserClient::Replication::MultiStreamColumns::AssignPropertyModel
 {
+	static void AssignPropertyTo(
+		ConcertSharedSlate::IEditableReplicationStreamModel& ClientEditModel,
+		TConstArrayView<TSoftObjectPtr<>> Objects,
+		const FConcertPropertyChain& Property
+		)
+	{
+		for (const TSoftObjectPtr<>& Object : Objects)
+		{
+			const FSoftObjectPath& ObjectPath = Object.GetUniqueID();
+			if (!ClientEditModel.ContainsObjects({ ObjectPath }))
+			{
+				ClientEditModel.AddObjects({ Object.Get() });
+			}
+
+			const FSoftClassPath ClassPath = ClientEditModel.GetObjectClass(ObjectPath);
+			TArray AddedProperties { Property };
+			ConcertClientSharedSlate::PropertyUtils::AppendAdditionalPropertiesToAdd(ClassPath, AddedProperties);
+			ClientEditModel.AddProperties(ObjectPath, AddedProperties);
+		}
+	}
+	
 	static void RemovePropertiesFromClient(
-		const FOnlineClient& ClientToRemoveFrom,
+		ConcertSharedSlate::IEditableReplicationStreamModel& ClientEditModel,
 		TConstArrayView<TSoftObjectPtr<>> Objects,
 		const FConcertPropertyChain& Property
 		)
 	{
-		const TSharedRef<ConcertSharedSlate::IEditableReplicationStreamModel> EditModel = ClientToRemoveFrom.GetClientEditModel();
 		for (const TSoftObjectPtr<>& Object : Objects)
 		{
 			const FSoftObjectPath& ObjectPath = Object.GetUniqueID();
-			const FSoftClassPath ClassPath = EditModel->GetObjectClass(ObjectPath);
-			EditModel->RemoveProperties(ObjectPath, { Property });
+			const FSoftClassPath ClassPath = ClientEditModel.GetObjectClass(ObjectPath);
+			ClientEditModel.RemoveProperties(ObjectPath, { Property });
 					
-			if (EditModel->HasAnyPropertyAssigned(ObjectPath))
+			if (ClientEditModel.HasAnyPropertyAssigned(ObjectPath))
 			{
 				continue;
 			}
@@ -43,25 +63,26 @@ namespace UE::MultiUserClient::Replication::MultiStreamColumns::AssignPropertyMo
 			const bool bIsTopLevelObject = ObjectClass && !ObjectClass->IsChildOf<AActor>();
 			if (bIsTopLevelObject)
 			{
-				EditModel->RemoveObjects({ ObjectPath });
+				ClientEditModel.RemoveObjects({ ObjectPath });
 			}
 		}
 	}
 	
 	template<typename TShouldRemove>
-	requires std::is_invocable_r_v<bool, TShouldRemove, const FOnlineClient&>
+	requires std::is_invocable_r_v<bool, TShouldRemove, const FGuid&>
 	static void UnassignPropertyFromClients(
-		const FOnlineClientManager& ClientManager,
+		const FUnifiedClientView& ClientView,
 		TConstArrayView<TSoftObjectPtr<>> Objects,
 		const FConcertPropertyChain& Property,
 		TShouldRemove&& ShouldRemoveFromClient
 		)
 	{
-		ClientManager.ForEachClient([Objects, Property, &ShouldRemoveFromClient](const FOnlineClient& ClientToRemoveFrom)
+		ClientView.ForEachOnlineClient([&ClientView, &Objects, &Property, &ShouldRemoveFromClient](const FGuid& EndpointId)
 		{
-			if (ShouldRemoveFromClient(ClientToRemoveFrom))
+			const TSharedPtr<ConcertSharedSlate::IEditableReplicationStreamModel> Stream = ClientView.GetEditableClientStreamById(EndpointId);
+			if (Stream && ShouldRemoveFromClient(EndpointId))
 			{
-				RemovePropertiesFromClient(ClientToRemoveFrom, Objects, Property);
+				RemovePropertiesFromClient(*Stream, Objects, Property);
 			}
 			return EBreakBehavior::Continue;
 		});
@@ -70,25 +91,26 @@ namespace UE::MultiUserClient::Replication::MultiStreamColumns::AssignPropertyMo
 
 namespace UE::MultiUserClient::Replication::MultiStreamColumns
 {
-	FAssignPropertyModel::FAssignPropertyModel(FOnlineClientManager& InClientManager)
-		: ClientManager(InClientManager)
+	FAssignPropertyModel::FAssignPropertyModel(FUnifiedClientView& InClientView)
+		: ClientView(InClientView)
 	{
-		ClientManager.OnRemoteClientsChanged().AddRaw(this, &FAssignPropertyModel::BroadcastOnOwnershipChanged);
+		ClientView.OnClientsChanged().AddRaw(this, &FAssignPropertyModel::BroadcastOnOwnershipChanged);
+		ClientView.GetStreamCache().OnCacheChanged().AddRaw(this, &FAssignPropertyModel::BroadcastOnOwnershipChanged);
 	}
 
 	FAssignPropertyModel::~FAssignPropertyModel()
 	{
-		ClientManager.OnRemoteClientsChanged().RemoveAll(this);
+		ClientView.OnClientsChanged().RemoveAll(this);
+		ClientView.GetStreamCache().OnCacheChanged().RemoveAll(this);
 	}
 
 #define SET_REASON(Text) if (Reason) { *Reason = Text; }
 	bool FAssignPropertyModel::CanChangePropertyFor(const FGuid& ClientId, FText* Reason) const
 	{
-		const FOnlineClient* Client = ClientManager.FindClient(ClientId);
-		// Remote clients can disconnect after the combo-box is opened.
-		if (!Client)
+		const TOptional<EClientType> ClientType = ClientView.GetClientType(ClientId);
+		if (!ClientType || !IsOnlineClient(*ClientType))
 		{
-			SET_REASON(LOCTEXT("ClientDisconnected", "Client disconnected."));
+			SET_REASON(LOCTEXT("ClientDisconnected", "Client is not online."));
 			return false;
 		}
 		
@@ -99,8 +121,11 @@ namespace UE::MultiUserClient::Replication::MultiStreamColumns
 	bool FAssignPropertyModel::CanClear(TConstArrayView<TSoftObjectPtr<>> Objects, const FConcertPropertyChain& Property) const
 	{
 		bool bIsAssignedToAnyClient = false;
-		ClientManager.ForEachClient([this, &Objects, &Property, &bIsAssignedToAnyClient](const FOnlineClient& Client)
+		ClientView.ForEachOnlineClient([this, &Objects, &Property, &bIsAssignedToAnyClient](const FGuid& EndpointId)
 		{
+			const TSharedPtr<const ConcertSharedSlate::IReplicationStreamModel> StreamModel = ClientView.GetClientStreamById(EndpointId);
+			check(StreamModel);
+			
 			for (const TSoftObjectPtr<>& EditedObject : Objects)
 			{
 				if (bIsAssignedToAnyClient)
@@ -108,8 +133,7 @@ namespace UE::MultiUserClient::Replication::MultiStreamColumns
 					break;
 				}
 				
-				const TSharedRef<ConcertSharedSlate::IEditableReplicationStreamModel> Model = Client.GetClientEditModel();
-				const bool bHasProperty = Model->HasProperty(EditedObject.GetUniqueID(), Property);
+				const bool bHasProperty = StreamModel->HasProperty(EditedObject.GetUniqueID(), Property);
 				bIsAssignedToAnyClient |= bHasProperty;
 			}
 			
@@ -120,18 +144,18 @@ namespace UE::MultiUserClient::Replication::MultiStreamColumns
 
 	EPropertyOnObjectsOwnershipState FAssignPropertyModel::GetPropertyOwnershipState(const FGuid& ClientId, TConstArrayView<TSoftObjectPtr<>> Objects, const FConcertPropertyChain& Property) const
 	{
-		const FOnlineClient* Client = ClientManager.FindClient(ClientId);
+		const TSharedPtr<const ConcertSharedSlate::IReplicationStreamModel> StreamModel = ClientView.GetClientStreamById(ClientId);
+		const TOptional<EClientType> ClientType = ClientView.GetClientType(ClientId);
 		// Remote clients can disconnect after the combo-box is opened.
-		if (!Client)
+		if (!ClientType || !IsOnlineClient(*ClientType) || !StreamModel)
 		{
-			return {};
+			return EPropertyOnObjectsOwnershipState::NotOwnedOnAllObjects;
 		}
 
-		const TSharedRef<ConcertSharedSlate::IEditableReplicationStreamModel> Model = Client->GetClientEditModel();
 		EPropertyOnObjectsOwnershipState Result = EPropertyOnObjectsOwnershipState::Mixed;
 		for (const TSoftObjectPtr<>& ObjectPath : Objects)
 		{
-			const bool bHasProperty = Model->HasProperty(ObjectPath.GetUniqueID(), Property);
+			const bool bHasProperty = StreamModel->HasProperty(ObjectPath.GetUniqueID(), Property);
 			const EPropertyOnObjectsOwnershipState ExpectedState = bHasProperty
 				? EPropertyOnObjectsOwnershipState::OwnedOnAllObjects
 				: EPropertyOnObjectsOwnershipState::NotOwnedOnAllObjects;
@@ -150,32 +174,15 @@ namespace UE::MultiUserClient::Replication::MultiStreamColumns
 		return Result;
 	}
 
-	void FAssignPropertyModel::AssignPropertyTo(const FOnlineClient* Client, TConstArrayView<TSoftObjectPtr<>> Objects, const FConcertPropertyChain& Property)
-	{
-		const TSharedRef<ConcertSharedSlate::IEditableReplicationStreamModel> EditModel = Client->GetClientEditModel();
-		for (const TSoftObjectPtr<>& Object : Objects)
-		{
-			const FSoftObjectPath& ObjectPath = Object.GetUniqueID();
-			if (!EditModel->ContainsObjects({ ObjectPath }))
-			{
-				EditModel->AddObjects({ Object.Get() });
-			}
-
-			const FSoftClassPath ClassPath = EditModel->GetObjectClass(ObjectPath);
-			TArray AddedProperties { Property };
-			ConcertClientSharedSlate::PropertyUtils::AppendAdditionalPropertiesToAdd(ClassPath, AddedProperties);
-			EditModel->AddProperties(ObjectPath, AddedProperties);
-		}
-	}
-
 	void FAssignPropertyModel::TogglePropertyFor(const FGuid& ClientId, TConstArrayView<TSoftObjectPtr<>> Objects, const FConcertPropertyChain& Property)
 	{
 		// Remote clients can disconnect after the combo-box is opened.
-		const FOnlineClient* Client = ClientManager.FindClient(ClientId);
-		if (!Client)
+		const TSharedPtr<ConcertSharedSlate::IEditableReplicationStreamModel> StreamModel = ClientView.GetEditableClientStreamById(ClientId);
+		if (!StreamModel)
 		{
 			return;
 		}
+		
 		const FText TransactionText = FText::Format(
 			LOCTEXT("AllClientsAssignFmt", "Assign {0} property"),
 			FText::FromString(Property.ToString(FConcertPropertyChain::EToStringMethod::LeafProperty))
@@ -194,12 +201,12 @@ namespace UE::MultiUserClient::Replication::MultiStreamColumns
 		else
 		{
 			// ... remove the property from all clients but the one we'll assign to ...
-			AssignPropertyModel::UnassignPropertyFromClients(ClientManager, Objects, Property,
-				[Client](const FOnlineClient& ClientToRemoveFrom){ return *Client != ClientToRemoveFrom; }
+			AssignPropertyModel::UnassignPropertyFromClients(ClientView, Objects, Property,
+				[ClientId](const FGuid& ClientToRemoveFrom){ return ClientId != ClientToRemoveFrom; }
 				);
 
 			// ... and then assign the property
-			AssignPropertyTo(Client, Objects, Property);
+			AssignPropertyModel::AssignPropertyTo(*StreamModel, Objects, Property);
 		}
 	}
 
@@ -213,7 +220,7 @@ namespace UE::MultiUserClient::Replication::MultiStreamColumns
 				);
 			FScopedTransaction Transaction(TransactionText);
 			
-			AssignPropertyModel::UnassignPropertyFromClients(ClientManager, Objects, Property, [](auto&){ return true; });
+			AssignPropertyModel::UnassignPropertyFromClients(ClientView, Objects, Property, [](auto&){ return true; });
 		}
 	}
 }
diff --git a/Engine/Plugins/Developer/Concert/ConcertApp/MultiUserClient/Source/MultiUserClient/Private/Widgets/ActiveSession/Replication/Client/Multi/Columns/AssignProperty/AssignPropertyModel.h b/Engine/Plugins/Developer/Concert/ConcertApp/MultiUserClient/Source/MultiUserClient/Private/Widgets/ActiveSession/Replication/Client/Multi/Columns/AssignProperty/AssignPropertyModel.h
--- a/Engine/Plugins/Developer/Concert/ConcertApp/MultiUserClient/Source/MultiUserClient/Private/Widgets/ActiveSession/Replication/Client/Multi/Columns/AssignProperty/AssignPropertyModel.h
+++ b/Engine/Plugins/Developer/Concert/ConcertApp/MultiUserClient/Source/MultiUserClient/Private/Widgets/ActiveSession/Replication/Client/Multi/Columns/AssignProperty/AssignPropertyModel.h
@@ -13,6 +13,7 @@ namespace UE::MultiUserClient::Replication
 {
 	class FOnlineClient;
 	class FOnlineClientManager;
+	class FUnifiedClientView;
 }
 
 namespace UE::MultiUserClient::Replication::MultiStreamColumns
@@ -32,7 +33,7 @@ namespace UE::MultiUserClient::Replication::MultiStreamColumns
 	{
 	public:
 		
-		FAssignPropertyModel(FOnlineClientManager& InClientManager UE_LIFETIMEBOUND);
+		FAssignPropertyModel(FUnifiedClientView& InClientView UE_LIFETIMEBOUND);
 		~FAssignPropertyModel();
 
 		/** @return Whether property ownership can be changed for the given client. */
@@ -41,7 +42,6 @@ namespace UE::MultiUserClient::Replication::MultiStreamColumns
 		bool CanClear(TConstArrayView<TSoftObjectPtr<>> Objects, const FConcertPropertyChain& Property) const;
 		/** @return How the property is owned by ClientId */
 		EPropertyOnObjectsOwnershipState GetPropertyOwnershipState(const FGuid& ClientId, TConstArrayView<TSoftObjectPtr<>> Objects, const FConcertPropertyChain& Property) const;
-		void AssignPropertyTo(const FOnlineClient* Client, TConstArrayView<TSoftObjectPtr<>, signed int> Objects, const FConcertPropertyChain& Property);
 
 		/** Assigns the property to ClientId or unassigns it. Removes the property from all other clients in both cases. */
 		void TogglePropertyFor(const FGuid& ClientId, TConstArrayView<TSoftObjectPtr<>> Objects, const FConcertPropertyChain& Property);
@@ -53,8 +53,12 @@ namespace UE::MultiUserClient::Replication::MultiStreamColumns
 
 	private:
 
-		/** Used to get all clients that can be assigned / reassigned to. */
-		FOnlineClientManager& ClientManager;
+		/**
+		 * Used to
+		 * - get all online clients that can be assigned / reassigned to,
+		 * - detect changes made to client content for broadcasting OnOwnershipChangedDelegate
+		 */
+		FUnifiedClientView& ClientView;
 
 		/** Broadcasts when property ownership may have changed. */
 		FSimpleMulticastDelegate OnOwnershipChangedDelegate;
diff --git a/Engine/Plugins/Developer/Concert/ConcertApp/MultiUserClient/Source/MultiUserClient/Private/Widgets/ActiveSession/Replication/Client/Multi/Columns/AssignProperty/MultiStreamColumns_AssignProperties.cpp b/Engine/Plugins/Developer/Concert/ConcertApp/MultiUserClient/Source/MultiUserClient/Private/Widgets/ActiveSession/Replication/Client/Multi/Columns/AssignProperty/MultiStreamColumns_AssignProperties.cpp
--- a/Engine/Plugins/Developer/Concert/ConcertApp/MultiUserClient/Source/MultiUserClient/Private/Widgets/ActiveSession/Replication/Client/Multi/Columns/AssignProperty/MultiStreamColumns_AssignProperties.cpp
+++ b/Engine/Plugins/Developer/Concert/ConcertApp/MultiUserClient/Source/MultiUserClient/Private/Widgets/ActiveSession/Replication/Client/Multi/Columns/AssignProperty/MultiStreamColumns_AssignProperties.cpp
@@ -2,84 +2,60 @@
 
 #include "AssignPropertyModel.h"
 
-#include "IConcertClient.h"
 #include "MultiUserReplicationStyle.h"
-#include "Replication/Client/Online/OnlineClient.h"
-#include "Replication/Client/Online/OnlineClientManager.h"
 #include "Replication/Editor/Model/IEditableMultiReplicationStreamModel.h"
-#include "Replication/Editor/Model/IEditableReplicationStreamModel.h"
+#include "Replication/Client/UnifiedClientViewExtensions.h"
 #include "Replication/Editor/View/IMultiReplicationStreamEditor.h"
 #include "Replication/Editor/View/IReplicationStreamEditor.h"
-#include "SAssignPropertyComboBox.h"
 #include "Replication/Editor/View/Column/IPropertyTreeColumn.h"
 #include "Replication/Editor/View/Column/ReplicationColumnDelegates.h"
-#include "Replication/Client/ClientUtils.h"
+#include "SAssignPropertyComboBox.h"
+
+#include <type_traits>
 
 #define LOCTEXT_NAMESPACE "AssignPropertyColumn"
 
-namespace UE::MultiUserClient::Replication::MultiStreamColumns
+namespace UE::MultiUserClient::Replication::MultiStreamColumns::AssignPropertyColumnUtils
 {
-	const FName AssignPropertyColumnId(TEXT("AssignPropertyColumn"));
-
-	namespace AssignPropertyColumnUtils
+	template<typename TLambda> requires std::is_invocable_v<TLambda, const FGuid&>
+	static void ForEachClientAssignedToProperty(
+		const TArray<TSoftObjectPtr<>>& Objects,
+		const FConcertPropertyChain& Property,
+		const FUnifiedClientView& ClientView,
+		const ConcertSharedSlate::IMultiReplicationStreamEditor& MultiEditor,
+		TLambda&& Callback
+	)
 	{
-		static void ForEachStreamAssignedTo(
-			const ConcertSharedSlate::IMultiReplicationStreamEditor& MultiEditor,
-			const ConcertSharedSlate::FPropertyTreeRowContext& InItem,
-			TFunctionRef<void(const TSharedRef<ConcertSharedSlate::IEditableReplicationStreamModel>& Stream)> Consume
-		)
-		{
-			const FConcertPropertyChain& Property = InItem.RowData.GetProperty();
-		
-			for (const TSharedRef<ConcertSharedSlate::IEditableReplicationStreamModel>& Stream : MultiEditor.GetMultiStreamModel().GetEditableStreams())
+		MultiEditor.GetMultiStreamModel().ForEachStream(
+			[&Objects, &Property,  &ClientView, &Callback](const TSharedRef<ConcertSharedSlate::IReplicationStreamModel>& Stream)
 			{
-				for (const TSoftObjectPtr<>& ContextObject : InItem.RowData.GetContextObjects())
+				for (const TSoftObjectPtr<>& Object : Objects)
 				{
-					const bool bStreamAssignedToProperty = Stream->HasProperty(ContextObject.GetUniqueID(), Property);
-					if (bStreamAssignedToProperty)
+					if (!Stream->HasProperty(Object.ToSoftObjectPath() , Property))
 					{
-						Consume(Stream);
-						// Each IEditableReplicationStreamModel should only be visited at most once.
-						break;
+						continue;
+					}
+				
+					const TOptional<FGuid> ClientEndpointId = FindClientIdByStream(ClientView, *Stream);
+					if (ensure(ClientEndpointId))
+					{
+						Callback(*ClientEndpointId);
+						return EBreakBehavior::Continue;
 					}
 				}
+				return EBreakBehavior::Continue;
 			}
-		}
-
-		const FOnlineClient* FindClientByStream(const FOnlineClientManager& ClientManager, const ConcertSharedSlate::IReplicationStreamModel& StreamModel)
-		{
-			if (&ClientManager.GetLocalClient().GetClientEditModel().Get() == &StreamModel)
-			{
-				return &ClientManager.GetLocalClient();
-			}
-
-			for (const TNonNullPtr<const FRemoteClient> Client : ClientManager.GetRemoteClients())
-			{
-				if (&Client->GetClientEditModel().Get() == &StreamModel)
-				{
-					return Client;
-				}
-			}
-			
-			return nullptr;
-		}
-		
-		static FString GetClientDisplayText(const IConcertClient& InConcertClient, const FOnlineClientManager& ClientManager, const ConcertSharedSlate::IReplicationStreamModel& StreamModel)
-		{
-			if (const FOnlineClient* Client = FindClientByStream(ClientManager, StreamModel))
-			{
-				return ClientUtils::GetClientDisplayName(InConcertClient, Client->GetEndpointId());
-			}
-
-			ensure(false);
-			return {};
-		}
+		);
 	}
+}
+
+namespace UE::MultiUserClient::Replication::MultiStreamColumns
+{
+	const FName AssignPropertyColumnId(TEXT("AssignPropertyColumn"));
 	
 	ConcertSharedSlate::FPropertyColumnEntry AssignPropertyColumn(
 		TAttribute<TSharedPtr<ConcertSharedSlate::IMultiReplicationStreamEditor>> MultiStreamEditor,
-		TSharedRef<IConcertClient> ConcertClient,
-		FOnlineClientManager& ClientManager,
+		FUnifiedClientView& ClientView,
 		const int32 ColumnsSortPriority
 		)
 	{
@@ -90,13 +66,11 @@ namespace UE::MultiUserClient::Replication::MultiStreamColumns
 
 			FPropertyColumn_AssignPropertyColumn(
 				TAttribute<TSharedPtr<IMultiReplicationStreamEditor>> MultiStreamEditor,
-				TSharedRef<IConcertClient> ConcertClient,
-				FOnlineClientManager& ClientManager
+				FUnifiedClientView& ClientView UE_LIFETIMEBOUND
 				)
 				: MultiStreamEditor(MoveTemp(MultiStreamEditor))
-				, ConcertClient(MoveTemp(ConcertClient))
-				, ClientManager(ClientManager)
-				, Model(ClientManager)
+				, Model(ClientView)
+				, ClientView(ClientView)
 			{}
 			
 			virtual SHeaderRow::FColumn::FArguments CreateHeaderRowArgs() const override
@@ -110,7 +84,7 @@ namespace UE::MultiUserClient::Replication::MultiStreamColumns
 			virtual TSharedRef<SWidget> GenerateColumnWidget(const FBuildArgs& InArgs) override
 			{
 				const TArray<TSoftObjectPtr<>> DisplayedObjects = InArgs.RowItem.RowData.GetContextObjects();
-				return SNew(SAssignPropertyComboBox, MultiStreamEditor.Get().ToSharedRef(), ConcertClient, ClientManager, Model)
+				return SNew(SAssignPropertyComboBox, Model, ClientView)
 					.DisplayedProperty(InArgs.RowItem.RowData.GetProperty())
 					.EditedObjects(DisplayedObjects)
 					.HighlightText(InArgs.HighlightText)
@@ -125,21 +99,22 @@ namespace UE::MultiUserClient::Replication::MultiStreamColumns
 			
 			virtual void PopulateSearchString(const FPropertyTreeRowContext& InItem, TArray<FString>& InOutSearchStrings) const override
 			{
-				AssignPropertyColumnUtils::ForEachStreamAssignedTo(
+				AssignPropertyColumnUtils::ForEachClientAssignedToProperty(
+					InItem.RowData.GetContextObjects(),
+					InItem.RowData.GetProperty(),
+					ClientView,
 					*MultiStreamEditor.Get(),
-					InItem,
-					[this, &InOutSearchStrings](const TSharedRef<IEditableReplicationStreamModel>& Stream)
+					[this, &InOutSearchStrings](const FGuid& ClientId)
 					{
-						InOutSearchStrings.Add(AssignPropertyColumnUtils::GetClientDisplayText(*ConcertClient, ClientManager, *Stream));
+						InOutSearchStrings.Add(GetClientDisplayString(ClientView, ClientId));
 					});
 			}
 
 			virtual bool CanBeSorted() const override { return true; }
 			virtual bool IsLessThan(const FPropertyTreeRowContext& Left, const FPropertyTreeRowContext& Right) const override
 			{
-				const TArray<TSoftObjectPtr<>> DisplayedObjects = MultiStreamEditor.Get()->GetEditorBase().GetSelectedObjects();
-				const TOptional<FString> LeftClientDisplayString = SAssignPropertyComboBox::GetDisplayString(ConcertClient, ClientManager, Left.RowData.GetProperty(), DisplayedObjects);
-				const TOptional<FString> RightClientDisplayString = SAssignPropertyComboBox::GetDisplayString(ConcertClient, ClientManager, Right.RowData.GetProperty(), DisplayedObjects);
+				const TOptional<FString> LeftClientDisplayString = SAssignPropertyComboBox::GetDisplayString(ClientView, Left.RowData.GetProperty(), Left.RowData.GetContextObjects());
+				const TOptional<FString> RightClientDisplayString = SAssignPropertyComboBox::GetDisplayString(ClientView, Right.RowData.GetProperty(), Right.RowData.GetContextObjects());
 			
 				if (LeftClientDisplayString && RightClientDisplayString)
 				{
@@ -151,20 +126,21 @@ namespace UE::MultiUserClient::Replication::MultiStreamColumns
 
 		private:
 
+			/** Used to refresh the sort state of the editor, if sorting by this column. */
 			const TAttribute<TSharedPtr<IMultiReplicationStreamEditor>> MultiStreamEditor;
-			const TSharedRef<IConcertClient> ConcertClient;
-			FOnlineClientManager& ClientManager;
 
 			/** The model the view displays. */
 			FAssignPropertyModel Model;
+			/** Used to get display information about clients. */
+			FUnifiedClientView& ClientView;
 		};
 		
 		check(MultiStreamEditor.IsBound() || MultiStreamEditor.IsSet());
 		return {
 			TReplicationColumnDelegates<FPropertyTreeRowContext>::FCreateColumn::CreateLambda(
-				[MultiStreamEditor = MoveTemp(MultiStreamEditor), ConcertClient = MoveTemp(ConcertClient), &ClientManager]()
+				[MultiStreamEditor = MoveTemp(MultiStreamEditor), &ClientView]()
 				{
-					return MakeShared<FPropertyColumn_AssignPropertyColumn>(MultiStreamEditor, ConcertClient, ClientManager);
+					return MakeShared<FPropertyColumn_AssignPropertyColumn>(MultiStreamEditor, ClientView);
 				}),
 			AssignPropertyColumnId,
 			{ ColumnsSortPriority }
diff --git a/Engine/Plugins/Developer/Concert/ConcertApp/MultiUserClient/Source/MultiUserClient/Private/Widgets/ActiveSession/Replication/Client/Multi/Columns/AssignProperty/SAssignPropertyComboBox.cpp b/Engine/Plugins/Developer/Concert/ConcertApp/MultiUserClient/Source/MultiUserClient/Private/Widgets/ActiveSession/Replication/Client/Multi/Columns/AssignProperty/SAssignPropertyComboBox.cpp
--- a/Engine/Plugins/Developer/Concert/ConcertApp/MultiUserClient/Source/MultiUserClient/Private/Widgets/ActiveSession/Replication/Client/Multi/Columns/AssignProperty/SAssignPropertyComboBox.cpp
+++ b/Engine/Plugins/Developer/Concert/ConcertApp/MultiUserClient/Source/MultiUserClient/Private/Widgets/ActiveSession/Replication/Client/Multi/Columns/AssignProperty/SAssignPropertyComboBox.cpp
@@ -4,59 +4,70 @@
 
 #include "AssignPropertyModel.h"
 #include "IConcertClient.h"
-#include "Replication/Client/ClientUtils.h"
 #include "Replication/Client/Online/OnlineClient.h"
 #include "Replication/Client/Online/OnlineClientManager.h"
+#include "Replication/Client/UnifiedClientViewExtensions.h"
 #include "Replication/Editor/Model/PropertyUtils.h"
 #include "Widgets/ActiveSession/Replication/Misc/SNoClients.h"
 #include "Widgets/Client/ClientInfoHelpers.h"
 #include "Widgets/Client/SHorizontalClientList.h"
-#include "Widgets/Client/SLocalClientName.h"
-#include "Widgets/Client/SRemoteClientName.h"
+#include "Widgets/Client/SClientName.h"
 
 #include "Framework/MultiBox/MultiBoxBuilder.h"
 #include "Styling/AppStyle.h"
 #include "Styling/SlateTypes.h"
 #include "Widgets/Input/SComboButton.h"
 
+#include <type_traits>
+
 #define LOCTEXT_NAMESPACE "SAssignPropertyComboBox"
 
 namespace UE::MultiUserClient::Replication::MultiStreamColumns
 {
 	namespace AssignPropertyComboBox
 	{
-		TArray<FGuid> GetDisplayedClients(const FOnlineClientManager& ClientManager, const FConcertPropertyChain& DisplayedProperty, const TArray<TSoftObjectPtr<>>& EditedObjects)
+		template<typename TLambda> requires std::is_invocable_r_v<EBreakBehavior, TLambda, const FGuid&>
+		static void ForEachDisplayedClient(
+			const FUnifiedClientView& ClientView,
+			const FConcertPropertyChain& DisplayedProperty,
+			const TArray<TSoftObjectPtr<>>& EditedObjects,
+			TLambda&& Callback
+			)
 		{
-			TArray<FGuid> Clients;
-			ClientManager.ForEachClient([&DisplayedProperty, &EditedObjects, &Clients](const FOnlineClient& Client)
+			for (const TSoftObjectPtr<>& Object : EditedObjects)
 			{
-				const TMap<FSoftObjectPath, FConcertReplicatedObjectInfo>& ObjectInfoMap = Client.GetStreamSynchronizer().GetServerState().ReplicatedObjects;
-				for (const TSoftObjectPtr<>& ObjectPath : EditedObjects)
-				{
-					if (const FConcertReplicatedObjectInfo* ObjectInfo = ObjectInfoMap.Find(ObjectPath.GetUniqueID())
-						; ObjectInfo && ObjectInfo->PropertySelection.ReplicatedProperties.Contains(DisplayedProperty))
-					{
-						Clients.Add(Client.GetEndpointId());
-						return EBreakBehavior::Continue;
-					}
-				}
-				return EBreakBehavior::Continue;
-			});
+				ClientView.GetStreamCache().EnumerateClientsWithObjectAndProperty(
+					Object.ToSoftObjectPath(),
+					DisplayedProperty,
+					[&Callback](const FGuid& ClientId){ return Callback(ClientId); }
+				);
+			}
+		}
+		
+		static TArray<FGuid> GetDisplayedClients(
+			const FUnifiedClientView& ClientView,
+			const FConcertPropertyChain& DisplayedProperty,
+			const TArray<TSoftObjectPtr<>>& EditedObjects
+			)
+		{
+			TArray<FGuid> Clients;
+			ForEachDisplayedClient(ClientView, DisplayedProperty, EditedObjects,
+				[&Clients](const FGuid& ClientId){ Clients.AddUnique(ClientId); return EBreakBehavior::Continue; }
+				);
 			return Clients;
 		}
 	}
 	
 	TOptional<FString> SAssignPropertyComboBox::GetDisplayString(
-		const TSharedRef<IConcertClient>& LocalConcertClient,
-		const FOnlineClientManager& ClientManager,
+		const FUnifiedClientView& ClientView,
 		const FConcertPropertyChain& DisplayedProperty,
 		const TArray<TSoftObjectPtr<>>& EditedObjects)
 	{
 		using SWidgetType = ConcertSharedSlate::SHorizontalClientList;
-		const TArray<FGuid> Clients = AssignPropertyComboBox::GetDisplayedClients(ClientManager, DisplayedProperty, EditedObjects);
+		const TArray<FGuid> Clients = AssignPropertyComboBox::GetDisplayedClients(ClientView, DisplayedProperty, EditedObjects);
 		
 		const ConcertSharedSlate::FGetClientParenthesesContent GetParenthesesContent =
-			ConcertClientSharedSlate::MakeGetLocalClientParenthesesContent(LocalConcertClient);
+			MakeLocalAndOfflineParenthesesContentGetter(ClientView);
 		const auto SortPredicate = [&GetParenthesesContent](const FConcertSessionClientInfo& Left, const FConcertSessionClientInfo& Right)
 		{
 			return ConcertSharedSlate::SortLocalClientParenthesesFirstThenThenAlphabetical(Left, Right, GetParenthesesContent);
@@ -64,23 +75,20 @@ namespace UE::MultiUserClient::Replication::MultiStreamColumns
 		
 		return SWidgetType::GetDisplayString(
 			Clients,
-			ConcertClientSharedSlate::MakeClientInfoGetter(LocalConcertClient),
+			MakeOnlineThenOfflineClientInfoGetter(ClientView),
 			ConcertSharedSlate::FClientSortPredicate::CreateLambda(SortPredicate),
 			GetParenthesesContent
 			);
 	}
 
-	void SAssignPropertyComboBox::Construct(const FArguments& InArgs,
-	    TSharedRef<ConcertSharedSlate::IMultiReplicationStreamEditor> InEditor,
-	    TSharedRef<IConcertClient> InConcertClient,
-	    FOnlineClientManager& InClientManager,
-	    FAssignPropertyModel& InModel
+	void SAssignPropertyComboBox::Construct(
+		const FArguments& InArgs,
+	    FAssignPropertyModel& InModel,
+	    FUnifiedClientView& InClientView
 	)
 	{
-		Editor = MoveTemp(InEditor);
-		ConcertClient = MoveTemp(InConcertClient);
-		ClientManager = &InClientManager;
 		Model = &InModel;
+		ClientView = &InClientView;
 		
 		Property = InArgs._DisplayedProperty;
 		EditedObjects = InArgs._EditedObjects;
@@ -96,46 +104,30 @@ namespace UE::MultiUserClient::Replication::MultiStreamColumns
 			.ButtonContent()
 			[
 				SAssignNew(ClientListWidget, ConcertSharedSlate::SHorizontalClientList)
-				.GetClientParenthesesContent(ConcertClientSharedSlate::MakeGetLocalClientParenthesesContent(ConcertClient.ToSharedRef()))
-				.GetClientInfo(ConcertClientSharedSlate::MakeClientInfoGetter(ConcertClient.ToSharedRef()))
+				.GetClientParenthesesContent(MakeLocalAndOfflineParenthesesContentGetter(InClientView))
+				.GetClientInfo(MakeOnlineThenOfflineClientInfoGetter(InClientView))
 				.Font(FAppStyle::GetFontStyle("PropertyWindow.NormalFont"))
 				.HighlightText_Lambda([this](){ return HighlightText ? *HighlightText : FText::GetEmpty(); })
 				.EmptyListSlot() [ SNew(SNoClients) ]
 			]
 			.OnGetMenuContent(this, &SAssignPropertyComboBox::GetMenuContent)
+			.ToolTipText_Raw(this, &SAssignPropertyComboBox::GetComboBoxToolTipText)
 		];
 
-		Model->OnOwnershipChanged().AddSP(this, &SAssignPropertyComboBox::RebuildSubscriptionsAndRefresh);
-		RebuildSubscriptions();
+		Model->OnOwnershipChanged().AddSP(this, &SAssignPropertyComboBox::RefreshContentBoxContent);
 		RefreshContentBoxContent();
 	}
 	
 	void SAssignPropertyComboBox::RefreshContentBoxContent() const
 	{
 		ClientListWidget->RefreshList(
-			AssignPropertyComboBox::GetDisplayedClients(*ClientManager, Property, EditedObjects)
+			AssignPropertyComboBox::GetDisplayedClients(*ClientView, Property, EditedObjects)
 			);
 	}
 
 	TSharedRef<SWidget> SAssignPropertyComboBox::GetMenuContent()
 	{
 		using namespace ConcertSharedSlate;
-
-		const auto MakeWidget = [this](const FGuid& EndpointId) -> TSharedRef<SWidget>
-		{
-			const bool bIsLocalClient = EndpointId == ConcertClient->GetCurrentSession()->GetSessionClientEndpointId();
-			if (bIsLocalClient)
-			{
-				return SNew(SLocalClientName)
-					.DisplayInfo(ConcertClientSharedSlate::MakeLocalClientInfoAttribute(ConcertClient.ToSharedRef()))
-					.Font(FAppStyle::GetFontStyle("PropertyWindow.NormalFont"))
-					.HighlightText_Lambda([this](){ return HighlightText ? *HighlightText : FText::GetEmpty(); });
-			}
-			return SNew(SRemoteClientName)
-				.DisplayInfo(ConcertClientSharedSlate::MakeClientInfoAttribute(ConcertClient.ToSharedRef(), EndpointId))
-				.Font(FAppStyle::GetFontStyle("PropertyWindow.NormalFont"))
-				.HighlightText_Lambda([this](){ return HighlightText ? *HighlightText : FText::GetEmpty(); });
-		};
 		
 		FMenuBuilder MenuBuilder(true, nullptr);
 		MenuBuilder.AddMenuEntry(
@@ -150,10 +142,19 @@ namespace UE::MultiUserClient::Replication::MultiStreamColumns
 			EUserInterfaceActionType::Button
 		);
 		
-		MenuBuilder.BeginSection(TEXT("AssignTo"), LOCTEXT("AssignTo", "Assign to"));
-		for (const FOnlineClient* Client : ClientUtils::GetSortedClientList(*ConcertClient, *ClientManager))
+		const auto MakeWidget = [this](const FGuid& EndpointId) -> TSharedRef<SWidget>
+		{
+			return SNew(SClientName)
+				.ClientInfo_Lambda([this, EndpointId]{ return ClientView->GetClientInfoByEndpoint(EndpointId); })
+				.ParenthesisContent_Lambda([this, EndpointId]{ return GetParenthesesContent(*ClientView, EndpointId); })
+				.HighlightText_Lambda([this](){ return HighlightText ? *HighlightText : FText::GetEmpty(); })
+				.Font(FAppStyle::GetFontStyle("PropertyWindow.NormalFont"));
+		};
+		
+		MenuBuilder.BeginSection(TEXT("AssignTo"), LOCTEXT("AssignTo", "Assign to online client"));
+		for (const FGuid& EndpointId : GetSortedOnlineClients(*ClientView))
 		{
-			TAttribute<FText> Tooltip = TAttribute<FText>::CreateLambda([this, EndpointId = Client->GetEndpointId()]()
+			TAttribute<FText> Tooltip = TAttribute<FText>::CreateLambda([this, EndpointId]()
 			{
 				FText Reason;
 				const bool bCanClick = CanClickOptionWithReason(EndpointId, &Reason);
@@ -173,11 +174,11 @@ namespace UE::MultiUserClient::Replication::MultiStreamColumns
 			
 			MenuBuilder.AddMenuEntry(
 				FUIAction(
-					FExecuteAction::CreateSP(this, &SAssignPropertyComboBox::OnClickOption, Client->GetEndpointId()),
-					FCanExecuteAction::CreateSP(this, &SAssignPropertyComboBox::CanClickOption, Client->GetEndpointId()),
-					FGetActionCheckState::CreateSP(this, &SAssignPropertyComboBox::GetOptionCheckState, Client->GetEndpointId())
+					FExecuteAction::CreateSP(this, &SAssignPropertyComboBox::OnClickOption, EndpointId),
+					FCanExecuteAction::CreateSP(this, &SAssignPropertyComboBox::CanClickOption, EndpointId),
+					FGetActionCheckState::CreateSP(this, &SAssignPropertyComboBox::GetOptionCheckState, EndpointId)
 					),
-				MakeWidget(Client->GetEndpointId()),
+				MakeWidget(EndpointId),
 				NAME_None,
 				Tooltip,
 				EUserInterfaceActionType::Check
@@ -221,14 +222,31 @@ namespace UE::MultiUserClient::Replication::MultiStreamColumns
 		return Model->CanClear(EditedObjects, Property);
 	}
 
-	void SAssignPropertyComboBox::RebuildSubscriptions()
+	FText SAssignPropertyComboBox::GetComboBoxToolTipText() const
 	{
-		ClientManager->ForEachClient([this](FOnlineClient& Client)
+		bool bHasOfflineClients = false;
+		int32 NumClients = 0;
+		AssignPropertyComboBox::ForEachDisplayedClient(*ClientView, Property, EditedObjects,
+			[this, &bHasOfflineClients, &NumClients](const FGuid& ClientId)
+			{
+				++NumClients;
+				const TOptional<EClientType> ClientType = ClientView->GetClientType(ClientId);
+				bHasOfflineClients |= ClientType && IsOfflineClient(*ClientType);
+				return NumClients > 1 ? EBreakBehavior::Break : EBreakBehavior::Continue;
+			});
+
+		if (bHasOfflineClients)
 		{
-			Client.OnModelChanged().RemoveAll(this);
-			Client.OnModelChanged().AddSP(this, &SAssignPropertyComboBox::RefreshContentBoxContent);
-			return EBreakBehavior::Continue;
-		});
+			return FText::Format(
+				LOCTEXT("AssignProperty.ToolTip.HasOfflineClients", "The assigned offline {0}|plural(one=client,other=clients) will replicate this property upon rejoining."),
+				NumClients
+				);
+		}
+
+		return FText::Format(
+			LOCTEXT("AssignProperty.ToolTip.Normal", "The {0}|plural(one=client,other=clients) that {0}|plural(one=is,other=are) have registered to replicate this property."),
+			NumClients
+			);
 	}
 }
 
diff --git a/Engine/Plugins/Developer/Concert/ConcertApp/MultiUserClient/Source/MultiUserClient/Private/Widgets/ActiveSession/Replication/Client/Multi/Columns/AssignProperty/SAssignPropertyComboBox.h b/Engine/Plugins/Developer/Concert/ConcertApp/MultiUserClient/Source/MultiUserClient/Private/Widgets/ActiveSession/Replication/Client/Multi/Columns/AssignProperty/SAssignPropertyComboBox.h
--- a/Engine/Plugins/Developer/Concert/ConcertApp/MultiUserClient/Source/MultiUserClient/Private/Widgets/ActiveSession/Replication/Client/Multi/Columns/AssignProperty/SAssignPropertyComboBox.h
+++ b/Engine/Plugins/Developer/Concert/ConcertApp/MultiUserClient/Source/MultiUserClient/Private/Widgets/ActiveSession/Replication/Client/Multi/Columns/AssignProperty/SAssignPropertyComboBox.h
@@ -13,17 +13,8 @@ class IConcertClient;
 enum class ECheckBoxState : uint8;
 struct FGuid;
 
-namespace UE::ConcertSharedSlate
-{
-	class IMultiReplicationStreamEditor;
-	class SHorizontalClientList;
-}
-
-namespace UE::MultiUserClient::Replication
-{
-	class FOnlineClient;
-	class FOnlineClientManager;
-}
+namespace UE::ConcertSharedSlate { class SHorizontalClientList; }
+namespace UE::MultiUserClient::Replication { class FUnifiedClientView; }
 
 namespace UE::MultiUserClient::Replication::MultiStreamColumns
 {
@@ -43,8 +34,7 @@ namespace UE::MultiUserClient::Replication::MultiStreamColumns
 
 		/** @return The display string this widget would have with the given state. If unset, no clients are displayed in the combobox.*/
 		static TOptional<FString> GetDisplayString(
-			const TSharedRef<IConcertClient>& LocalConcertClient,
-			const FOnlineClientManager& ClientManager,
+			const FUnifiedClientView& ClientView,
 			const FConcertPropertyChain& DisplayedProperty,
 			const TArray<TSoftObjectPtr<>>& EditedObjects
 			);
@@ -59,21 +49,16 @@ namespace UE::MultiUserClient::Replication::MultiStreamColumns
 		SLATE_END_ARGS()
 
 		void Construct(const FArguments& InArgs,
-           TSharedRef<ConcertSharedSlate::IMultiReplicationStreamEditor> InEditor,
-           TSharedRef<IConcertClient> InConcertClient,
-           FOnlineClientManager& InClientManager,
-           FAssignPropertyModel& InModel UE_LIFETIMEBOUND
+           FAssignPropertyModel& InModel UE_LIFETIMEBOUND,
+           FUnifiedClientView& InClientView UE_LIFETIMEBOUND
 		);
 
 	private:
 
-		FOnlineClientManager* ClientManager = nullptr;
 		/** The model this view is displaying. */
 		FAssignPropertyModel* Model = nullptr;
-		
-		/** Used to obtain info about the streams */
-		TSharedPtr<ConcertSharedSlate::IMultiReplicationStreamEditor> Editor;
-		TSharedPtr<IConcertClient> ConcertClient;
+		/** Used to obtain display information about clients. */
+		FUnifiedClientView* ClientView = nullptr;
 
 		/** The objects for which the property is being displayed. */
 		TArray<TSoftObjectPtr<>> EditedObjects;
@@ -106,10 +91,8 @@ namespace UE::MultiUserClient::Replication::MultiStreamColumns
 		void OnClickClear();
 		/** Whether the property is assigned to any client */
 		bool CanClickClear() const;
-
-		/** Resubscribes to all clients changing. */
-		void RebuildSubscriptions();
-		void RebuildSubscriptionsAndRefresh() { RebuildSubscriptions(); RefreshContentBoxContent(); }
+		
+		FText GetComboBoxToolTipText() const;
 	};
 }
 
diff --git a/Engine/Plugins/Developer/Concert/ConcertApp/MultiUserClient/Source/MultiUserClient/Private/Widgets/ActiveSession/Replication/Client/Multi/Columns/AssignedClients/AssignedClientsModel.cpp b/Engine/Plugins/Developer/Concert/ConcertApp/MultiUserClient/Source/MultiUserClient/Private/Widgets/ActiveSession/Replication/Client/Multi/Columns/AssignedClients/AssignedClientsModel.cpp
--- a/Engine/Plugins/Developer/Concert/ConcertApp/MultiUserClient/Source/MultiUserClient/Private/Widgets/ActiveSession/Replication/Client/Multi/Columns/AssignedClients/AssignedClientsModel.cpp
+++ b/Engine/Plugins/Developer/Concert/ConcertApp/MultiUserClient/Source/MultiUserClient/Private/Widgets/ActiveSession/Replication/Client/Multi/Columns/AssignedClients/AssignedClientsModel.cpp
@@ -2,8 +2,7 @@
 
 #include "AssignedClientsModel.h"
 
-#include "Replication/Client/Online/OnlineClient.h"
-#include "Replication/Client/Online/OnlineClientManager.h"
+#include "Replication/Client/UnifiedStreamCache.h"
 #include "Replication/Editor/Model/Object/IObjectHierarchyModel.h"
 
 #include "Containers/Array.h"
@@ -11,54 +10,33 @@
 
 #include <type_traits>
 
-namespace UE::MultiUserClient::Replication::MultiStreamColumns::Private
-{
-	template<typename TCallback> requires std::is_invocable_v<TCallback, const FGuid&>
-	static void EnumerateOwningOnlineClients(
-		const FOnlineClientManager& ClientManager,
-		const FSoftObjectPath& ObjectPath,
-		TCallback&& Callback
-		) 
-	{
-		ClientManager.ForEachClient([&ObjectPath, &Callback](const FOnlineClient& ReplicationClient)
-		{
-			const bool bHasProperties = ReplicationClient.GetStreamSynchronizer().GetServerState().HasProperties(ObjectPath);
-			if (bHasProperties)
-			{
-				Callback(ReplicationClient.GetEndpointId());
-			}
-			return EBreakBehavior::Continue;
-		});
-	}
-}
-
 namespace UE::MultiUserClient::Replication::MultiStreamColumns
 {
 	FAssignedClientsModel::FAssignedClientsModel(
 		const ConcertSharedSlate::IObjectHierarchyModel& ObjectHierarchy,
-		FOnlineClientManager& InClientManager
+		FUnifiedStreamCache& InStreamCache
 		)
 		: ObjectHierarchy(ObjectHierarchy)
-		, ClientManager(InClientManager)
+		, StreamCache(InStreamCache)
 	{
-		ClientManager.OnRemoteClientsChanged().AddRaw(this, &FAssignedClientsModel::BroadcastOwnershipChanged);
-		ClientManager.GetAuthorityCache().OnCacheChanged().AddRaw(this, &FAssignedClientsModel::OnClientChanged);
+		StreamCache.OnCacheChanged().AddRaw(this, &FAssignedClientsModel::BroadcastOwnershipChanged);
 	}
 
 	FAssignedClientsModel::~FAssignedClientsModel()
 	{
-		ClientManager.OnRemoteClientsChanged().RemoveAll(this);
-		ClientManager.GetAuthorityCache().OnCacheChanged().RemoveAll(this);
+		StreamCache.OnCacheChanged().RemoveAll(this);
 	}
 
 	TArray<FGuid> FAssignedClientsModel::GetAssignedClients(const FSoftObjectPath& ObjectPath) const
 	{
 		TArray<FGuid> ClientsWithOwnership;
+
 		const auto ProcessObject = [this, &ClientsWithOwnership](const FSoftObjectPath& ObjectPath)
 		{
-			Private::EnumerateOwningOnlineClients(ClientManager, ObjectPath, [&ClientsWithOwnership](const FGuid& ClientId)
+			StreamCache.EnumerateClientsWithObject(ObjectPath, [&ClientsWithOwnership](const FGuid& ClientId)
 			{
 				ClientsWithOwnership.AddUnique(ClientId);
+				return EBreakBehavior::Continue;
 			});
 		};
 			
@@ -70,7 +48,17 @@ namespace UE::MultiUserClient::Replication::MultiStreamColumns
 				ProcessObject(ChildObject.GetUniqueID());
 				return EBreakBehavior::Continue;
 			});
-			
+		
 		return ClientsWithOwnership;
 	}
+
+	void FAssignedClientsModel::OnClientChanged(const FGuid&) const
+	{
+		BroadcastOwnershipChanged();
+	}
+
+	void FAssignedClientsModel::BroadcastOwnershipChanged() const
+	{
+		OnOwnershipChangedDelegate.Broadcast();
+	}
 }
diff --git a/Engine/Plugins/Developer/Concert/ConcertApp/MultiUserClient/Source/MultiUserClient/Private/Widgets/ActiveSession/Replication/Client/Multi/Columns/AssignedClients/AssignedClientsModel.h b/Engine/Plugins/Developer/Concert/ConcertApp/MultiUserClient/Source/MultiUserClient/Private/Widgets/ActiveSession/Replication/Client/Multi/Columns/AssignedClients/AssignedClientsModel.h
--- a/Engine/Plugins/Developer/Concert/ConcertApp/MultiUserClient/Source/MultiUserClient/Private/Widgets/ActiveSession/Replication/Client/Multi/Columns/AssignedClients/AssignedClientsModel.h
+++ b/Engine/Plugins/Developer/Concert/ConcertApp/MultiUserClient/Source/MultiUserClient/Private/Widgets/ActiveSession/Replication/Client/Multi/Columns/AssignedClients/AssignedClientsModel.h
@@ -7,6 +7,16 @@
 #include "HAL/Platform.h"
 #include "Templates/UnrealTemplate.h"
 
+namespace UE::MultiUserClient::Replication
+{
+	class FUnifiedStreamCache;
+}
+
+namespace UE::MultiUserClient::Replication
+{
+	class FUnifiedClientView;
+}
+
 struct FGuid;
 
 namespace UE::ConcertSharedSlate { class IObjectHierarchyModel; }
@@ -25,10 +35,11 @@ namespace UE::MultiUserClient::Replication::MultiStreamColumns
 		
 		FAssignedClientsModel(
 			const ConcertSharedSlate::IObjectHierarchyModel& ObjectHierarchy UE_LIFETIMEBOUND,
-			FOnlineClientManager& InClientManager UE_LIFETIMEBOUND
+			FUnifiedStreamCache& InStreamCache UE_LIFETIMEBOUND
 			);
 		~FAssignedClientsModel();
 
+		/** @return The endpoint IDs of clients that have properties assigned to ObjecPath. */
 		TArray<FGuid> GetAssignedClients(const FSoftObjectPath& ObjectPath) const;
 		
 		/** Broadcast when the result of GetAssignedClients may have changed. */
@@ -38,13 +49,14 @@ namespace UE::MultiUserClient::Replication::MultiStreamColumns
 
 		/** Used to compute clients's recursive ownership of subobjects for a managed object. */
 		const ConcertSharedSlate::IObjectHierarchyModel& ObjectHierarchy;
-		FOnlineClientManager& ClientManager;
+		/** Holds the stream content of both online and offline clients. */
+		FUnifiedStreamCache& StreamCache;
 		
 		/** Broadcast when the result of GetAssignedClients may have changed. */
 		FSimpleMulticastDelegate OnOwnershipChangedDelegate;
 		
-		void OnClientChanged(const FGuid&) const { BroadcastOwnershipChanged(); }
-		void BroadcastOwnershipChanged() const { OnOwnershipChangedDelegate.Broadcast(); }
+		void OnClientChanged(const FGuid&) const;
+		void BroadcastOwnershipChanged() const;
 	};
 }
 
diff --git a/Engine/Plugins/Developer/Concert/ConcertApp/MultiUserClient/Source/MultiUserClient/Private/Widgets/ActiveSession/Replication/Client/Multi/Columns/AssignedClients/MultiStreamColumns_AssignedClients.cpp b/Engine/Plugins/Developer/Concert/ConcertApp/MultiUserClient/Source/MultiUserClient/Private/Widgets/ActiveSession/Replication/Client/Multi/Columns/AssignedClients/MultiStreamColumns_AssignedClients.cpp
--- a/Engine/Plugins/Developer/Concert/ConcertApp/MultiUserClient/Source/MultiUserClient/Private/Widgets/ActiveSession/Replication/Client/Multi/Columns/AssignedClients/MultiStreamColumns_AssignedClients.cpp
+++ b/Engine/Plugins/Developer/Concert/ConcertApp/MultiUserClient/Source/MultiUserClient/Private/Widgets/ActiveSession/Replication/Client/Multi/Columns/AssignedClients/MultiStreamColumns_AssignedClients.cpp
@@ -7,6 +7,7 @@
 #include "Replication/Client/ClientUtils.h"
 #include "Replication/Editor/View/IMultiReplicationStreamEditor.h"
 #include "SAssignedClientsWidget.h"
+#include "Replication/Client/UnifiedClientView.h"
 #include "Widgets/ActiveSession/Replication/Client/Multi/Columns/MultiStreamColumns.h"
 #include "Widgets/Client/ClientInfoHelpers.h"
 #include "Widgets/Client/SHorizontalClientList.h"
@@ -20,10 +21,10 @@ namespace UE::MultiUserClient::Replication::MultiStreamColumns
 	const FName AssignedClientsColumnId(TEXT("AssignedClientsColumn"));
 
 	ConcertSharedSlate::FObjectColumnEntry AssignedClientsColumn(
-		TSharedRef<IConcertClient> ConcertClient,
+		const TSharedRef<IConcertClient>& ConcertClient,
 		TAttribute<TSharedPtr<ConcertSharedSlate::IMultiReplicationStreamEditor>> MultiStreamModelAttribute,
 		const ConcertSharedSlate::IObjectHierarchyModel& ObjectHierarchy,
-		FOnlineClientManager& ClientManager,
+		FUnifiedClientView& ClientView,
 		const int32 ColumnsSortPriority
 		)
 	{
@@ -32,14 +33,15 @@ namespace UE::MultiUserClient::Replication::MultiStreamColumns
 		public:
 
 			FObjectColumn_AssignedClients(
-				TSharedRef<IConcertClient> ConcertClient,
+				const TSharedRef<IConcertClient>& ConcertClient,
 				TAttribute<TSharedPtr<ConcertSharedSlate::IMultiReplicationStreamEditor>> MultiStreamModelAttribute,
 				const ConcertSharedSlate::IObjectHierarchyModel& ObjectHierarchy UE_LIFETIMEBOUND,
-				FOnlineClientManager& ClientManager UE_LIFETIMEBOUND
+				FUnifiedClientView& ClientView UE_LIFETIMEBOUND
 				)
-				: ConcertClient(MoveTemp(ConcertClient))
+				: ConcertClient(ConcertClient)
 				, MultiStreamModelAttribute(MoveTemp(MultiStreamModelAttribute))
-				, Model(ObjectHierarchy, ClientManager)
+				, Model(ObjectHierarchy, ClientView.GetStreamCache())
+				, ClientView(ClientView)
 			{}
 			
 			virtual SHeaderRow::FColumn::FArguments CreateHeaderRowArgs() const override
@@ -52,7 +54,7 @@ namespace UE::MultiUserClient::Replication::MultiStreamColumns
 			
 			virtual TSharedRef<SWidget> GenerateColumnWidget(const FBuildArgs& InArgs) override
 			{
-				return SNew(SAssignedClientsWidget, ConcertClient, Model)
+				return SNew(SAssignedClientsWidget, Model, ClientView)
 					.ManagedObject(InArgs.RowItem.RowData.GetObjectPath())
 					.HighlightText_Lambda([HighlightText = InArgs.HighlightText](){ return HighlightText ? *HighlightText : FText::GetEmpty(); });
 			}
@@ -85,6 +87,8 @@ namespace UE::MultiUserClient::Replication::MultiStreamColumns
 
 			/** The model that the view displays; the model of the MVC pattern. */
 			FAssignedClientsModel Model;
+			/** Passed to the UI. */
+			FUnifiedClientView& ClientView;
 			
 			TOptional<FString> GetDisplayString(const FSoftObjectPath& ManagedObject) const
 			{
@@ -108,9 +112,9 @@ namespace UE::MultiUserClient::Replication::MultiStreamColumns
 
 		return {
 			ConcertSharedSlate::TReplicationColumnDelegates<ConcertSharedSlate::FObjectTreeRowContext>::FCreateColumn::CreateLambda(
-				[ConcertClient = MoveTemp(ConcertClient), MultiStreamModelAttribute = MoveTemp(MultiStreamModelAttribute), &ObjectHierarchy, &ClientManager]()
+				[ConcertClient, MultiStreamModelAttribute = MoveTemp(MultiStreamModelAttribute), &ObjectHierarchy, &ClientView]()
 				{
-					return MakeShared<FObjectColumn_AssignedClients>(ConcertClient, MultiStreamModelAttribute, ObjectHierarchy, ClientManager);
+					return MakeShared<FObjectColumn_AssignedClients>(ConcertClient, MultiStreamModelAttribute, ObjectHierarchy, ClientView);
 				}),
 			AssignedClientsColumnId,
 			{ ColumnsSortPriority }
diff --git a/Engine/Plugins/Developer/Concert/ConcertApp/MultiUserClient/Source/MultiUserClient/Private/Widgets/ActiveSession/Replication/Client/Multi/Columns/AssignedClients/SAssignedClientsWidget.cpp b/Engine/Plugins/Developer/Concert/ConcertApp/MultiUserClient/Source/MultiUserClient/Private/Widgets/ActiveSession/Replication/Client/Multi/Columns/AssignedClients/SAssignedClientsWidget.cpp
--- a/Engine/Plugins/Developer/Concert/ConcertApp/MultiUserClient/Source/MultiUserClient/Private/Widgets/ActiveSession/Replication/Client/Multi/Columns/AssignedClients/SAssignedClientsWidget.cpp
+++ b/Engine/Plugins/Developer/Concert/ConcertApp/MultiUserClient/Source/MultiUserClient/Private/Widgets/ActiveSession/Replication/Client/Multi/Columns/AssignedClients/SAssignedClientsWidget.cpp
@@ -3,6 +3,7 @@
 #include "SAssignedClientsWidget.h"
 
 #include "AssignedClientsModel.h"
+#include "Replication/Client/UnifiedClientViewExtensions.h"
 #include "Widgets/Client/ClientInfoHelpers.h"
 #include "Widgets/Client/SHorizontalClientList.h"
 
@@ -14,8 +15,8 @@ namespace UE::MultiUserClient::Replication::MultiStreamColumns
 {
 	void SAssignedClientsWidget::Construct(
 		const FArguments& InArgs,
-		const TSharedRef<IConcertClient>& InConcertClient,
-		FAssignedClientsModel& InModel
+		FAssignedClientsModel& InModel,
+		const FUnifiedClientView& InClientView
 		)
 	{
 		Model = &InModel;
@@ -24,8 +25,8 @@ namespace UE::MultiUserClient::Replication::MultiStreamColumns
 		ChildSlot
 		[
 			SAssignNew(ClientList, ConcertSharedSlate::SHorizontalClientList)
-			.GetClientParenthesesContent(ConcertClientSharedSlate::MakeGetLocalClientParenthesesContent(InConcertClient))
-			.GetClientInfo(ConcertClientSharedSlate::MakeClientInfoGetter(InConcertClient))
+			.GetClientParenthesesContent(MakeLocalAndOfflineParenthesesContentGetter(InClientView))
+			.GetClientInfo(MakeOnlineThenOfflineClientInfoGetter(InClientView))
 			.Font(FAppStyle::GetFontStyle("PropertyWindow.NormalFont"))
 			.HighlightText(InArgs._HighlightText)
 			.ListToolTipText(LOCTEXT("Clients.ToolTip", "These clients will replicate their assigned properties when replication is active.\nYou can pause & resume replication at the beginnig of this row."))
diff --git a/Engine/Plugins/Developer/Concert/ConcertApp/MultiUserClient/Source/MultiUserClient/Private/Widgets/ActiveSession/Replication/Client/Multi/Columns/AssignedClients/SAssignedClientsWidget.h b/Engine/Plugins/Developer/Concert/ConcertApp/MultiUserClient/Source/MultiUserClient/Private/Widgets/ActiveSession/Replication/Client/Multi/Columns/AssignedClients/SAssignedClientsWidget.h
--- a/Engine/Plugins/Developer/Concert/ConcertApp/MultiUserClient/Source/MultiUserClient/Private/Widgets/ActiveSession/Replication/Client/Multi/Columns/AssignedClients/SAssignedClientsWidget.h
+++ b/Engine/Plugins/Developer/Concert/ConcertApp/MultiUserClient/Source/MultiUserClient/Private/Widgets/ActiveSession/Replication/Client/Multi/Columns/AssignedClients/SAssignedClientsWidget.h
@@ -13,7 +13,11 @@ namespace UE::ConcertSharedSlate
 	class SHorizontalClientList;
 }
 
-namespace UE::MultiUserClient::Replication { class FReassignObjectPropertiesLogic; }
+namespace UE::MultiUserClient::Replication
+{
+	class FReassignObjectPropertiesLogic;
+	class FUnifiedClientView;
+}
 
 namespace UE::MultiUserClient::Replication::MultiStreamColumns
 {
@@ -32,8 +36,8 @@ namespace UE::MultiUserClient::Replication::MultiStreamColumns
 
 		void Construct(
 			const FArguments& InArgs,
-			const TSharedRef<IConcertClient>& InConcertClient,
-			FAssignedClientsModel& InModel UE_LIFETIMEBOUND
+			FAssignedClientsModel& InModel UE_LIFETIMEBOUND,
+			const FUnifiedClientView& InClientView UE_LIFETIMEBOUND
 			);
 	
 		virtual ~SAssignedClientsWidget() override;
diff --git a/Engine/Plugins/Developer/Concert/ConcertApp/MultiUserClient/Source/MultiUserClient/Private/Widgets/ActiveSession/Replication/Client/Multi/Columns/MultiStreamColumns.h b/Engine/Plugins/Developer/Concert/ConcertApp/MultiUserClient/Source/MultiUserClient/Private/Widgets/ActiveSession/Replication/Client/Multi/Columns/MultiStreamColumns.h
--- a/Engine/Plugins/Developer/Concert/ConcertApp/MultiUserClient/Source/MultiUserClient/Private/Widgets/ActiveSession/Replication/Client/Multi/Columns/MultiStreamColumns.h
+++ b/Engine/Plugins/Developer/Concert/ConcertApp/MultiUserClient/Source/MultiUserClient/Private/Widgets/ActiveSession/Replication/Client/Multi/Columns/MultiStreamColumns.h
@@ -16,6 +16,7 @@ namespace UE::MultiUserClient::Replication
 	class FMuteChangeTracker;
 	class FReassignObjectPropertiesLogic;
 	class FOnlineClientManager;
+	class FUnifiedClientView;
 }
 
 namespace UE::ConcertSharedSlate
@@ -44,6 +45,7 @@ namespace UE::MultiUserClient::Replication::MultiStreamColumns
 	 * Mutes and unmutes the object and its child objects.
 	 * 
 	 * @param MuteChangeTracker Tells us the mute state and changes it.
+	 * @param ColumnsSortPriority The order relative to the other columns
 	 * @return A checkbox with pause and unpause icons.
 	 */
 	ConcertSharedSlate::FObjectColumnEntry MuteToggleColumn(
@@ -57,14 +59,14 @@ namespace UE::MultiUserClient::Replication::MultiStreamColumns
 	 * @param ConcertClient Used to look up client names
 	 * @param MultiStreamModelAttribute Used for requesting resorting the column
 	 * @param ObjectHierarchy Used to display clients when a child object is replicated
-	 * @param ClientManager Used to access all clients for display in the combo box drop-down
+	 * @param ClientView Used to access all clients' stream content for display in the column
 	 * @param ColumnsSortPriority The order relative to the other columns
 	 */
 	ConcertSharedSlate::FObjectColumnEntry AssignedClientsColumn(
-		TSharedRef<IConcertClient> ConcertClient,
+		const TSharedRef<IConcertClient>& ConcertClient,
 		TAttribute<TSharedPtr<ConcertSharedSlate::IMultiReplicationStreamEditor>> MultiStreamModelAttribute,
 		const ConcertSharedSlate::IObjectHierarchyModel& ObjectHierarchy UE_LIFETIMEBOUND,
-		FOnlineClientManager& ClientManager UE_LIFETIMEBOUND,
+		FUnifiedClientView& ClientView UE_LIFETIMEBOUND,
 		int32 ColumnsSortPriority = static_cast<int32>(EColumnSortOrder::ReassignOwnership)
 		);
 	
@@ -72,16 +74,14 @@ namespace UE::MultiUserClient::Replication::MultiStreamColumns
 	 * Creates a property column which assigns the property to the stream selected in the combo box.
 	 * 
 	 * @param MultiStreamEditor Used to determine the selected objects.
-	 * @param ConcertClient Used to look up client names
-	 * @param ClientManager Used to map streams back to client display info
+	 * @param ClientView Used to obtain the clients that can be assigned to and display them
 	 * @param ColumnsSortPriority The order relative to the other columns
 	 * 
 	 * @return A column that spawns a combo box for assigning properties
 	 */
 	ConcertSharedSlate::FPropertyColumnEntry AssignPropertyColumn(
 		TAttribute<TSharedPtr<ConcertSharedSlate::IMultiReplicationStreamEditor>> MultiStreamEditor,
-		TSharedRef<IConcertClient> ConcertClient,
-		FOnlineClientManager& ClientManager UE_LIFETIMEBOUND,
+		FUnifiedClientView& ClientView UE_LIFETIMEBOUND,
 		const int32 ColumnsSortPriority = static_cast<int32>(EColumnSortOrder::AssignPropertyColumn)
 		);
 }
diff --git a/Engine/Plugins/Developer/Concert/ConcertApp/MultiUserClient/Source/MultiUserClient/Private/Widgets/ActiveSession/Replication/Client/Multi/MultiStreamModel.cpp b/Engine/Plugins/Developer/Concert/ConcertApp/MultiUserClient/Source/MultiUserClient/Private/Widgets/ActiveSession/Replication/Client/Multi/MultiStreamModel.cpp
--- a/Engine/Plugins/Developer/Concert/ConcertApp/MultiUserClient/Source/MultiUserClient/Private/Widgets/ActiveSession/Replication/Client/Multi/MultiStreamModel.cpp
+++ b/Engine/Plugins/Developer/Concert/ConcertApp/MultiUserClient/Source/MultiUserClient/Private/Widgets/ActiveSession/Replication/Client/Multi/MultiStreamModel.cpp
@@ -2,6 +2,7 @@
 
 #include "MultiStreamModel.h"
 
+#include "Replication/Client/Offline/OfflineClientManager.h"
 #include "Replication/Client/Online/OnlineClient.h"
 #include "Replication/Client/Online/OnlineClientManager.h"
 #include "Replication/Editor/Model/IEditableReplicationStreamModel.h"
@@ -9,17 +10,33 @@
 
 namespace UE::MultiUserClient::Replication
 {
-	FMultiStreamModel::FMultiStreamModel(IOnlineClientSelectionModel& InOnlineClientSelectionModel, FOnlineClientManager& InClientManager)
+	FMultiStreamModel::FMultiStreamModel(
+		IOnlineClientSelectionModel& InOnlineClientSelectionModel,
+		IOfflineClientSelectionModel& InOfflineClientSelectionModel,
+		FOnlineClientManager& InOnlineClientManager,
+		FOfflineClientManager& InOfflineClientManager
+		)
 		: OnlineClientSelectionModel(InOnlineClientSelectionModel)
-		, ClientManager(InClientManager)
+		, OfflineClientSelectionModel(InOfflineClientSelectionModel)
+		, OnlineClientManager(InOnlineClientManager)
+		, OfflineClientManager(InOfflineClientManager)
 	{
-		OnlineClientSelectionModel.OnSelectionChanged().AddRaw(this, &FMultiStreamModel::RebuildStreamsSets);
-		RebuildStreamsSets();
+		OnlineClientSelectionModel.OnSelectionChanged().AddRaw(this, &FMultiStreamModel::RebuildOnlineClients);
+		OfflineClientSelectionModel.OnSelectionChanged().AddRaw(this, &FMultiStreamModel::RebuildOfflineClients);
+		
+		RebuildOnlineClients();
+		RebuildOfflineClients();
 	}
 
-	void FMultiStreamModel::ForEachClient(TFunctionRef<EBreakBehavior(const FOnlineClient*)> ProcessClient) const
+	FMultiStreamModel::~FMultiStreamModel()
 	{
-		for (const FOnlineClient* WritableClient : CachedWritableClients)
+		UnsubscribeFromOnlineClients();
+		UnsubscribeFromOfflineClients();
+	}
+
+	void FMultiStreamModel::ForEachDisplayedOnlineClient(TFunctionRef<EBreakBehavior(const FOnlineClient*)> ProcessClient) const
+	{
+		for (const FOnlineClient* WritableClient : CachedOnlineClients)
 		{
 			if (ProcessClient(WritableClient) == EBreakBehavior::Break)
 			{
@@ -28,45 +45,99 @@ namespace UE::MultiUserClient::Replication
 		}
 	}
 
+	TSet<TSharedRef<ConcertSharedSlate::IReplicationStreamModel>> FMultiStreamModel::GetReadOnlyStreams() const
+	{
+		TSet<TSharedRef<ConcertSharedSlate::IReplicationStreamModel>> Result;
+		Algo::Transform(CachedOfflineClients, Result, [](const FOfflineClient* Client){ return Client->GetStreamModel(); });
+		return Result;
+	}
+
 	TSet<TSharedRef<ConcertSharedSlate::IEditableReplicationStreamModel>> FMultiStreamModel::GetEditableStreams() const
 	{
 		TSet<TSharedRef<ConcertSharedSlate::IEditableReplicationStreamModel>> Result;
-		Algo::Transform(CachedWritableClients, Result, [](const FOnlineClient* Client){ return Client->GetClientEditModel(); });
+		Algo::Transform(CachedOnlineClients, Result, [](const FOnlineClient* Client){ return Client->GetClientEditModel(); });
 		return Result;
 	}
 
-	void FMultiStreamModel::RebuildStreamsSets()
+	void FMultiStreamModel::RebuildOnlineClients()
 	{
 		// It is not safe to iterate through our cached client array because it may contain stale clients that were just removed.
-		ClientManager.ForEachClient([this](FOnlineClient& Client)
+		UnsubscribeFromOnlineClients();
+		
+		TSet<const FOnlineClient*> OnlineClients;
+		OnlineClientSelectionModel.ForEachItem([this, &OnlineClients](FOnlineClient& Client)
 		{
-			Client.OnModelChanged().RemoveAll(this);
+			const TSharedRef<ConcertSharedSlate::IEditableReplicationStreamModel> Stream = Client.GetClientEditModel();
+			Client.OnModelChanged().AddRaw(
+				this, &FMultiStreamModel::HandleOnlineClientStreamExternallyChanged, Stream.ToWeakPtr()
+				);
+			
+			OnlineClients.Add(&Client);
 			return EBreakBehavior::Continue;
 		});
+
+		const bool bClientsStayedSame = CachedOnlineClients.Num() == OnlineClients.Num() && CachedOnlineClients.Includes(OnlineClients);
+		if (!bClientsStayedSame)
+		{
+			CachedOnlineClients = MoveTemp(OnlineClients);
+			OnStreamSetChangedDelegate.Broadcast();
+		}
+	}
+
+	void FMultiStreamModel::RebuildOfflineClients()
+	{
+		// It is not safe to iterate through our cached client array because it may contain stale clients that were just removed.
+		UnsubscribeFromOfflineClients();
 		
-		TSet<const FOnlineClient*> WritableStreams;
-		OnlineClientSelectionModel.ForEachItem([this, &WritableStreams](FOnlineClient& Client)
+		TSet<const FOfflineClient*> OfflineClients;
+		OfflineClientSelectionModel.ForEachItem([this, &OfflineClients](FOfflineClient& Client)
 		{
-			const TSharedRef<ConcertSharedSlate::IEditableReplicationStreamModel> Stream = Client.GetClientEditModel();
-			Client.OnModelChanged().AddRaw(this, &FMultiStreamModel::OnStreamExternallyChanged, Stream.ToWeakPtr());
-			
-			WritableStreams.Add(&Client);
+			Client.OnStreamPredictionChanged().AddRaw(
+				this, &FMultiStreamModel::HandleOfflineClientStreamExternallyChanged, Client.GetStreamModel().ToWeakPtr()
+				);
+			OfflineClients.Add(&Client);
 			return EBreakBehavior::Continue;
 		});
 
-		const bool bWritableStayedSame = CachedWritableClients.Num() == WritableStreams.Num() && CachedWritableClients.Includes(WritableStreams);
-		if (!bWritableStayedSame)
+		const bool bClientsStayedSame = CachedOfflineClients.Num() == OfflineClients.Num() && CachedOfflineClients.Includes(OfflineClients);
+		if (!bClientsStayedSame)
 		{
-			CachedWritableClients = MoveTemp(WritableStreams);
+			CachedOfflineClients = MoveTemp(OfflineClients);
 			OnStreamSetChangedDelegate.Broadcast();
 		}
 	}
 
-	void FMultiStreamModel::OnStreamExternallyChanged(TWeakPtr<ConcertSharedSlate::IEditableReplicationStreamModel> ChangedStream)
+	void FMultiStreamModel::UnsubscribeFromOnlineClients() const
+	{
+		OnlineClientManager.ForEachClient([this](FOnlineClient& Client)
+		{
+			Client.OnModelChanged().RemoveAll(this);
+			return EBreakBehavior::Continue;
+		});
+	}
+
+	void FMultiStreamModel::UnsubscribeFromOfflineClients() const
+	{
+		OfflineClientManager.ForEachClient([this](FOfflineClient& Client)
+		{
+			Client.OnStreamPredictionChanged().RemoveAll(this);
+			return EBreakBehavior::Continue;
+		});
+	}
+
+	void FMultiStreamModel::HandleOnlineClientStreamExternallyChanged(TWeakPtr<ConcertSharedSlate::IEditableReplicationStreamModel> ChangedStream)
 	{
 		if (const TSharedPtr<ConcertSharedSlate::IEditableReplicationStreamModel> ChangedStreamPin = ChangedStream.Pin())
 		{
-			OnReadOnlyStreamChangedDelegate.Broadcast(ChangedStreamPin.ToSharedRef());
+			OnStreamsExternallyChanged.Broadcast(ChangedStreamPin.ToSharedRef());
+		}
+	}
+
+	void FMultiStreamModel::HandleOfflineClientStreamExternallyChanged(TWeakPtr<ConcertSharedSlate::IReplicationStreamModel> ChangedStream)
+	{
+		if (const TSharedPtr<ConcertSharedSlate::IReplicationStreamModel> ChangedStreamPin = ChangedStream.Pin())
+		{
+			OnStreamsExternallyChanged.Broadcast(ChangedStreamPin.ToSharedRef());
 		}
 	}
 }
diff --git a/Engine/Plugins/Developer/Concert/ConcertApp/MultiUserClient/Source/MultiUserClient/Private/Widgets/ActiveSession/Replication/Client/Multi/MultiStreamModel.h b/Engine/Plugins/Developer/Concert/ConcertApp/MultiUserClient/Source/MultiUserClient/Private/Widgets/ActiveSession/Replication/Client/Multi/MultiStreamModel.h
--- a/Engine/Plugins/Developer/Concert/ConcertApp/MultiUserClient/Source/MultiUserClient/Private/Widgets/ActiveSession/Replication/Client/Multi/MultiStreamModel.h
+++ b/Engine/Plugins/Developer/Concert/ConcertApp/MultiUserClient/Source/MultiUserClient/Private/Widgets/ActiveSession/Replication/Client/Multi/MultiStreamModel.h
@@ -10,6 +10,11 @@
 #include "Templates/Function.h"
 #include "Templates/SharedPointer.h"
 
+namespace UE::MultiUserClient::Replication
+{
+	class FOfflineClientManager;
+}
+
 enum class EBreakBehavior : uint8;
 
 namespace UE::MultiUserClient::Replication
@@ -17,19 +22,31 @@ namespace UE::MultiUserClient::Replication
 	class FOnlineClient;
 	class FOnlineClientManager;
 
-	/** Checks whether clients accept remote changes and categorizes them into read-only and writable. */
+	/**
+	 * Decides the client streams that are displayed in the multi-view.
+	 * 
+	 * This model, in turn, uses IOnlineClientSelectionModel and IOfflineClientSelectionModel to decide which online and offline clients are to be
+	 * displayed.
+	 */
 	class FMultiStreamModel : public ConcertSharedSlate::IEditableMultiReplicationStreamModel
 	{
 	public:
 		
-		FMultiStreamModel(IOnlineClientSelectionModel& InOnlineClientSelectionModel, FOnlineClientManager& InClientManager);
-		
-		void ForEachClient(TFunctionRef<EBreakBehavior(const FOnlineClient*)> ProcessClient) const;
+		FMultiStreamModel(
+			IOnlineClientSelectionModel& InOnlineClientSelectionModel UE_LIFETIMEBOUND,
+			IOfflineClientSelectionModel& InOfflineClientSelectionModel UE_LIFETIMEBOUND,
+			FOnlineClientManager& InOnlineClientManager UE_LIFETIMEBOUND,
+			FOfflineClientManager& InOfflineClientManager UE_LIFETIMEBOUND
+			);
+		~FMultiStreamModel();
+
+		/** Enumerates every displayed online client. */
+		void ForEachDisplayedOnlineClient(TFunctionRef<EBreakBehavior(const FOnlineClient*)> ProcessClient) const;
 
 		//~ Begin IEditableMultiReplicationStreamModel Interface
-		virtual TSet<TSharedRef<ConcertSharedSlate::IReplicationStreamModel>> GetReadOnlyStreams() const override { return {}; }
+		virtual TSet<TSharedRef<ConcertSharedSlate::IReplicationStreamModel>> GetReadOnlyStreams() const override;
 		virtual TSet<TSharedRef<ConcertSharedSlate::IEditableReplicationStreamModel>> GetEditableStreams() const override;
-		virtual FOnStreamExternallyChanged& OnStreamExternallyChanged() override { return OnReadOnlyStreamChangedDelegate; }
+		virtual FOnStreamExternallyChanged& OnStreamExternallyChanged() override { return OnStreamsExternallyChanged; }
 		virtual FOnStreamSetChanged& OnStreamSetChanged() override { return OnStreamSetChangedDelegate; }
 		//~ End IEditableMultiReplicationStreamModel Interface
 
@@ -37,18 +54,29 @@ namespace UE::MultiUserClient::Replication
 		
 		/** Gets all online clients that are supposed to be displayed. */
 		IOnlineClientSelectionModel& OnlineClientSelectionModel;
-		/** Used to obtain a list of clients for unsubscribing. */
-		FOnlineClientManager& ClientManager;
+		/** Gets all offline clients that are supposed to be displayed. */
+		IOfflineClientSelectionModel& OfflineClientSelectionModel;
+		
+		/** Used to clean up subscriptions when client list changes. */
+		FOnlineClientManager& OnlineClientManager;
+		/** Used to clean up subscriptions when client list changes. */
+		FOfflineClientManager& OfflineClientManager;
 
-		TSet<const FOnlineClient*> CachedWritableClients;
+		TSet<const FOnlineClient*> CachedOnlineClients;
+		TSet<const FOfflineClient*> CachedOfflineClients;
 
-		FOnStreamExternallyChanged OnReadOnlyStreamChangedDelegate;
+		FOnStreamExternallyChanged OnStreamsExternallyChanged;
 		FOnStreamSetChanged OnStreamSetChangedDelegate;
 		
-		void RebuildStreamsSets();
+		void RebuildOnlineClients();
+		void RebuildOfflineClients();
+		
+		void UnsubscribeFromOnlineClients() const;
+		void UnsubscribeFromOfflineClients() const;
 
 		/** Handle read-only streams changing */
-		void OnStreamExternallyChanged(TWeakPtr<ConcertSharedSlate::IEditableReplicationStreamModel> ChangedStream);
+		void HandleOnlineClientStreamExternallyChanged(TWeakPtr<ConcertSharedSlate::IEditableReplicationStreamModel> ChangedStream);
+		void HandleOfflineClientStreamExternallyChanged(TWeakPtr<ConcertSharedSlate::IReplicationStreamModel> ChangedStream);
 	};
 }
 
diff --git a/Engine/Plugins/Developer/Concert/ConcertApp/MultiUserClient/Source/MultiUserClient/Private/Widgets/ActiveSession/Replication/Client/Multi/SAllClientsView.cpp b/Engine/Plugins/Developer/Concert/ConcertApp/MultiUserClient/Source/MultiUserClient/Private/Widgets/ActiveSession/Replication/Client/Multi/SAllClientsView.cpp
--- a/Engine/Plugins/Developer/Concert/ConcertApp/MultiUserClient/Source/MultiUserClient/Private/Widgets/ActiveSession/Replication/Client/Multi/SAllClientsView.cpp
+++ b/Engine/Plugins/Developer/Concert/ConcertApp/MultiUserClient/Source/MultiUserClient/Private/Widgets/ActiveSession/Replication/Client/Multi/SAllClientsView.cpp
@@ -5,32 +5,17 @@
 #include "Selection/AllOnlineClientsSelectionModel.h"
 #include "Replication/MultiUserReplicationManager.h"
 #include "SMultiClientView.h"
-#include "Replication/Client/Online/OnlineClientManager.h"
-
-#include "Algo/Transform.h"
 
 namespace UE::MultiUserClient::Replication
 {
 	void SAllClientsView::Construct(const FArguments&, TSharedRef<IConcertClient> InConcertClient, FMultiUserReplicationManager& InMultiUserReplicationManager)
 	{
-		ClientManager = InMultiUserReplicationManager.GetOnlineClientManager();
-		check(ClientManager);
-		AllClientsModel = MakeUnique<FAllOnlineClientsSelectionModel>(*ClientManager);
+		AllOnlineClientsModel = MakeUnique<FAllOnlineClientsSelectionModel>(*InMultiUserReplicationManager.GetOnlineClientManager());
+		AllOfflineClientsModel = MakeUnique<FAllOfflineClientsSelectionModel>(*InMultiUserReplicationManager.GetOfflineClientManager());
 		
 		ChildSlot
 		[
-			SNew(SMultiClientView, InConcertClient, InMultiUserReplicationManager, *AllClientsModel)
+			SNew(SMultiClientView, InConcertClient, InMultiUserReplicationManager, *AllOnlineClientsModel, *AllOfflineClientsModel)
 		];
 	}
-
-	TSet<const FOnlineClient*> SAllClientsView::GetAllClients() const
-	{
-		TSet<const FOnlineClient*> Result;
-		Algo::Transform(ClientManager->GetRemoteClients(), Result, [](const TNonNullPtr<FRemoteClient>& Client)
-		{
-			return Client.Get();
-		});
-		Result.Add(&ClientManager->GetLocalClient());
-		return Result;
-	}
 }
diff --git a/Engine/Plugins/Developer/Concert/ConcertApp/MultiUserClient/Source/MultiUserClient/Private/Widgets/ActiveSession/Replication/Client/Multi/SAllClientsView.h b/Engine/Plugins/Developer/Concert/ConcertApp/MultiUserClient/Source/MultiUserClient/Private/Widgets/ActiveSession/Replication/Client/Multi/SAllClientsView.h
--- a/Engine/Plugins/Developer/Concert/ConcertApp/MultiUserClient/Source/MultiUserClient/Private/Widgets/ActiveSession/Replication/Client/Multi/SAllClientsView.h
+++ b/Engine/Plugins/Developer/Concert/ConcertApp/MultiUserClient/Source/MultiUserClient/Private/Widgets/ActiveSession/Replication/Client/Multi/SAllClientsView.h
@@ -2,6 +2,7 @@
 
 #pragma once
 
+#include "Selection/AllOfflineClientsSelectionModel.h"
 #include "Selection/AllOnlineClientsSelectionModel.h"
 
 #include "HAL/Platform.h"
@@ -16,7 +17,7 @@ namespace UE::MultiUserClient::Replication
 	class FOnlineClient;
 	class FOnlineClientManager;
 
-	/** Leverages SMultiClientView to display all replication clients. */
+	/** Leverages SMultiClientView to display all online and offline clients. */
 	class SAllClientsView : public SCompoundWidget
 	{
 		SLATE_BEGIN_ARGS(SAllClientsView){}
@@ -26,13 +27,9 @@ namespace UE::MultiUserClient::Replication
 
 	private:
 
-		/** Used to get all the replication clients and listen for client changes. */
-		FOnlineClientManager* ClientManager = nullptr;
-
-		/** Keeps the SMultiClientView updated of any changes to clients (e.g. disconnects, etc.) */
-		TUniquePtr<FAllOnlineClientsSelectionModel> AllClientsModel;
-
-		/** Gets all the clients to display */
-		TSet<const FOnlineClient*> GetAllClients() const;
+		/** Keeps track of all online clients. */
+		TUniquePtr<FAllOnlineClientsSelectionModel> AllOnlineClientsModel;
+		/** Keeps track of all offline clients. */
+		TUniquePtr<FAllOfflineClientsSelectionModel> AllOfflineClientsModel;
 	};
 }
diff --git a/Engine/Plugins/Developer/Concert/ConcertApp/MultiUserClient/Source/MultiUserClient/Private/Widgets/ActiveSession/Replication/Client/Multi/SMultiClientView.cpp b/Engine/Plugins/Developer/Concert/ConcertApp/MultiUserClient/Source/MultiUserClient/Private/Widgets/ActiveSession/Replication/Client/Multi/SMultiClientView.cpp
--- a/Engine/Plugins/Developer/Concert/ConcertApp/MultiUserClient/Source/MultiUserClient/Private/Widgets/ActiveSession/Replication/Client/Multi/SMultiClientView.cpp
+++ b/Engine/Plugins/Developer/Concert/ConcertApp/MultiUserClient/Source/MultiUserClient/Private/Widgets/ActiveSession/Replication/Client/Multi/SMultiClientView.cpp
@@ -32,18 +32,23 @@ namespace UE::MultiUserClient::Replication
 {
 	void SMultiClientView::Construct(
 		const FArguments&,
-		TSharedRef<IConcertClient> InConcertClient,
+		const TSharedRef<IConcertClient>& InConcertClient,
 		FMultiUserReplicationManager& InMultiUserReplicationManager,
-		IOnlineClientSelectionModel& InOnlineClientSelectionModel
+		IOnlineClientSelectionModel& InOnlineClientSelectionModel,
+		IOfflineClientSelectionModel& InOfflineClientSelectionModel
 		)
 	{
-		ClientManager = InMultiUserReplicationManager.GetOnlineClientManager();
+		ConcertClient = InConcertClient;
 		UserSelectedProperties = InMultiUserReplicationManager.GetUserPropertySelector();
-		StreamModel = MakeShared<FMultiStreamModel>(InOnlineClientSelectionModel, *ClientManager);
+		OnlineClientManager = InMultiUserReplicationManager.GetOnlineClientManager();
+		OfflineClientManager = InMultiUserReplicationManager.GetOfflineClientManager();
 		OnlineClientSelectionModel = &InOnlineClientSelectionModel;
-		ConcertClient = MoveTemp(InConcertClient);
+		OfflineClientSelectionModel = &InOfflineClientSelectionModel;
+		StreamModel = MakeShared<FMultiStreamModel>(
+			InOnlineClientSelectionModel, InOfflineClientSelectionModel, *OnlineClientManager, *OfflineClientManager
+			);
 		
-		ClientManager->OnRemoteClientsChanged().AddSP(this, &SMultiClientView::RebuildClientSubscriptions);
+		OnlineClientManager->OnRemoteClientsChanged().AddSP(this, &SMultiClientView::RebuildClientSubscriptions);
 		OnlineClientSelectionModel->OnSelectionChanged().AddSP(this, &SMultiClientView::RebuildClientSubscriptions);
 		UserSelectedProperties->OnPropertySelectionChanged().AddRaw(this, &SMultiClientView::RefreshUI);
 
@@ -60,10 +65,10 @@ namespace UE::MultiUserClient::Replication
 			]
 		];
 		
-		SReplicationStatus::AppendReplicationStatus(*Content, ClientManager->GetAuthorityCache(),
+		SReplicationStatus::AppendReplicationStatus(*Content, OnlineClientManager->GetAuthorityCache(),
 			SReplicationStatus::FArguments()
 			.ReplicatableClients(this, &SMultiClientView::GetReplicatableClientIds)
-			.ForEachObjectInStream(this, &SMultiClientView::EnumerateObjectsInStreams)
+			.ForEachObjectInStream(this, &SMultiClientView::EnumerateReplicatedObjectsInStreams)
 			);
 
 		RebuildClientSubscriptions();
@@ -77,12 +82,15 @@ namespace UE::MultiUserClient::Replication
 
 	SMultiClientView::~SMultiClientView()
 	{
-		ClientManager->OnRemoteClientsChanged().RemoveAll(this);
+		OnlineClientManager->OnRemoteClientsChanged().RemoveAll(this);
 		CleanClientSubscriptions();
 		UserSelectedProperties->OnPropertySelectionChanged().RemoveAll(this);
 	}
 
-	TSharedRef<SWidget> SMultiClientView::CreateEditorContent(const TSharedRef<IConcertClient>& InConcertClient, FMultiUserReplicationManager& InMultiUserReplicationManager)
+	TSharedRef<SWidget> SMultiClientView::CreateEditorContent(
+		const TSharedRef<IConcertClient>& InConcertClient,
+		FMultiUserReplicationManager& InMultiUserReplicationManager
+		)
 	{
 		using namespace UE::ConcertSharedSlate;
 
@@ -99,7 +107,7 @@ namespace UE::MultiUserClient::Replication
 		   });
 		FGetAutoAssignTarget GetAutoAssignTargetDelegate = FGetAutoAssignTarget::CreateLambda([this](TConstArrayView<UObject*>)
 		{
-			const TSharedRef<IEditableReplicationStreamModel>& LocalStream = ClientManager->GetLocalClient().GetClientEditModel();
+			const TSharedRef<IEditableReplicationStreamModel>& LocalStream = OnlineClientManager->GetLocalClient().GetClientEditModel();
 			return StreamModel->GetEditableStreams().Contains(LocalStream) ? LocalStream.ToSharedPtr() : nullptr;
 		});
 		
@@ -108,7 +116,9 @@ namespace UE::MultiUserClient::Replication
 			.PropertyColumns =
 			{
 				ReplicationColumns::Property::LabelColumn(),
-				MultiStreamColumns::AssignPropertyColumn(MultiStreamEditorAttribute, InConcertClient, *ClientManager)
+				MultiStreamColumns::AssignPropertyColumn(
+					MultiStreamEditorAttribute, *InMultiUserReplicationManager.GetUnifiedClientView()
+					)
 			},
 			.CreateCategoryRow = CreateDefaultCategoryGenerator(NameModel),
 		};
@@ -146,7 +156,9 @@ namespace UE::MultiUserClient::Replication
 			.ObjectColumns =
 			{
 				MultiStreamColumns::MuteToggleColumn(MuteManager.GetChangeTracker()),
-				MultiStreamColumns::AssignedClientsColumn(InConcertClient, MultiStreamEditorAttribute, *ObjectHierarchy, *ClientManager)
+				MultiStreamColumns::AssignedClientsColumn(
+					InConcertClient, MultiStreamEditorAttribute, *ObjectHierarchy, *InMultiUserReplicationManager.GetUnifiedClientView()
+					)
 			},
 			.ShouldDisplayObjectDelegate = FShouldDisplayObject::CreateSP(this, &SMultiClientView::ShouldDisplayObject),
 		};
@@ -171,7 +183,7 @@ namespace UE::MultiUserClient::Replication
 	TSet<FGuid> SMultiClientView::GetReplicatableClientIds() const
 	{
 		TSet<FGuid> ClientIds;
-		StreamModel->ForEachClient([&ClientIds](const FOnlineClient* Client)
+		StreamModel->ForEachDisplayedOnlineClient([&ClientIds](const FOnlineClient* Client)
 		{
 			ClientIds.Add(Client->GetEndpointId());
 			return EBreakBehavior::Continue;
@@ -179,9 +191,9 @@ namespace UE::MultiUserClient::Replication
 		return ClientIds;
 	}
 
-	void SMultiClientView::EnumerateObjectsInStreams(TFunctionRef<void(const FSoftObjectPath&)> Consumer) const
+	void SMultiClientView::EnumerateReplicatedObjectsInStreams(TFunctionRef<void(const FSoftObjectPath&)> Consumer) const
 	{
-		StreamModel->ForEachClient([&Consumer](const FOnlineClient* Client)
+		StreamModel->ForEachDisplayedOnlineClient([&Consumer](const FOnlineClient* Client)
 		{
 			Client->GetClientEditModel()->ForEachReplicatedObject([&Consumer](const FSoftObjectPath& Object)
 			{
@@ -208,7 +220,7 @@ namespace UE::MultiUserClient::Replication
 
 	void SMultiClientView::CleanClientSubscriptions() const
 	{
-		ClientManager->ForEachClient([this](FOnlineClient& Client)
+		OnlineClientManager->ForEachClient([this](FOnlineClient& Client)
 		{
 			Client.OnModelChanged().RemoveAll(this);
 			Client.OnHierarchyNeedsRefresh().RemoveAll(this);
@@ -223,7 +235,7 @@ namespace UE::MultiUserClient::Replication
 
 	void SMultiClientView::ExtendObjectContextMenu(FMenuBuilder& MenuBuilder, TConstArrayView<TSoftObjectPtr<>> ContextObjects) const
 	{
-		ContextMenuUtils::AddFrequencyOptionsIfOneContextObject_MultiClient(MenuBuilder, ContextObjects, *ClientManager);
+		ContextMenuUtils::AddFrequencyOptionsIfOneContextObject_MultiClient(MenuBuilder, ContextObjects, *OnlineClientManager);
 
 		if (ContextObjects.Num() == 1)
 		{
@@ -231,9 +243,9 @@ namespace UE::MultiUserClient::Replication
 				MenuBuilder,
 				ContextObjects[0],
 				*ConcertClient,
-				*ClientManager,
+				*OnlineClientManager,
 				*ObjectHierarchy,
-				ClientManager->GetReassignmentLogic(),
+				OnlineClientManager->GetReassignmentLogic(),
 				*StreamEditor
 				);
 		}
@@ -244,16 +256,16 @@ namespace UE::MultiUserClient::Replication
 		return HideObjectsNotInEditorWorld.ShouldShowObject(Object);
 	}
 
-	void SMultiClientView::OnPreAddObjectsFromComboButton(TArrayView<const ConcertSharedSlate::FSelectableObjectInfo>)
+	void SMultiClientView::OnPreAddObjectsFromComboButton(TArrayView<const ConcertSharedSlate::FSelectableObjectInfo>) const
 	{
 		// When the user adds using the combo button, automatically add discover relevant objects and properties
-		ClientManager->GetLocalClient().GetStreamExtender()
+		OnlineClientManager->GetLocalClient().GetStreamExtender()
 			.SetShouldExtend(true);
 	}
 
-	void SMultiClientView::OnPostAddObjectsFromComboButton(TArrayView<const ConcertSharedSlate::FSelectableObjectInfo>)
+	void SMultiClientView::OnPostAddObjectsFromComboButton(TArrayView<const ConcertSharedSlate::FSelectableObjectInfo>) const
 	{
-		ClientManager->GetLocalClient().GetStreamExtender()
+		OnlineClientManager->GetLocalClient().GetStreamExtender()
 			.SetShouldExtend(false);
 	}
 }
diff --git a/Engine/Plugins/Developer/Concert/ConcertApp/MultiUserClient/Source/MultiUserClient/Private/Widgets/ActiveSession/Replication/Client/Multi/SMultiClientView.h b/Engine/Plugins/Developer/Concert/ConcertApp/MultiUserClient/Source/MultiUserClient/Private/Widgets/ActiveSession/Replication/Client/Multi/SMultiClientView.h
--- a/Engine/Plugins/Developer/Concert/ConcertApp/MultiUserClient/Source/MultiUserClient/Private/Widgets/ActiveSession/Replication/Client/Multi/SMultiClientView.h
+++ b/Engine/Plugins/Developer/Concert/ConcertApp/MultiUserClient/Source/MultiUserClient/Private/Widgets/ActiveSession/Replication/Client/Multi/SMultiClientView.h
@@ -8,6 +8,11 @@
 #include "Widgets/DeclarativeSyntaxSupport.h"
 #include "Widgets/SCompoundWidget.h"
 
+namespace UE::MultiUserClient::Replication
+{
+	class FOfflineClientManager;
+}
+
 namespace UE::MultiUserClient::Replication
 {
 	class FUserPropertySelector;
@@ -41,19 +46,24 @@ namespace UE::MultiUserClient::Replication
 		SLATE_END_ARGS()
 
 		void Construct(
-			const FArguments& InArgs,
-			TSharedRef<IConcertClient> InConcertClient,
+			const FArguments& InArgs, const TSharedRef<IConcertClient>& InConcertClient,
 			FMultiUserReplicationManager& InMultiUserReplicationManager UE_LIFETIMEBOUND,
-			IOnlineClientSelectionModel& InOnlineClientSelectionModel UE_LIFETIMEBOUND
+			IOnlineClientSelectionModel& InOnlineClientSelectionModel UE_LIFETIMEBOUND,
+			IOfflineClientSelectionModel& InOfflineClientSelectionModel UE_LIFETIMEBOUND
 			);
 		virtual ~SMultiClientView() override;
 
 	private:
 
 		TSharedPtr<IConcertClient> ConcertClient;
-		FOnlineClientManager* ClientManager = nullptr;
+		/** Keeps track of the properties that the user has selected to iterate on. */
 		FUserPropertySelector* UserSelectedProperties = nullptr;
+
+		// These are used to know when to refresh the UI.
+		FOnlineClientManager* OnlineClientManager = nullptr;
+		FOfflineClientManager* OfflineClientManager = nullptr;
 		IOnlineClientSelectionModel* OnlineClientSelectionModel = nullptr;
+		IOfflineClientSelectionModel* OfflineClientSelectionModel = nullptr;
 		
 		/** Combines the clients */
 		TSharedPtr<FMultiStreamModel> StreamModel;
@@ -81,7 +91,7 @@ namespace UE::MultiUserClient::Replication
 		/** @return Gets the clients that may be replicating */
 		TSet<FGuid> GetReplicatableClientIds() const;
 		/** Calls Consumer for each object path that is in a stream - independent of whether it is being replicated or not. */
-		void EnumerateObjectsInStreams(TFunctionRef<void(const FSoftObjectPath&)> Consumer) const;
+		void EnumerateReplicatedObjectsInStreams(TFunctionRef<void(const FSoftObjectPath&)> Consumer) const;
 		
 		void RebuildClientSubscriptions();
 		void CleanClientSubscriptions() const;
@@ -92,7 +102,7 @@ namespace UE::MultiUserClient::Replication
 		/** Decides whether the object should be displayed: do not show it if it's not in the editor world. */
 		bool ShouldDisplayObject(const FSoftObjectPath& Object) const;
 		
-		void OnPreAddObjectsFromComboButton(TArrayView<const ConcertSharedSlate::FSelectableObjectInfo>);
-		void OnPostAddObjectsFromComboButton(TArrayView<const ConcertSharedSlate::FSelectableObjectInfo>);
+		void OnPreAddObjectsFromComboButton(TArrayView<const ConcertSharedSlate::FSelectableObjectInfo>) const;
+		void OnPostAddObjectsFromComboButton(TArrayView<const ConcertSharedSlate::FSelectableObjectInfo>) const;
 	};
 }
diff --git a/Engine/Plugins/Developer/Concert/ConcertApp/MultiUserClient/Source/MultiUserClient/Private/Widgets/ActiveSession/Replication/Client/Multi/Selection/AllOfflineClientsSelectionModel.cpp b/Engine/Plugins/Developer/Concert/ConcertApp/MultiUserClient/Source/MultiUserClient/Private/Widgets/ActiveSession/Replication/Client/Multi/Selection/AllOfflineClientsSelectionModel.cpp
new file mode 100644
--- /dev/null
+++ b/Engine/Plugins/Developer/Concert/ConcertApp/MultiUserClient/Source/MultiUserClient/Private/Widgets/ActiveSession/Replication/Client/Multi/Selection/AllOfflineClientsSelectionModel.cpp
@@ -0,0 +1,29 @@
+﻿// Copyright Epic Games, Inc. All Rights Reserved.
+
+#include "AllOfflineClientsSelectionModel.h"
+
+#include "Replication/Client/Offline/OfflineClientManager.h"
+
+namespace UE::MultiUserClient::Replication
+{
+	FAllOfflineClientsSelectionModel::FAllOfflineClientsSelectionModel(FOfflineClientManager& InClientManager)
+		: ClientManager(InClientManager)
+	{
+		ClientManager.OnClientsChanged().AddRaw(this, &FAllOfflineClientsSelectionModel::OnClientsChanged);
+	}
+
+	FAllOfflineClientsSelectionModel::~FAllOfflineClientsSelectionModel()
+	{
+		ClientManager.OnClientsChanged().RemoveAll(this);
+	}
+
+	void FAllOfflineClientsSelectionModel::ForEachItem(TFunctionRef<EBreakBehavior(FOfflineClient&)> ProcessClient) const
+	{
+		ClientManager.ForEachClient([&ProcessClient](FOfflineClient& Client){ return ProcessClient(Client); });
+	}
+
+	void FAllOfflineClientsSelectionModel::OnClientsChanged()
+	{
+		OnSelectionChangedDelegate.Broadcast();
+	}
+}
diff --git a/Engine/Plugins/Developer/Concert/ConcertApp/MultiUserClient/Source/MultiUserClient/Private/Widgets/ActiveSession/Replication/Client/Multi/Selection/AllOfflineClientsSelectionModel.h b/Engine/Plugins/Developer/Concert/ConcertApp/MultiUserClient/Source/MultiUserClient/Private/Widgets/ActiveSession/Replication/Client/Multi/Selection/AllOfflineClientsSelectionModel.h
new file mode 100644
--- /dev/null
+++ b/Engine/Plugins/Developer/Concert/ConcertApp/MultiUserClient/Source/MultiUserClient/Private/Widgets/ActiveSession/Replication/Client/Multi/Selection/AllOfflineClientsSelectionModel.h
@@ -0,0 +1,38 @@
+﻿// Copyright Epic Games, Inc. All Rights Reserved.
+
+#pragma once
+
+#include "ISelectionModel.h"
+#include "SelectionModelFwd.h"
+#include "Replication/Client/Offline/OfflineClient.h"
+
+#include "HAL/Platform.h"
+
+namespace UE::MultiUserClient::Replication
+{
+	class FOfflineClientManager;
+	
+	/** Exposes all clients and detects when clients disconnect. */
+	class FAllOfflineClientsSelectionModel : public IOfflineClientSelectionModel
+	{
+	public:
+		
+		FAllOfflineClientsSelectionModel(FOfflineClientManager& InClientManager UE_LIFETIMEBOUND);
+		virtual ~FAllOfflineClientsSelectionModel() override;
+
+		//~ Begin IOfflineClientSelectionModel Interface
+		virtual void ForEachItem(TFunctionRef<EBreakBehavior(FOfflineClient&)> ProcessClient) const override;
+		virtual FOnSelectionChanged& OnSelectionChanged() override { return OnSelectionChangedDelegate; }
+		//~ End IOfflineClientSelectionModel Interface
+
+	private:
+
+		/** Informs us when the list of clients changes */
+		FOfflineClientManager& ClientManager;
+		
+		/** Called when the clients ForEachSelectedClient enumerates has changed. */
+		FOnSelectionChanged OnSelectionChangedDelegate;
+		
+		void OnClientsChanged();
+	};
+}
diff --git a/Engine/Plugins/Developer/Concert/ConcertApp/MultiUserClient/Source/MultiUserClient/Private/Widgets/ActiveSession/Replication/Client/Multi/Selection/SelectionModelFwd.h b/Engine/Plugins/Developer/Concert/ConcertApp/MultiUserClient/Source/MultiUserClient/Private/Widgets/ActiveSession/Replication/Client/Multi/Selection/SelectionModelFwd.h
--- a/Engine/Plugins/Developer/Concert/ConcertApp/MultiUserClient/Source/MultiUserClient/Private/Widgets/ActiveSession/Replication/Client/Multi/Selection/SelectionModelFwd.h
+++ b/Engine/Plugins/Developer/Concert/ConcertApp/MultiUserClient/Source/MultiUserClient/Private/Widgets/ActiveSession/Replication/Client/Multi/Selection/SelectionModelFwd.h
@@ -6,8 +6,10 @@
 
 namespace UE::MultiUserClient::Replication
 {
+	class FOfflineClient;
 	class FOnlineClient;
 	struct FClientItem;
 	
+	using IOfflineClientSelectionModel = ISelectionModel<FOfflineClient>;
 	using IOnlineClientSelectionModel = ISelectionModel<FOnlineClient>;
 }
diff --git a/Engine/Plugins/Developer/Concert/ConcertUI/ConcertSharedSlate/Source/ConcertSharedSlate/Private/Widgets/Client/SClientName.cpp b/Engine/Plugins/Developer/Concert/ConcertUI/ConcertSharedSlate/Source/ConcertSharedSlate/Private/Widgets/Client/SClientName.cpp
--- a/Engine/Plugins/Developer/Concert/ConcertUI/ConcertSharedSlate/Source/ConcertSharedSlate/Private/Widgets/Client/SClientName.cpp
+++ b/Engine/Plugins/Developer/Concert/ConcertUI/ConcertSharedSlate/Source/ConcertSharedSlate/Private/Widgets/Client/SClientName.cpp
@@ -13,6 +13,7 @@
 namespace UE::ConcertSharedSlate::ParenthesesClientNameContent
 {
 	const FText LocalClient = LOCTEXT("ParenthesesClientNameContent.LocalClient", "You");
+	const FText OfflineClient = LOCTEXT("ParenthesesClientNameContent.OfflineClient", "Offline");
 }
 
 namespace UE::ConcertSharedSlate
diff --git a/Engine/Plugins/Developer/Concert/ConcertUI/ConcertSharedSlate/Source/ConcertSharedSlate/Public/Widgets/Client/SClientName.h b/Engine/Plugins/Developer/Concert/ConcertUI/ConcertSharedSlate/Source/ConcertSharedSlate/Public/Widgets/Client/SClientName.h
--- a/Engine/Plugins/Developer/Concert/ConcertUI/ConcertSharedSlate/Source/ConcertSharedSlate/Public/Widgets/Client/SClientName.h
+++ b/Engine/Plugins/Developer/Concert/ConcertUI/ConcertSharedSlate/Source/ConcertSharedSlate/Public/Widgets/Client/SClientName.h
@@ -14,6 +14,8 @@ namespace UE::ConcertSharedSlate::ParenthesesClientNameContent
 {
 	/** The client corresponds to the local user. "You" is appended to name, e.g. "ClientName(You)". */
 	CONCERTSHAREDSLATE_API extern const FText LocalClient;
+	/** The client corresponds to a client that is not connected to the session. */
+	CONCERTSHAREDSLATE_API extern const FText OfflineClient;
 }
 
 namespace UE::ConcertSharedSlate
