Extract creation of object overlay widgets to the public interface.
Added the FMakeObjectRowOverlayWidget to FCreateViewerParams so the API user can control how the overlay widget is created rather than it being a hard coded bin icon.

diff --git a/Engine/Plugins/Developer/Concert/ConcertApp/MultiUserClient/Source/MultiUserClient/Private/Widgets/ActiveSession/Replication/Client/Multi/SMultiClientView.cpp b/Engine/Plugins/Developer/Concert/ConcertApp/MultiUserClient/Source/MultiUserClient/Private/Widgets/ActiveSession/Replication/Client/Multi/SMultiClientView.cpp
--- a/Engine/Plugins/Developer/Concert/ConcertApp/MultiUserClient/Source/MultiUserClient/Private/Widgets/ActiveSession/Replication/Client/Multi/SMultiClientView.cpp
+++ b/Engine/Plugins/Developer/Concert/ConcertApp/MultiUserClient/Source/MultiUserClient/Private/Widgets/ActiveSession/Replication/Client/Multi/SMultiClientView.cpp
@@ -4,6 +4,7 @@
 
 #include "Selection/ISelectionModel.h"
 #include "MultiStreamModel.h"
+#include "Misc/ObjectUtils.h"
 #include "Replication/ClientReplicationWidgetFactories.h"
 #include "Replication/MultiUserReplicationManager.h"
 #include "Replication/ReplicationWidgetFactories.h"
@@ -22,6 +23,8 @@
 #include "Widgets/ActiveSession/Replication/Client/SPresetComboButton.h"
 #include "Widgets/ActiveSession/Replication/Client/SReplicationStatus.h"
 
+#include "Widgets/Images/SImage.h"
+#include "Widgets/Input/SButton.h"
 #include "Widgets/Layout/SBox.h"
 #include "Widgets/SBoxPanel.h"
 #include "Widgets/Text/STextBlock.h"
@@ -161,6 +164,7 @@ namespace UE::MultiUserClient::Replication
 					)
 			},
 			.ShouldDisplayObjectDelegate = FShouldDisplayObject::CreateSP(this, &SMultiClientView::ShouldDisplayObject),
+			.MakeObjectRowOverlayWidgetDelegate = FMakeObjectRowOverlayWidget::CreateSP(this, &SMultiClientView::MakeObjectRowOverlayWidget)
 		};
 		ViewerParams.RightOfObjectSearchBar.Widget = SNew(SPresetComboButton, *InConcertClient, *InMultiUserReplicationManager.GetPresetManager());
 		StreamEditor = CreateBaseMultiStreamEditor(MoveTemp(Params), MoveTemp(ViewerParams));
@@ -256,6 +260,46 @@ namespace UE::MultiUserClient::Replication
 		return HideObjectsNotInEditorWorld.ShouldShowObject(Object);
 	}
 
+	TSharedRef<SWidget> SMultiClientView::MakeObjectRowOverlayWidget(const ConcertSharedSlate::FReplicatedObjectData& ReplicatedObjectData)
+	{
+		if (!ConcertSyncCore::IsActor(ReplicatedObjectData.GetObjectPath()))
+		{
+			return SNullWidget::NullWidget;
+		}
+		
+		return SNew(SButton)
+			.ButtonStyle( FAppStyle::Get(), "SimpleButton")
+			.OnClicked_Lambda([this, ReplicatedObjectData]()
+			{
+				OnPressBinIcon(ReplicatedObjectData.GetObjectPath());
+				return FReply::Handled();
+			})
+			[
+				SNew(SImage)
+				.Image(FAppStyle::GetBrush("Icons.Delete"))
+			];
+	}
+
+	void SMultiClientView::OnPressBinIcon(const FSoftObjectPath& RootObject) const
+	{
+		ConcertSharedSlate::IEditableReplicationStreamModel& ConsolidatedModel = StreamEditor->GetConsolidatedModel();
+		
+		// We want to delete all children, too, i.e. the ones not listed in the outliner, such as components and other subobjects
+		TArray<FSoftObjectPath> ObjectAndChildren = { RootObject };
+		ConsolidatedModel.ForEachReplicatedObject([&RootObject, &ObjectAndChildren](const FSoftObjectPath& ReplicatedObject)
+		{
+			// ReplicatedObject is a child of a deleted object if the path before it contains one of the deleted objects
+			const bool bIsChildOfDeletedObject = ReplicatedObject.ToString().Contains(RootObject.ToString());
+			if (bIsChildOfDeletedObject)
+			{
+				ObjectAndChildren.Add(ReplicatedObject);
+			}
+			return EBreakBehavior::Continue;
+		});
+		
+		ConsolidatedModel.RemoveObjects(ObjectAndChildren);
+	}
+
 	void SMultiClientView::OnPreAddObjectsFromComboButton(TArrayView<const ConcertSharedSlate::FSelectableObjectInfo>) const
 	{
 		// When the user adds using the combo button, automatically add discover relevant objects and properties
diff --git a/Engine/Plugins/Developer/Concert/ConcertApp/MultiUserClient/Source/MultiUserClient/Private/Widgets/ActiveSession/Replication/Client/Multi/SMultiClientView.h b/Engine/Plugins/Developer/Concert/ConcertApp/MultiUserClient/Source/MultiUserClient/Private/Widgets/ActiveSession/Replication/Client/Multi/SMultiClientView.h
--- a/Engine/Plugins/Developer/Concert/ConcertApp/MultiUserClient/Source/MultiUserClient/Private/Widgets/ActiveSession/Replication/Client/Multi/SMultiClientView.h
+++ b/Engine/Plugins/Developer/Concert/ConcertApp/MultiUserClient/Source/MultiUserClient/Private/Widgets/ActiveSession/Replication/Client/Multi/SMultiClientView.h
@@ -2,6 +2,7 @@
 
 #pragma once
 
+#include "Replication/Editor/Model/Data/ReplicatedObjectData.h"
 #include "Replication/Editor/Model/ObjectSource/IObjectSourceModel.h"
 #include "Replication/Editor/UnrealEditor/HideObjectsNotInWorldLogic.h"
 #include "Selection/SelectionModelFwd.h"
@@ -102,6 +103,11 @@ namespace UE::MultiUserClient::Replication
 		/** Decides whether the object should be displayed: do not show it if it's not in the editor world. */
 		bool ShouldDisplayObject(const FSoftObjectPath& Object) const;
 		
+		/** Creates the widget that overlays actor rows. */
+		TSharedRef<SWidget> MakeObjectRowOverlayWidget(const ConcertSharedSlate::FReplicatedObjectData& ReplicatedObjectData);
+		/** Called when the delete icon over an actor is pressed. Clears the entire hierarchy. */
+		void OnPressBinIcon(const FSoftObjectPath& RootObject) const;
+		
 		void OnPreAddObjectsFromComboButton(TArrayView<const ConcertSharedSlate::FSelectableObjectInfo>) const;
 		void OnPostAddObjectsFromComboButton(TArrayView<const ConcertSharedSlate::FSelectableObjectInfo>) const;
 	};
diff --git a/Engine/Plugins/Developer/Concert/ConcertUI/ConcertSharedSlate/Source/ConcertSharedSlate/Private/Replication/Editor/View/MultiEditor/SMultiReplicationStreamEditor.h b/Engine/Plugins/Developer/Concert/ConcertUI/ConcertSharedSlate/Source/ConcertSharedSlate/Private/Replication/Editor/View/MultiEditor/SMultiReplicationStreamEditor.h
--- a/Engine/Plugins/Developer/Concert/ConcertUI/ConcertSharedSlate/Source/ConcertSharedSlate/Private/Replication/Editor/View/MultiEditor/SMultiReplicationStreamEditor.h
+++ b/Engine/Plugins/Developer/Concert/ConcertUI/ConcertSharedSlate/Source/ConcertSharedSlate/Private/Replication/Editor/View/MultiEditor/SMultiReplicationStreamEditor.h
@@ -41,7 +41,7 @@ namespace UE::ConcertSharedSlate
 
 		virtual IReplicationStreamEditor& GetEditorBase() const override { return *EditorView; }
 		virtual IEditableMultiReplicationStreamModel& GetMultiStreamModel() const override { return *MultiStreamModel; }
-		virtual IReplicationStreamModel& GetConsolidatedModel() const override { return *ConsolidatedModel; }
+		virtual IEditableReplicationStreamModel& GetConsolidatedModel() const override { return *ConsolidatedModel; }
 
 	private:
 
diff --git a/Engine/Plugins/Developer/Concert/ConcertUI/ConcertSharedSlate/Source/ConcertSharedSlate/Private/Replication/Editor/View/ObjectEditor/SBaseReplicationStreamEditor.cpp b/Engine/Plugins/Developer/Concert/ConcertUI/ConcertSharedSlate/Source/ConcertSharedSlate/Private/Replication/Editor/View/ObjectEditor/SBaseReplicationStreamEditor.cpp
--- a/Engine/Plugins/Developer/Concert/ConcertUI/ConcertSharedSlate/Source/ConcertSharedSlate/Private/Replication/Editor/View/ObjectEditor/SBaseReplicationStreamEditor.cpp
+++ b/Engine/Plugins/Developer/Concert/ConcertUI/ConcertSharedSlate/Source/ConcertSharedSlate/Private/Replication/Editor/View/ObjectEditor/SBaseReplicationStreamEditor.cpp
@@ -70,7 +70,13 @@ namespace UE::ConcertSharedSlate
 					]
 				]
 				.RightOfObjectSearchBar() [ InArgs._RightOfObjectSearchBar.Widget ]
-				.GetHoveredRowContent(this, &SBaseReplicationStreamEditor::MakeHoveredRowContent)
+				.GetHoveredRowContent_Lambda(
+					[MakeOverlay = InArgs._MakeObjectRowOverlayWidget, Alignment = InArgs._ObjectOverlayAlignment]
+					(const TSharedPtr<FReplicatedObjectData>& Data)
+					{
+						const TSharedRef<SWidget> OverlayWidget = MakeOverlay.IsBound() ? MakeOverlay.Execute(*Data) : SNullWidget::NullWidget;
+						return FHoverRowContent{ OverlayWidget, Alignment };
+					})
 				.NoOutlinerObjects(LOCTEXT("NoObjects", "Add objects to replicate"))
 		];
 	}
@@ -101,33 +107,6 @@ namespace UE::ConcertSharedSlate
 		return ReplicationViewer->GetSelectedObjects();
 	}
 
-	FHoverRowContent SBaseReplicationStreamEditor::MakeHoveredRowContent(const TSharedPtr<FReplicatedObjectData>& Data) const
-	{
-		const FSoftObjectPath Object = Data->GetObjectPath();
-		
-		if (ConcertSyncCore::IsActor(Object))
-		{
-			const TSharedRef<SWidget> Content =
-				SNew(SButton)
-				.ButtonStyle( FAppStyle::Get(), "SimpleButton")
-				.OnClicked_Lambda([this, WeakData = Data.ToWeakPtr()]()
-				{
-					if (const TSharedPtr<FReplicatedObjectData> DataPin = WeakData.Pin())
-					{
-						OnDeleteObjects({ DataPin });
-					}
-					return FReply::Handled();
-				})
-				[
-					SNew(SImage)
-					.Image(FAppStyle::GetBrush("Icons.Delete"))
-				];
-			return { Content, HAlign_Right };
-		}
-		
-		return { SNullWidget::NullWidget };
-	}
-
 	bool SBaseReplicationStreamEditor::IsEditingDisabled() const
 	{
 		return (IsEditingEnabledAttribute.IsBound() || IsEditingEnabledAttribute.IsSet())
diff --git a/Engine/Plugins/Developer/Concert/ConcertUI/ConcertSharedSlate/Source/ConcertSharedSlate/Private/Replication/Editor/View/ObjectEditor/SBaseReplicationStreamEditor.h b/Engine/Plugins/Developer/Concert/ConcertUI/ConcertSharedSlate/Source/ConcertSharedSlate/Private/Replication/Editor/View/ObjectEditor/SBaseReplicationStreamEditor.h
--- a/Engine/Plugins/Developer/Concert/ConcertUI/ConcertSharedSlate/Source/ConcertSharedSlate/Private/Replication/Editor/View/ObjectEditor/SBaseReplicationStreamEditor.h
+++ b/Engine/Plugins/Developer/Concert/ConcertUI/ConcertSharedSlate/Source/ConcertSharedSlate/Private/Replication/Editor/View/ObjectEditor/SBaseReplicationStreamEditor.h
@@ -79,6 +79,11 @@ namespace UE::ConcertSharedSlate
 			/** Optional. Whether a given object should be displayed. If this returns false on an object, none of its children will be shown either. */
 			SLATE_EVENT(FShouldDisplayObject, ShouldDisplayObject)
 
+			/** Optional. A widget that should be overlayed on object rows. */
+			SLATE_EVENT(FMakeObjectRowOverlayWidget, MakeObjectRowOverlayWidget)
+			/** Optional. Controls how the widgets created by MakeObjectRowOverlayWidget are supposed to be aligned. */
+			SLATE_ARGUMENT(EHorizontalAlignment, ObjectOverlayAlignment)
+
 			/** Optional widget to add to the left of the object list search bar. */
 			SLATE_NAMED_SLOT(FArguments, LeftOfObjectSearchBar)
 			/** Optional widget to add to the right of the object list search bar. */
@@ -138,9 +143,6 @@ namespace UE::ConcertSharedSlate
 		 */
 		bool bIsAddingFromSelection = false;
 
-		/** Generates a widget when an object row is hovered. Adds a bin and reset button. */
-		FHoverRowContent MakeHoveredRowContent(const TSharedPtr<FReplicatedObjectData>& Data) const;
-
 		bool IsEditingDisabled() const;
 		FText GetEditingDisabledText() const;
 		
diff --git a/Engine/Plugins/Developer/Concert/ConcertUI/ConcertSharedSlate/Source/ConcertSharedSlate/Private/Replication/ReplicationWidgetFactories.cpp b/Engine/Plugins/Developer/Concert/ConcertUI/ConcertSharedSlate/Source/ConcertSharedSlate/Private/Replication/ReplicationWidgetFactories.cpp
--- a/Engine/Plugins/Developer/Concert/ConcertUI/ConcertSharedSlate/Source/ConcertSharedSlate/Private/Replication/ReplicationWidgetFactories.cpp
+++ b/Engine/Plugins/Developer/Concert/ConcertUI/ConcertSharedSlate/Source/ConcertSharedSlate/Private/Replication/ReplicationWidgetFactories.cpp
@@ -43,6 +43,8 @@ namespace UE::ConcertSharedSlate
 			.OnPreAddSelectedObjectsDelegate(MoveTemp(EditorParams.OnPreAddSelectedObjectsDelegate))
 			.OnPostAddSelectedObjectsDelegate(MoveTemp(EditorParams.OnPostAddSelectedObjectsDelegate))
 			.ShouldDisplayObject(MoveTemp(ViewerParams.ShouldDisplayObjectDelegate))
+			.MakeObjectRowOverlayWidget(MoveTemp(ViewerParams.MakeObjectRowOverlayWidgetDelegate))
+			.ObjectOverlayAlignment(ViewerParams.OverlayWidgetAlignment)
 			.LeftOfObjectSearchBar() [ MoveTemp(ViewerParams.LeftOfObjectSearchBar.Widget) ]
 			.RightOfObjectSearchBar() [ MoveTemp(ViewerParams.RightOfObjectSearchBar.Widget) ]
 			.IsEditingEnabled(MoveTemp(EditorParams.IsEditingEnabled))
diff --git a/Engine/Plugins/Developer/Concert/ConcertUI/ConcertSharedSlate/Source/ConcertSharedSlate/Public/Replication/Editor/View/IMultiReplicationStreamEditor.h b/Engine/Plugins/Developer/Concert/ConcertUI/ConcertSharedSlate/Source/ConcertSharedSlate/Public/Replication/Editor/View/IMultiReplicationStreamEditor.h
--- a/Engine/Plugins/Developer/Concert/ConcertUI/ConcertSharedSlate/Source/ConcertSharedSlate/Public/Replication/Editor/View/IMultiReplicationStreamEditor.h
+++ b/Engine/Plugins/Developer/Concert/ConcertUI/ConcertSharedSlate/Source/ConcertSharedSlate/Public/Replication/Editor/View/IMultiReplicationStreamEditor.h
@@ -7,6 +7,7 @@
 namespace UE::ConcertSharedSlate
 {
 	class IEditableMultiReplicationStreamModel;
+	class IEditableReplicationStreamModel;
 	class IReplicationStreamEditor;
 	class IReplicationStreamModel;
 	
@@ -25,6 +26,6 @@ namespace UE::ConcertSharedSlate
 		virtual IEditableMultiReplicationStreamModel& GetMultiStreamModel() const = 0;
 
 		/** @return Gets a model that combines all streams into one. */
-		virtual IReplicationStreamModel& GetConsolidatedModel() const = 0;
+		virtual IEditableReplicationStreamModel& GetConsolidatedModel() const = 0;
 	};
 }
diff --git a/Engine/Plugins/Developer/Concert/ConcertUI/ConcertSharedSlate/Source/ConcertSharedSlate/Public/Replication/ReplicationWidgetFactories.h b/Engine/Plugins/Developer/Concert/ConcertUI/ConcertSharedSlate/Source/ConcertSharedSlate/Public/Replication/ReplicationWidgetFactories.h
--- a/Engine/Plugins/Developer/Concert/ConcertUI/ConcertSharedSlate/Source/ConcertSharedSlate/Public/Replication/ReplicationWidgetFactories.h
+++ b/Engine/Plugins/Developer/Concert/ConcertUI/ConcertSharedSlate/Source/ConcertSharedSlate/Public/Replication/ReplicationWidgetFactories.h
@@ -11,6 +11,7 @@
 
 #include "Delegates/Delegate.h"
 #include "Misc/Attribute.h"
+#include "Types/SlateEnums.h"
 #include "Templates/SharedPointer.h"
 
 class UObject;
@@ -94,6 +95,11 @@ namespace UE::ConcertSharedSlate
 		 * Child objects are determined using ObjectHierarchy.
 		 */
 		FShouldDisplayObject ShouldDisplayObjectDelegate;
+
+		/** Optional. Delegate executed to create a widget that overlays an object row. */
+		FMakeObjectRowOverlayWidget MakeObjectRowOverlayWidgetDelegate;
+		/** Optional. If MakeRowOverlayWidgetDelegate is specified, this controls how the widget is aligned in the column. */
+		EHorizontalAlignment OverlayWidgetAlignment = HAlign_Right;
 	};
 
 	/** Params for creating an IReplicationStreamEditor */
diff --git a/Engine/Plugins/Developer/Concert/ConcertUI/ConcertSharedSlate/Source/ConcertSharedSlate/Public/Replication/Utils/ReplicationWidgetDelegates.h b/Engine/Plugins/Developer/Concert/ConcertUI/ConcertSharedSlate/Source/ConcertSharedSlate/Public/Replication/Utils/ReplicationWidgetDelegates.h
--- a/Engine/Plugins/Developer/Concert/ConcertUI/ConcertSharedSlate/Source/ConcertSharedSlate/Public/Replication/Utils/ReplicationWidgetDelegates.h
+++ b/Engine/Plugins/Developer/Concert/ConcertUI/ConcertSharedSlate/Source/ConcertSharedSlate/Public/Replication/Utils/ReplicationWidgetDelegates.h
@@ -5,12 +5,14 @@
 #include "Containers/ContainersFwd.h"
 #include "Delegates/Delegate.h"
 #include "Delegates/DelegateCombinations.h"
-#include "Templates/Function.h"
+#include "Templates/SharedPointer.h"
 #include "UObject/SoftObjectPtr.h"
 
 class FMenuBuilder;
+class SWidget;
 enum class EBreakBehavior : uint8;
 struct FConcertPropertyChain;
+struct FReplicatedObjectData;
 struct FSoftClassPath;
 struct FSoftObjectPath;
 
@@ -36,4 +38,7 @@ namespace UE::ConcertSharedSlate
 
 	/** Delegate executed when object options are selected from the Add button to the left of the search bar in the top view of the replication panel. */
 	DECLARE_DELEGATE_OneParam(FSelectObjectsFromComboButton, TConstArrayView<FSelectableObjectInfo>);
+
+	/** Delegate executed to create a widget that overlays an object row. */
+	DECLARE_DELEGATE_RetVal_OneParam(TSharedRef<SWidget>, FMakeObjectRowOverlayWidget, const FReplicatedObjectData& ObjectData);
 }
\ No newline at end of file
