The bin icon on replicated object rows now deactivates when all properties are "owned" by offline clients.

diff --git a/Engine/Plugins/Developer/Concert/ConcertApp/MultiUserClient/Source/MultiUserClient/Private/Widgets/ActiveSession/Replication/Client/Multi/SMultiClientView.cpp b/Engine/Plugins/Developer/Concert/ConcertApp/MultiUserClient/Source/MultiUserClient/Private/Widgets/ActiveSession/Replication/Client/Multi/SMultiClientView.cpp
--- a/Engine/Plugins/Developer/Concert/ConcertApp/MultiUserClient/Source/MultiUserClient/Private/Widgets/ActiveSession/Replication/Client/Multi/SMultiClientView.cpp
+++ b/Engine/Plugins/Developer/Concert/ConcertApp/MultiUserClient/Source/MultiUserClient/Private/Widgets/ActiveSession/Replication/Client/Multi/SMultiClientView.cpp
@@ -4,6 +4,7 @@
 
 #include "Selection/ISelectionModel.h"
 #include "MultiStreamModel.h"
+#include "SObjectOverlayRow.h"
 #include "Misc/ObjectUtils.h"
 #include "Replication/ClientReplicationWidgetFactories.h"
 #include "Replication/MultiUserReplicationManager.h"
@@ -23,8 +24,6 @@
 #include "Widgets/ActiveSession/Replication/Client/SPresetComboButton.h"
 #include "Widgets/ActiveSession/Replication/Client/SReplicationStatus.h"
 
-#include "Widgets/Images/SImage.h"
-#include "Widgets/Input/SButton.h"
 #include "Widgets/Layout/SBox.h"
 #include "Widgets/SBoxPanel.h"
 #include "Widgets/Text/STextBlock.h"
@@ -260,44 +259,11 @@ namespace UE::MultiUserClient::Replication
 		return HideObjectsNotInEditorWorld.ShouldShowObject(Object);
 	}
 
-	TSharedRef<SWidget> SMultiClientView::MakeObjectRowOverlayWidget(const ConcertSharedSlate::FReplicatedObjectData& ReplicatedObjectData)
+	TSharedRef<SWidget> SMultiClientView::MakeObjectRowOverlayWidget(const ConcertSharedSlate::FReplicatedObjectData& ReplicatedObjectData) const
 	{
-		if (!ConcertSyncCore::IsActor(ReplicatedObjectData.GetObjectPath()))
-		{
-			return SNullWidget::NullWidget;
-		}
-		
-		return SNew(SButton)
-			.ButtonStyle( FAppStyle::Get(), "SimpleButton")
-			.OnClicked_Lambda([this, ReplicatedObjectData]()
-			{
-				OnPressBinIcon(ReplicatedObjectData.GetObjectPath());
-				return FReply::Handled();
-			})
-			[
-				SNew(SImage)
-				.Image(FAppStyle::GetBrush("Icons.Delete"))
-			];
-	}
-
-	void SMultiClientView::OnPressBinIcon(const FSoftObjectPath& RootObject) const
-	{
-		ConcertSharedSlate::IEditableReplicationStreamModel& ConsolidatedModel = StreamEditor->GetConsolidatedModel();
-		
-		// We want to delete all children, too, i.e. the ones not listed in the outliner, such as components and other subobjects
-		TArray<FSoftObjectPath> ObjectAndChildren = { RootObject };
-		ConsolidatedModel.ForEachReplicatedObject([&RootObject, &ObjectAndChildren](const FSoftObjectPath& ReplicatedObject)
-		{
-			// ReplicatedObject is a child of a deleted object if the path before it contains one of the deleted objects
-			const bool bIsChildOfDeletedObject = ReplicatedObject.ToString().Contains(RootObject.ToString());
-			if (bIsChildOfDeletedObject)
-			{
-				ObjectAndChildren.Add(ReplicatedObject);
-			}
-			return EBreakBehavior::Continue;
-		});
-		
-		ConsolidatedModel.RemoveObjects(ObjectAndChildren);
+		return ConcertSyncCore::IsActor(ReplicatedObjectData.GetObjectPath())
+			? SNew(SObjectOverlayRow, ReplicatedObjectData.GetObjectPath(), StreamEditor.ToSharedRef())
+			: SNullWidget::NullWidget;
 	}
 
 	void SMultiClientView::OnPreAddObjectsFromComboButton(TArrayView<const ConcertSharedSlate::FSelectableObjectInfo>) const
diff --git a/Engine/Plugins/Developer/Concert/ConcertApp/MultiUserClient/Source/MultiUserClient/Private/Widgets/ActiveSession/Replication/Client/Multi/SMultiClientView.h b/Engine/Plugins/Developer/Concert/ConcertApp/MultiUserClient/Source/MultiUserClient/Private/Widgets/ActiveSession/Replication/Client/Multi/SMultiClientView.h
--- a/Engine/Plugins/Developer/Concert/ConcertApp/MultiUserClient/Source/MultiUserClient/Private/Widgets/ActiveSession/Replication/Client/Multi/SMultiClientView.h
+++ b/Engine/Plugins/Developer/Concert/ConcertApp/MultiUserClient/Source/MultiUserClient/Private/Widgets/ActiveSession/Replication/Client/Multi/SMultiClientView.h
@@ -9,16 +9,6 @@
 #include "Widgets/DeclarativeSyntaxSupport.h"
 #include "Widgets/SCompoundWidget.h"
 
-namespace UE::MultiUserClient::Replication
-{
-	class FOfflineClientManager;
-}
-
-namespace UE::MultiUserClient::Replication
-{
-	class FUserPropertySelector;
-}
-
 class IConcertClient;
 class FMenuBuilder;
 
@@ -35,8 +25,10 @@ namespace UE::MultiUserClient::Replication
 	class FGlobalAuthorityCache;
 	class FMultiStreamModel;
 	class FMultiUserReplicationManager;
+	class FOfflineClientManager;
 	class FOnlineClient;
 	class FOnlineClientManager;
+	class FUserPropertySelector;
 	class SPropertySelectionComboButton;
 
 	/** Displays a selection of clients. */
@@ -102,11 +94,9 @@ namespace UE::MultiUserClient::Replication
 		void ExtendObjectContextMenu(FMenuBuilder& MenuBuilder, TConstArrayView<TSoftObjectPtr<>> ContextObjects) const;
 		/** Decides whether the object should be displayed: do not show it if it's not in the editor world. */
 		bool ShouldDisplayObject(const FSoftObjectPath& Object) const;
-		
+
 		/** Creates the widget that overlays actor rows. */
-		TSharedRef<SWidget> MakeObjectRowOverlayWidget(const ConcertSharedSlate::FReplicatedObjectData& ReplicatedObjectData);
-		/** Called when the delete icon over an actor is pressed. Clears the entire hierarchy. */
-		void OnPressBinIcon(const FSoftObjectPath& RootObject) const;
+		TSharedRef<SWidget> MakeObjectRowOverlayWidget(const ConcertSharedSlate::FReplicatedObjectData& ReplicatedObjectData) const;
 		
 		void OnPreAddObjectsFromComboButton(TArrayView<const ConcertSharedSlate::FSelectableObjectInfo>) const;
 		void OnPostAddObjectsFromComboButton(TArrayView<const ConcertSharedSlate::FSelectableObjectInfo>) const;
diff --git a/Engine/Plugins/Developer/Concert/ConcertApp/MultiUserClient/Source/MultiUserClient/Private/Widgets/ActiveSession/Replication/Client/Multi/SObjectOverlayRow.cpp b/Engine/Plugins/Developer/Concert/ConcertApp/MultiUserClient/Source/MultiUserClient/Private/Widgets/ActiveSession/Replication/Client/Multi/SObjectOverlayRow.cpp
new file mode 100644
--- /dev/null
+++ b/Engine/Plugins/Developer/Concert/ConcertApp/MultiUserClient/Source/MultiUserClient/Private/Widgets/ActiveSession/Replication/Client/Multi/SObjectOverlayRow.cpp
@@ -0,0 +1,87 @@
+﻿// Copyright Epic Games, Inc. All Rights Reserved.
+
+#include "SObjectOverlayRow.h"
+
+#include "Replication/Editor/Model/IEditableMultiReplicationStreamModel.h"
+#include "Replication/Editor/Model/IEditableReplicationStreamModel.h"
+#include "Replication/Editor/View/IMultiReplicationStreamEditor.h"
+
+#include "Algo/AnyOf.h"
+#include "Styling/AppStyle.h"
+#include "Widgets/Images/SImage.h"
+#include "Widgets/Input/SButton.h"
+
+#define LOCTEXT_NAMESPACE "SObjectOverlayRow"
+
+namespace UE::MultiUserClient::Replication
+{
+	void SObjectOverlayRow::Construct(
+		const FArguments& InArgs,
+		const FSoftObjectPath& InRootObject,
+		const TSharedRef<ConcertSharedSlate::IMultiReplicationStreamEditor>& InStreamEditor
+		)
+	{
+		RootObject = InRootObject;
+		StreamEditor = InStreamEditor;
+		ChildSlot
+		[
+			SNew(SButton)
+			.ButtonStyle( FAppStyle::Get(), "SimpleButton")
+			.IsEnabled(this, &SObjectOverlayRow::IsBinIconEnabled)
+			.ToolTipText(this, &SObjectOverlayRow::GetBinIconToolTipText)
+			.OnClicked(this, &SObjectOverlayRow::OnPressBinIcon)
+			[
+				SNew(SImage)
+				.Image(FAppStyle::GetBrush("Icons.Delete"))
+			]
+		];
+	}
+
+	FReply SObjectOverlayRow::OnPressBinIcon() const
+	{
+		ConcertSharedSlate::IEditableReplicationStreamModel& ConsolidatedModel = StreamEditor->GetConsolidatedModel();
+		
+		// We want to delete all children, too, i.e. the ones not listed in the outliner, such as components and other subobjects
+		TArray<FSoftObjectPath> ObjectAndChildren = { RootObject };
+		ConsolidatedModel.ForEachReplicatedObject([this, &ObjectAndChildren](const FSoftObjectPath& ReplicatedObject)
+		{
+			// ReplicatedObject is a child of a deleted object if the path before it contains one of the deleted objects
+			const bool bIsChildOfDeletedObject = ReplicatedObject.ToString().Contains(RootObject.ToString());
+			if (bIsChildOfDeletedObject)
+			{
+				ObjectAndChildren.Add(ReplicatedObject);
+			}
+			return EBreakBehavior::Continue;
+		});
+		
+		ConsolidatedModel.RemoveObjects(ObjectAndChildren);
+		return FReply::Handled();
+	}
+
+	FText SObjectOverlayRow::GetBinIconToolTipText() const
+	{
+		return IsBinIconEnabled()
+			? LOCTEXT("ToolTipText.Enabled", "Removes actor and its subobjects")
+			: LOCTEXT("ToolTipText.Disabled", "The properties are not editable.\nContent assigned to offline clients can not be edited until the clients rejoin.");
+	}
+
+	bool SObjectOverlayRow::IsBinIconEnabled() const
+	{
+		using namespace ConcertSharedSlate;
+		
+		const TSet<TSharedRef<IEditableReplicationStreamModel>> Models = StreamEditor->GetMultiStreamModel().GetEditableStreams();
+		return Algo::AnyOf(Models, [this](const TSharedRef<IEditableReplicationStreamModel>& EditableModel)
+		{
+			bool bContainsObject = false;
+			EditableModel->ForEachReplicatedObject([this, &bContainsObject](const FSoftObjectPath& ReplicatedObject)
+			{
+				// ReplicatedObject is a child of a deleted object if the path before it contains one of the deleted objects
+				bContainsObject = ReplicatedObject.ToString().Contains(RootObject.ToString());
+				return bContainsObject ? EBreakBehavior::Break : EBreakBehavior::Continue;
+			});
+			return bContainsObject;
+		});
+	}
+}
+
+#undef LOCTEXT_NAMESPACE
\ No newline at end of file
diff --git a/Engine/Plugins/Developer/Concert/ConcertApp/MultiUserClient/Source/MultiUserClient/Private/Widgets/ActiveSession/Replication/Client/Multi/SObjectOverlayRow.h b/Engine/Plugins/Developer/Concert/ConcertApp/MultiUserClient/Source/MultiUserClient/Private/Widgets/ActiveSession/Replication/Client/Multi/SObjectOverlayRow.h
new file mode 100644
--- /dev/null
+++ b/Engine/Plugins/Developer/Concert/ConcertApp/MultiUserClient/Source/MultiUserClient/Private/Widgets/ActiveSession/Replication/Client/Multi/SObjectOverlayRow.h
@@ -0,0 +1,43 @@
+﻿// Copyright Epic Games, Inc. All Rights Reserved.
+
+#pragma once
+
+#include "Widgets/DeclarativeSyntaxSupport.h"
+#include "Widgets/SCompoundWidget.h"
+
+namespace UE::ConcertSharedSlate
+{
+	class IMultiReplicationStreamEditor;
+}
+
+namespace UE::MultiUserClient::Replication
+{
+	/** Overlaid on actor rows. */
+	class SObjectOverlayRow : public SCompoundWidget
+	{
+	public:
+
+		SLATE_BEGIN_ARGS(SObjectOverlayRow){}
+		SLATE_END_ARGS()
+
+		void Construct(
+			const FArguments& InArgs,
+			const FSoftObjectPath& InRootObject,
+			const TSharedRef<ConcertSharedSlate::IMultiReplicationStreamEditor>& InStreamEditor
+			);
+
+	private:
+
+		/** The top-level object that this row is being shown for. */
+		FSoftObjectPath RootObject;
+		/** Used to delete the object. */
+		TSharedPtr<ConcertSharedSlate::IMultiReplicationStreamEditor> StreamEditor;
+		
+		/** Called when the delete icon over an actor is pressed. Clears the entire hierarchy. */
+		FReply OnPressBinIcon() const;
+		
+		FText GetBinIconToolTipText() const;
+		bool IsBinIconEnabled() const;
+	};
+}
+
diff --git a/Engine/Plugins/Developer/Concert/ConcertUI/ConcertSharedSlate/Source/ConcertSharedSlate/Public/Replication/Editor/Model/IEditableMultiReplicationStreamModel.h b/Engine/Plugins/Developer/Concert/ConcertUI/ConcertSharedSlate/Source/ConcertSharedSlate/Public/Replication/Editor/Model/IEditableMultiReplicationStreamModel.h
--- a/Engine/Plugins/Developer/Concert/ConcertUI/ConcertSharedSlate/Source/ConcertSharedSlate/Public/Replication/Editor/Model/IEditableMultiReplicationStreamModel.h
+++ b/Engine/Plugins/Developer/Concert/ConcertUI/ConcertSharedSlate/Source/ConcertSharedSlate/Public/Replication/Editor/Model/IEditableMultiReplicationStreamModel.h
@@ -32,7 +32,7 @@ namespace UE::ConcertSharedSlate
 		 */
 		virtual TSet<TSharedRef<IEditableReplicationStreamModel>> GetEditableStreams() const = 0;
 
-		/** Util for iterating through both read-only and  */
+		/** Util for iterating through both read-only and editable streams. */
 		void ForEachStream(TFunctionRef<EBreakBehavior(const TSharedRef<IReplicationStreamModel>& Model)> Callback) const
 		{
 			const TSet<TSharedRef<IReplicationStreamModel>> ReadableModels = GetReadOnlyStreams();
