Updated toaster message when applying replication preset that made no changes: instead of saying "Success" it now says "Match no actors for {0} preset" and logs the saved actors.

diff --git a/Engine/Plugins/Developer/Concert/ConcertApp/MultiUserClient/Source/MultiUserClient/Private/Replication/Preset/PresetManager.cpp b/Engine/Plugins/Developer/Concert/ConcertApp/MultiUserClient/Source/MultiUserClient/Private/Replication/Preset/PresetManager.cpp
--- a/Engine/Plugins/Developer/Concert/ConcertApp/MultiUserClient/Source/MultiUserClient/Private/Replication/Preset/PresetManager.cpp
+++ b/Engine/Plugins/Developer/Concert/ConcertApp/MultiUserClient/Source/MultiUserClient/Private/Replication/Preset/PresetManager.cpp
@@ -19,6 +19,20 @@
 
 namespace UE::MultiUserClient::Replication::Private
 {
+	enum class EPresetState : uint8
+	{
+		/** The preset is valid to apply. */
+		Valid,
+		/** Preset contains changes but none could be mapped to objects. */
+		FailedToMapObjects
+	};
+	
+	EPresetState Combine(EPresetState Base, EPresetState New)
+	{
+		return Base == EPresetState::FailedToMapObjects && New == EPresetState::FailedToMapObjects
+			? EPresetState::FailedToMapObjects : EPresetState::Valid;
+	}
+	
 	static void RemoveEmptyObjectsFromRequest(FConcertReplicationStream& Stream)
 	{
 		TMap<FSoftObjectPath, FConcertReplicatedObjectInfo>& ReplicationMap = Stream.BaseDescription.ReplicationMap.ReplicatedObjects;
@@ -50,7 +64,7 @@ namespace UE::MultiUserClient::Replication::Private
 		});
 	}
 
-	static void AddClientToRequest(
+	static EPresetState AddClientToRequest(
 		const FConcertSessionClientInfo& ClientSessionInfo,
 		FConcertReplication_PutState_Request& Request,
 		const UMultiUserReplicationSessionPreset& Preset,
@@ -69,14 +83,14 @@ namespace UE::MultiUserClient::Replication::Private
 			{
 				AddRequestToClearClient();
 			}
-			return;
+			return EPresetState::Valid;
 		}
 
 		const FConcertObjectReplicationMap& OriginalReplicationMap = ClientSessionContent->ReplicationMap;
 		if (OriginalReplicationMap.IsEmpty())
 		{
 			AddRequestToClearClient();
-			return;
+			return EPresetState::Valid;
 		}
 
 		check(GWorld);
@@ -88,14 +102,16 @@ namespace UE::MultiUserClient::Replication::Private
 		if (TranslatedReplicationMap.IsEmpty())
 		{
 			AddRequestToClearClient();
+			return EPresetState::FailedToMapObjects;
 		}
 		else
 		{
 			ProcessNonEmptyRequest(ClientSessionInfo, Request, ClientSessionContent, TranslatedReplicationMap);
+			return EPresetState::Valid;
 		}
 	}
 
-	static void FillStreamAndAuthorityRequest(
+	static TOptional<EPresetState> FillStreamAndAuthorityRequest(
 		FConcertReplication_PutState_Request& Request,
 		const UMultiUserReplicationSessionPreset& Preset,
 		const IConcertClientSession& Session,
@@ -105,12 +121,17 @@ namespace UE::MultiUserClient::Replication::Private
 		// GWorld is required by AddClientToRequest
 		if (ensure(GWorld))
 		{
-			AddClientToRequest({ Session.GetSessionClientEndpointId(), Session.GetLocalClientInfo() }, Request, Preset, bClearUnreferencedClients);
+			EPresetState PresetState = AddClientToRequest(
+				{ Session.GetSessionClientEndpointId(), Session.GetLocalClientInfo() }, Request, Preset, bClearUnreferencedClients
+				);
 			for (const FConcertSessionClientInfo& ClientSessionInfo : Session.GetSessionClients())
 			{
-				AddClientToRequest(ClientSessionInfo, Request, Preset, bClearUnreferencedClients);
+				const EPresetState RemoteClientState = AddClientToRequest(ClientSessionInfo, Request, Preset, bClearUnreferencedClients);
+				PresetState = Combine(PresetState, RemoteClientState);
 			}
+			return PresetState;
 		}
+		return {};
 	}
 
 	static void FillMuteStateRequest(
@@ -154,7 +175,7 @@ namespace UE::MultiUserClient::Replication::Private
 		}
 	}
 	
-	static FConcertReplication_PutState_Request BuildRequest(
+	static TPair<FConcertReplication_PutState_Request, TOptional<EPresetState>> BuildRequest(
 		const UMultiUserReplicationSessionPreset& Preset,
 		const IConcertClientSession& Session,
 		EApplyPresetFlags Flags
@@ -163,13 +184,13 @@ namespace UE::MultiUserClient::Replication::Private
 		FConcertReplication_PutState_Request Request;
 
 		const bool bClearUnreferencedClients = EnumHasAnyFlags(Flags, EApplyPresetFlags::ClearUnreferencedClients);
-		FillStreamAndAuthorityRequest(Request, Preset, Session, bClearUnreferencedClients);
+		const TOptional<EPresetState> PresetState = FillStreamAndAuthorityRequest(Request, Preset, Session, bClearUnreferencedClients);
 
 		// TODO UE-219829: Once the server allows sending the mute state disconnected clients should have when they rejoin,
 		// simply send over all mute state instead of doing filtering here.
 		FillMuteStateRequest(Request, Preset, Session);
 		
-		return Request;
+		return { MoveTemp(Request), PresetState };
 	}
 
 	static EReplaceSessionContentErrorCode ExtractErrorCode(const FConcertReplication_PutState_Response& Response)
@@ -270,10 +291,18 @@ namespace UE::MultiUserClient::Replication
 		{
 			return MakeFulfilledPromise<FReplaceSessionContentResult>(EReplaceSessionContentErrorCode::Timeout).GetFuture();
 		}
+
+		const auto[Request, PresetState] = Private::BuildRequest(Preset, *Session, Flags);
+		if (!PresetState || *PresetState == Private::EPresetState::FailedToMapObjects)
+		{
+			const EReplaceSessionContentErrorCode ErrorCode = PresetState
+				? EReplaceSessionContentErrorCode::NoObjectsFound : EReplaceSessionContentErrorCode::NoWorld;
+			return MakeFulfilledPromise<FReplaceSessionContentResult>(ErrorCode).GetFuture();
+		}
 		
 		InProgressSessionReplacementOp = MakeShared<TPromise<FReplaceSessionContentResult>>();
 		SyncClient.GetReplicationManager()
-			->PutClientState(Private::BuildRequest(Preset, *Session, Flags))
+			->PutClientState(Request)
 			.Next(
 				[this, WeakPromise = TWeakPtr<TPromise<FReplaceSessionContentResult>>(InProgressSessionReplacementOp)]
 				(FConcertReplication_PutState_Response&& Response)
diff --git a/Engine/Plugins/Developer/Concert/ConcertApp/MultiUserClient/Source/MultiUserClient/Private/Replication/Preset/PresetUtils.cpp b/Engine/Plugins/Developer/Concert/ConcertApp/MultiUserClient/Source/MultiUserClient/Private/Replication/Preset/PresetUtils.cpp
new file mode 100644
--- /dev/null
+++ b/Engine/Plugins/Developer/Concert/ConcertApp/MultiUserClient/Source/MultiUserClient/Private/Replication/Preset/PresetUtils.cpp
@@ -0,0 +1,8 @@
+﻿// Copyright Epic Games, Inc. All Rights Reserved.
+
+#include "PresetUtils.h"
+
+namespace UE::MultiUserClient::Replication
+{
+	
+}
\ No newline at end of file
diff --git a/Engine/Plugins/Developer/Concert/ConcertApp/MultiUserClient/Source/MultiUserClient/Private/Replication/Preset/PresetUtils.h b/Engine/Plugins/Developer/Concert/ConcertApp/MultiUserClient/Source/MultiUserClient/Private/Replication/Preset/PresetUtils.h
new file mode 100644
--- /dev/null
+++ b/Engine/Plugins/Developer/Concert/ConcertApp/MultiUserClient/Source/MultiUserClient/Private/Replication/Preset/PresetUtils.h
@@ -0,0 +1,31 @@
+﻿// Copyright Epic Games, Inc. All Rights Reserved.
+
+#pragma once
+
+#include "Assets/MultiUserReplicationClientPreset.h"
+#include "Assets/MultiUserReplicationSessionPreset.h"
+#include "Misc/EBreakBehavior.h"
+
+#include <type_traits>
+
+namespace UE::MultiUserClient::Replication
+{
+	/** Lists out all actors saved in the preset. */
+	template<typename TLambda>
+	requires std::is_invocable_r_v<EBreakBehavior, TLambda, const FSoftObjectPath& /*ActorPath*/, const FString& /*Label*/>
+	void ForEachSavedActorLabel(const UMultiUserReplicationSessionPreset& Preset, TLambda&& Callback)
+	{
+		for (const FMultiUserReplicationClientPreset& ClientPreset : Preset.GetClientPresets())
+		{
+			for (const TPair<FSoftObjectPath, FConcertReplicationRemappingData_Actor>& Pair : ClientPreset.ActorLabelRemappingData.ActorData)
+			{
+				const FSoftObjectPath& ActorPath = Pair.Key;
+				const FString& Label = Pair.Value.Label;
+				if (Callback(ActorPath, Label) == EBreakBehavior::Break)
+				{
+					return;
+				}
+			}
+		}
+	}
+}
diff --git a/Engine/Plugins/Developer/Concert/ConcertApp/MultiUserClient/Source/MultiUserClient/Private/Replication/Preset/ReplaceSessionContentResult.h b/Engine/Plugins/Developer/Concert/ConcertApp/MultiUserClient/Source/MultiUserClient/Private/Replication/Preset/ReplaceSessionContentResult.h
--- a/Engine/Plugins/Developer/Concert/ConcertApp/MultiUserClient/Source/MultiUserClient/Private/Replication/Preset/ReplaceSessionContentResult.h
+++ b/Engine/Plugins/Developer/Concert/ConcertApp/MultiUserClient/Source/MultiUserClient/Private/Replication/Preset/ReplaceSessionContentResult.h
@@ -10,6 +10,10 @@ namespace UE::MultiUserClient::Replication
 	{
 		/** Request completed successfully. */
 		Success,
+		/** The preset did not make any changes because no objects could be found in the world. */
+		NoObjectsFound,
+		/** GWorld was not valid. */
+		NoWorld,
 		
 		/** Request cancelled because FPresetManager was destroyed - probably because the use left the session during the request. */
 		Cancelled,
diff --git a/Engine/Plugins/Developer/Concert/ConcertApp/MultiUserClient/Source/MultiUserClient/Private/Widgets/ActiveSession/Replication/Client/SPresetComboButton.cpp b/Engine/Plugins/Developer/Concert/ConcertApp/MultiUserClient/Source/MultiUserClient/Private/Widgets/ActiveSession/Replication/Client/SPresetComboButton.cpp
--- a/Engine/Plugins/Developer/Concert/ConcertApp/MultiUserClient/Source/MultiUserClient/Private/Widgets/ActiveSession/Replication/Client/SPresetComboButton.cpp
+++ b/Engine/Plugins/Developer/Concert/ConcertApp/MultiUserClient/Source/MultiUserClient/Private/Widgets/ActiveSession/Replication/Client/SPresetComboButton.cpp
@@ -2,51 +2,104 @@
 
 #include "SPresetComboButton.h"
 
+#include "ConcertLogGlobal.h"
 #include "Assets/MultiUserReplicationSessionPreset.h"
 #include "Replication/Preset/PresetManager.h"
 
+#include "Containers/StringFwd.h"
 #include "Containers/Ticker.h"
 #include "ContentBrowserModule.h"
+#include "Engine/World.h"
 #include "IContentBrowserSingleton.h"
 #include "Framework/MultiBox/MultiBoxBuilder.h"
 #include "Framework/Notifications/NotificationManager.h"
+#include "Misc/StringBuilder.h"
+#include "Replication/Preset/PresetUtils.h"
 #include "SSimpleComboButton.h"
 #include "Styling/AppStyle.h"
+#include "UObject/WeakObjectPtr.h"
+#include "UObject/WeakObjectPtrTemplates.h"
 #include "Widgets/Client/SClientName.h"
 #include "Widgets/Layout/SBorder.h"
 #include "Widgets/Notifications/SNotificationList.h"
 
 #define LOCTEXT_NAMESPACE "SPresetComboButton"
 
-namespace UE::MultiUserClient::Replication
+namespace UE::MultiUserClient::Replication::Private
 {
-	namespace Private
+	static void LogActorsInPreset(const UMultiUserReplicationSessionPreset& Preset)
+	{
+		TStringBuilder<2048> StringBuilder;
+		ForEachSavedActorLabel(Preset,
+			[&StringBuilder](const FSoftObjectPath& ActorPath, const FString& Label)
+			{
+				StringBuilder << TEXT("\tLabel=\"") << *Label << TEXT("\", Path=\"") << *ActorPath.ToString() << TEXT("\"\n");
+				return EBreakBehavior::Continue;
+			});
+		
+		UE_LOG(LogConcert, Warning, TEXT("No actors could be mapped to world %s for preset %s. Saved actors:\n%s"),
+			GWorld ? *GWorld->GetPathName() : TEXT("none"),
+			*Preset.GetPathName(),
+			*StringBuilder
+			);
+	}
+	
+	static void LogPresetErrorsIfNeeded(
+		const TWeakObjectPtr<const UMultiUserReplicationSessionPreset>& WeakPreset,
+		EReplaceSessionContentErrorCode ErrorCode
+		)
 	{
-		static FText MakeTitle(const FReplaceSessionContentResult& Result, const FText& PresetText)
+		const UMultiUserReplicationSessionPreset* Preset = WeakPreset.Get();
+		if (!Preset)
 		{
-			const FText Format = Result.IsSuccess()
-				? LOCTEXT("ApplyPreset.Title.SuccessFmt", "Applied {0} preset")
-				: LOCTEXT("ApplyPreset.Title.FailFmt", "Failed to apply {0} preset");
-			return FText::Format(Format, PresetText);
+			return;
 		}
 
-		static FText MakeSubText(const FReplaceSessionContentResult& Result)
+		switch (ErrorCode)
+		{
+		case EReplaceSessionContentErrorCode::NoObjectsFound:
+			LogActorsInPreset(*Preset);
+			break;
+		default: break;
+		}
+	}
+	
+	static FText MakeTitle(const FReplaceSessionContentResult& Result, const FText& PresetText)
+	{
+		const FText Format = [&Result]()
 		{
 			switch (Result.ErrorCode)
 			{
-			case EReplaceSessionContentErrorCode::Success: return FText::GetEmpty();
-			case EReplaceSessionContentErrorCode::Cancelled: return LOCTEXT("ApplyPreset.SubText.Success", "Disconnected from session.");
-			case EReplaceSessionContentErrorCode::InProgress: return LOCTEXT("ApplyPreset.SubText.InProgress", "Another operation is already in progress.");
-			case EReplaceSessionContentErrorCode::Timeout: return LOCTEXT("ApplyPreset.SubText.Timeout", "Request timed out.");
-			case EReplaceSessionContentErrorCode::FeatureDisabled: return LOCTEXT("ApplyPreset.SubText.FeatureDisabled", "This session does not support presets.");
-			case EReplaceSessionContentErrorCode::Rejected: return LOCTEXT("ApplyPreset.SubText.Rejected", "Rejected by server.");
-			default: checkNoEntry(); return FText::GetEmpty();
+			case EReplaceSessionContentErrorCode::Success: return LOCTEXT("ApplyPreset.Title.SuccessFmt", "Applied {0} preset");
+			case EReplaceSessionContentErrorCode::NoObjectsFound: return LOCTEXT("ApplyPreset.Title.NoObjectsFoundFmt", "No actors matched for {0} preset");
+			default: return LOCTEXT("ApplyPreset.Title.FailFmt", "Failed to apply {0} preset");
 			}
-		}
+		}();
+		
+		return FText::Format(Format, PresetText);
+	}
 
-		static FText GetSavedClientsText() { return LOCTEXT("Save.IncludedClients.Label", "Saved clients"); }
+	static FText MakeSubText(const FReplaceSessionContentResult& Result)
+	{
+		switch (Result.ErrorCode)
+		{
+		case EReplaceSessionContentErrorCode::Success: return FText::GetEmpty();
+		case EReplaceSessionContentErrorCode::NoObjectsFound: return LOCTEXT("ApplyPreset.SubText.NoObjectsFound", "No actors from the preset were found in the world.\n\nCheck the output log to see the actors saved in the preset.");
+		case EReplaceSessionContentErrorCode::NoWorld: return LOCTEXT("ApplyPreset.SubText.NoWorld", "No world instance to remap preset content to.");
+		case EReplaceSessionContentErrorCode::Cancelled: return LOCTEXT("ApplyPreset.SubText.Success", "Disconnected from session.");
+		case EReplaceSessionContentErrorCode::InProgress: return LOCTEXT("ApplyPreset.SubText.InProgress", "Another operation is already in progress.");
+		case EReplaceSessionContentErrorCode::Timeout: return LOCTEXT("ApplyPreset.SubText.Timeout", "Request timed out.");
+		case EReplaceSessionContentErrorCode::FeatureDisabled: return LOCTEXT("ApplyPreset.SubText.FeatureDisabled", "This session does not support presets.");
+		case EReplaceSessionContentErrorCode::Rejected: return LOCTEXT("ApplyPreset.SubText.Rejected", "Rejected by server.");
+		default: checkNoEntry(); return FText::GetEmpty();
+		}
 	}
-	
+
+	static FText GetSavedClientsText() { return LOCTEXT("Save.IncludedClients.Label", "Saved clients"); }
+}
+
+namespace UE::MultiUserClient::Replication
+{
 	void SPresetComboButton::Construct(const FArguments& InArgs, const IConcertClient& InClient, FPresetManager& InPresetManager)
 	{
 		Client = &InClient;
@@ -247,10 +300,16 @@ namespace UE::MultiUserClient::Replication
 			Notification->SetCompletionState(SNotificationItem::CS_Pending);
 
 			PresetManager->ReplaceSessionContentWithPreset(*Preset, BuildFlags())
-				.Next([PresetText = MoveTemp(PresetText), Notification = MoveTemp(Notification)](const FReplaceSessionContentResult& Result) mutable
+				.Next([Preset, PresetText = MoveTemp(PresetText), Notification = MoveTemp(Notification)]
+					(const FReplaceSessionContentResult& Result) mutable
 				{
-					ExecuteOnGameThread(TEXT("SPresetComboButton"), [PresetText = MoveTemp(PresetText), Notification = MoveTemp(Notification), Result]()
+					ExecuteOnGameThread(TEXT("SPresetComboButton"),
+						[Preset = TWeakObjectPtr<const UMultiUserReplicationSessionPreset>(Preset),
+							PresetText = MoveTemp(PresetText),
+							Notification = MoveTemp(Notification),
+							Result]
 					{
+						Private::LogPresetErrorsIfNeeded(Preset, Result.ErrorCode);
 						Notification->SetText(Private::MakeTitle(Result, PresetText));
 						Notification->SetSubText(Private::MakeSubText(Result));
 						Notification->SetCompletionState(Result.IsSuccess() ? SNotificationItem::CS_Success : SNotificationItem::CS_Fail);
