Added view option to MU replication: whether to show offline clients.

diff --git a/Engine/Plugins/Developer/Concert/ConcertApp/MultiUserClient/Source/MultiUserClient/Private/Replication/Client/UnifiedStreamCache.cpp b/Engine/Plugins/Developer/Concert/ConcertApp/MultiUserClient/Source/MultiUserClient/Private/Replication/Client/UnifiedStreamCache.cpp
--- a/Engine/Plugins/Developer/Concert/ConcertApp/MultiUserClient/Source/MultiUserClient/Private/Replication/Client/UnifiedStreamCache.cpp
+++ b/Engine/Plugins/Developer/Concert/ConcertApp/MultiUserClient/Source/MultiUserClient/Private/Replication/Client/UnifiedStreamCache.cpp
@@ -76,7 +76,7 @@ namespace UE::MultiUserClient::Replication
 	void FUnifiedStreamCache::EnumerateClientsWithObject(
 		const FSoftObjectPath& ObjectPath,
 		const TFunctionRef<EBreakBehavior(const FGuid& ClientId)>& Callback,
-		EClientEnumerationFlags Flags
+		EClientEnumerationMode Option
 		) const
 	{
 		EBreakBehavior BreakBehavior = EBreakBehavior::Continue;
@@ -87,12 +87,13 @@ namespace UE::MultiUserClient::Replication
 			return BreakBehavior;
 		});
 
-		if (BreakBehavior == EBreakBehavior::Break)
+		if (BreakBehavior == EBreakBehavior::Break
+			|| Option == EClientEnumerationMode::SkipOfflineClients)
 		{
 			return;
 		}
 
-		const bool bShouldCheckNonOverlapping = EnumHasAnyFlags(Flags, EClientEnumerationFlags::SkipOfflineClientsThatFullyOverlapWithOnline);
+		const bool bShouldCheckNonOverlapping = Option == EClientEnumerationMode::SkipOfflineClientsThatFullyOverlapWithOnlineClients;
 		OfflineClientManager.ForEachClient(
 			[this, &ObjectPath, &Callback, bShouldCheckNonOverlapping, &OnlineCache](const FOfflineClient& Client)
 			{
@@ -109,7 +110,7 @@ namespace UE::MultiUserClient::Replication
 		const FSoftObjectPath& ObjectPath,
 		const FConcertPropertyChain& Property,
 		const TFunctionRef<EBreakBehavior(const FGuid& ClientId)>& Callback,
-		EClientEnumerationFlags Option
+		EClientEnumerationMode Option
 		) const
 	{
 		EBreakBehavior BreakBehavior = EBreakBehavior::Continue;
@@ -138,8 +139,9 @@ namespace UE::MultiUserClient::Replication
 		);
 
 		if (BreakBehavior == EBreakBehavior::Break
+			|| Option == EClientEnumerationMode::SkipOfflineClients
 			// If there is an online client, the offline clients are not supposed to be displayed because the caller requested no property overlaps.
-			|| (Option == EClientEnumerationFlags::SkipOfflineClientsThatFullyOverlapWithOnline && bEncounteredOnlineClient))
+			|| (Option == EClientEnumerationMode::SkipOfflineClientsThatFullyOverlapWithOnlineClients && bEncounteredOnlineClient))
 		{
 			return;
 		}
diff --git a/Engine/Plugins/Developer/Concert/ConcertApp/MultiUserClient/Source/MultiUserClient/Private/Replication/Client/UnifiedStreamCache.h b/Engine/Plugins/Developer/Concert/ConcertApp/MultiUserClient/Source/MultiUserClient/Private/Replication/Client/UnifiedStreamCache.h
--- a/Engine/Plugins/Developer/Concert/ConcertApp/MultiUserClient/Source/MultiUserClient/Private/Replication/Client/UnifiedStreamCache.h
+++ b/Engine/Plugins/Developer/Concert/ConcertApp/MultiUserClient/Source/MultiUserClient/Private/Replication/Client/UnifiedStreamCache.h
@@ -4,18 +4,11 @@
 
 #include "Delegates/Delegate.h"
 #include "HAL/Platform.h"
-#include "Misc/EnumClassFlags.h"
 #include "Templates/FunctionFwd.h"
 
+enum class EBreakBehavior : uint8;
 struct FConcertObjectReplicationMap;
 struct FConcertPropertyChain;
-
-namespace UE::MultiUserClient::Replication
-{
-	class FUnifiedClientView;
-}
-
-enum class EBreakBehavior : uint8;
 struct FGuid;
 struct FSoftObjectPath;
 
@@ -23,13 +16,16 @@ namespace UE::MultiUserClient::Replication
 {
 	class FOfflineClientManager;
 	class FOnlineClientManager;
+	class FUnifiedClientView;
 
-	enum class EClientEnumerationFlags : uint8
+	enum class EClientEnumerationMode : uint8
 	{
-		None = 0,
-		SkipOfflineClientsThatFullyOverlapWithOnline = 1 << 0
+		All,
+		/** Skip offline clients if all their properties are also owned by online clients. */
+		SkipOfflineClientsThatFullyOverlapWithOnlineClients,
+		/** Offline clients are supposed to be skipped completely. */
+		SkipOfflineClients
 	};
-	ENUM_CLASS_FLAGS(EClientEnumerationFlags);
 	
 	/** Access point for querying stream content of both online and offline clients through a single interface. */
 	class FUnifiedStreamCache : public FNoncopyable
@@ -53,7 +49,7 @@ namespace UE::MultiUserClient::Replication
 		void EnumerateClientsWithObject(
 			const FSoftObjectPath& ObjectPath,
 			const TFunctionRef<EBreakBehavior(const FGuid& ClientId)>& Callback,
-			EClientEnumerationFlags Flags = EClientEnumerationFlags::SkipOfflineClientsThatFullyOverlapWithOnline
+			EClientEnumerationMode Option = EClientEnumerationMode::SkipOfflineClientsThatFullyOverlapWithOnlineClients
 			) const;
 
 		/**
@@ -64,7 +60,7 @@ namespace UE::MultiUserClient::Replication
 			const FSoftObjectPath& ObjectPath,
 			const FConcertPropertyChain& Property,
 			const TFunctionRef<EBreakBehavior(const FGuid& ClientId)>& Callback,
-			EClientEnumerationFlags Option = EClientEnumerationFlags::SkipOfflineClientsThatFullyOverlapWithOnline
+			EClientEnumerationMode Option = EClientEnumerationMode::SkipOfflineClientsThatFullyOverlapWithOnlineClients
 			) const;
 
 		/** Broadcasts when the content of a client has changed. */
diff --git a/Engine/Plugins/Developer/Concert/ConcertApp/MultiUserClient/Source/MultiUserClient/Private/Widgets/ActiveSession/Replication/Client/Multi/Columns/AssignProperty/AssignPropertyModel.cpp b/Engine/Plugins/Developer/Concert/ConcertApp/MultiUserClient/Source/MultiUserClient/Private/Widgets/ActiveSession/Replication/Client/Multi/Columns/AssignProperty/AssignPropertyModel.cpp
--- a/Engine/Plugins/Developer/Concert/ConcertApp/MultiUserClient/Source/MultiUserClient/Private/Widgets/ActiveSession/Replication/Client/Multi/Columns/AssignProperty/AssignPropertyModel.cpp
+++ b/Engine/Plugins/Developer/Concert/ConcertApp/MultiUserClient/Source/MultiUserClient/Private/Widgets/ActiveSession/Replication/Client/Multi/Columns/AssignProperty/AssignPropertyModel.cpp
@@ -4,14 +4,14 @@
 
 #include "ConcertLogGlobal.h"
 #include "Replication/Editor/Model/PropertyUtils.h"
+#include "Replication/Client/UnifiedClientView.h"
+#include "Widgets/ActiveSession/Replication/Client/Multi/ViewOptions/MultiViewOptions.h"
 
 #include "GameFramework/Actor.h"
 #include "ScopedTransaction.h"
 
 #include <type_traits>
 
-#include "Replication/Client/UnifiedClientView.h"
-
 #define LOCTEXT_NAMESPACE "FAssignPropertyModel"
 
 namespace UE::MultiUserClient::Replication::MultiStreamColumns::AssignPropertyModel
@@ -91,17 +91,20 @@ namespace UE::MultiUserClient::Replication::MultiStreamColumns::AssignPropertyMo
 
 namespace UE::MultiUserClient::Replication::MultiStreamColumns
 {
-	FAssignPropertyModel::FAssignPropertyModel(FUnifiedClientView& InClientView)
+	FAssignPropertyModel::FAssignPropertyModel(FUnifiedClientView& InClientView, FMultiViewOptions& InViewOptions)
 		: ClientView(InClientView)
+		, ViewOptions(InViewOptions)
 	{
 		ClientView.OnClientsChanged().AddRaw(this, &FAssignPropertyModel::BroadcastOnOwnershipChanged);
 		ClientView.GetStreamCache().OnCacheChanged().AddRaw(this, &FAssignPropertyModel::BroadcastOnOwnershipChanged);
+		ViewOptions.OnOptionsChanged().AddRaw(this, &FAssignPropertyModel::BroadcastOnOwnershipChanged);
 	}
 
 	FAssignPropertyModel::~FAssignPropertyModel()
 	{
 		ClientView.OnClientsChanged().RemoveAll(this);
 		ClientView.GetStreamCache().OnCacheChanged().RemoveAll(this);
+		ViewOptions.OnOptionsChanged().RemoveAll(this);
 	}
 
 #define SET_REASON(Text) if (Reason) { *Reason = Text; }
@@ -174,7 +177,11 @@ namespace UE::MultiUserClient::Replication::MultiStreamColumns
 		return Result;
 	}
 
-	void FAssignPropertyModel::TogglePropertyFor(const FGuid& ClientId, TConstArrayView<TSoftObjectPtr<>> Objects, const FConcertPropertyChain& Property)
+	void FAssignPropertyModel::TogglePropertyFor(
+		const FGuid& ClientId,
+		TConstArrayView<TSoftObjectPtr<>> Objects,
+		const FConcertPropertyChain& Property
+		) const
 	{
 		// Remote clients can disconnect after the combo-box is opened.
 		const TSharedPtr<ConcertSharedSlate::IEditableReplicationStreamModel> StreamModel = ClientView.GetEditableClientStreamById(ClientId);
@@ -223,6 +230,26 @@ namespace UE::MultiUserClient::Replication::MultiStreamColumns
 			AssignPropertyModel::UnassignPropertyFromClients(ClientView, Objects, Property, [](auto&){ return true; });
 		}
 	}
+
+	void FAssignPropertyModel::ForEachAssignedClient(
+		const FConcertPropertyChain& DisplayedProperty,
+		 const TArray<TSoftObjectPtr<>>& EditedObjects,
+		 TFunctionRef<EBreakBehavior(const FGuid& ClientId)> Callback
+		 ) const
+	{
+		const EClientEnumerationMode Mode = ViewOptions.ShouldShowOfflineClients()
+			? EClientEnumerationMode::SkipOfflineClientsThatFullyOverlapWithOnlineClients
+			: EClientEnumerationMode::SkipOfflineClients;
+		for (const TSoftObjectPtr<>& Object : EditedObjects)
+		{
+			ClientView.GetStreamCache().EnumerateClientsWithObjectAndProperty(
+				Object.ToSoftObjectPath(),
+				DisplayedProperty,
+				[&Callback](const FGuid& ClientId){ return Callback(ClientId); },
+				Mode
+			);
+		}
+	}
 }
 
 #undef LOCTEXT_NAMESPACE
\ No newline at end of file
diff --git a/Engine/Plugins/Developer/Concert/ConcertApp/MultiUserClient/Source/MultiUserClient/Private/Widgets/ActiveSession/Replication/Client/Multi/Columns/AssignProperty/AssignPropertyModel.h b/Engine/Plugins/Developer/Concert/ConcertApp/MultiUserClient/Source/MultiUserClient/Private/Widgets/ActiveSession/Replication/Client/Multi/Columns/AssignProperty/AssignPropertyModel.h
--- a/Engine/Plugins/Developer/Concert/ConcertApp/MultiUserClient/Source/MultiUserClient/Private/Widgets/ActiveSession/Replication/Client/Multi/Columns/AssignProperty/AssignPropertyModel.h
+++ b/Engine/Plugins/Developer/Concert/ConcertApp/MultiUserClient/Source/MultiUserClient/Private/Widgets/ActiveSession/Replication/Client/Multi/Columns/AssignProperty/AssignPropertyModel.h
@@ -7,10 +7,12 @@
 #include "Templates/UnrealTemplate.h"
 #include "UObject/SoftObjectPtr.h"
 
+enum class EBreakBehavior : uint8;
 struct FConcertPropertyChain;
 
 namespace UE::MultiUserClient::Replication
 {
+	class FMultiViewOptions;
 	class FOnlineClient;
 	class FOnlineClientManager;
 	class FUnifiedClientView;
@@ -33,7 +35,7 @@ namespace UE::MultiUserClient::Replication::MultiStreamColumns
 	{
 	public:
 		
-		FAssignPropertyModel(FUnifiedClientView& InClientView UE_LIFETIMEBOUND);
+		FAssignPropertyModel(FUnifiedClientView& InClientView UE_LIFETIMEBOUND, FMultiViewOptions& InViewOptions UE_LIFETIMEBOUND);
 		~FAssignPropertyModel();
 
 		/** @return Whether property ownership can be changed for the given client. */
@@ -44,9 +46,16 @@ namespace UE::MultiUserClient::Replication::MultiStreamColumns
 		EPropertyOnObjectsOwnershipState GetPropertyOwnershipState(const FGuid& ClientId, TConstArrayView<TSoftObjectPtr<>> Objects, const FConcertPropertyChain& Property) const;
 
 		/** Assigns the property to ClientId or unassigns it. Removes the property from all other clients in both cases. */
-		void TogglePropertyFor(const FGuid& ClientId, TConstArrayView<TSoftObjectPtr<>> Objects, const FConcertPropertyChain& Property);
+		void TogglePropertyFor(const FGuid& ClientId, TConstArrayView<TSoftObjectPtr<>> Objects, const FConcertPropertyChain& Property) const;
 		/** Removes the property from all clients. */
 		void ClearProperty(TConstArrayView<TSoftObjectPtr<>> Objects, const FConcertPropertyChain& Property) const;
+
+		/** Iterates every client that has DisplayedProperty assigned for any of EditedObjects. */
+		void ForEachAssignedClient(
+			const FConcertPropertyChain& DisplayedProperty,
+			const TArray<TSoftObjectPtr<>>& EditedObjects,
+			TFunctionRef<EBreakBehavior(const FGuid& ClientId)> Callback
+			) const;
 		
 		/** Broadcasts when property ownership may have changed. */
 		FSimpleMulticastDelegate& OnOwnershipChanged() { return OnOwnershipChangedDelegate; }
@@ -59,6 +68,8 @@ namespace UE::MultiUserClient::Replication::MultiStreamColumns
 		 * - detect changes made to client content for broadcasting OnOwnershipChangedDelegate
 		 */
 		FUnifiedClientView& ClientView;
+		/** Controls whether offline clients should be listed by ForEachAssignedClient. */
+		FMultiViewOptions& ViewOptions;
 
 		/** Broadcasts when property ownership may have changed. */
 		FSimpleMulticastDelegate OnOwnershipChangedDelegate;
diff --git a/Engine/Plugins/Developer/Concert/ConcertApp/MultiUserClient/Source/MultiUserClient/Private/Widgets/ActiveSession/Replication/Client/Multi/Columns/AssignProperty/MultiStreamColumns_AssignProperties.cpp b/Engine/Plugins/Developer/Concert/ConcertApp/MultiUserClient/Source/MultiUserClient/Private/Widgets/ActiveSession/Replication/Client/Multi/Columns/AssignProperty/MultiStreamColumns_AssignProperties.cpp
--- a/Engine/Plugins/Developer/Concert/ConcertApp/MultiUserClient/Source/MultiUserClient/Private/Widgets/ActiveSession/Replication/Client/Multi/Columns/AssignProperty/MultiStreamColumns_AssignProperties.cpp
+++ b/Engine/Plugins/Developer/Concert/ConcertApp/MultiUserClient/Source/MultiUserClient/Private/Widgets/ActiveSession/Replication/Client/Multi/Columns/AssignProperty/MultiStreamColumns_AssignProperties.cpp
@@ -54,9 +54,10 @@ namespace UE::MultiUserClient::Replication::MultiStreamColumns
 	const FName AssignPropertyColumnId(TEXT("AssignPropertyColumn"));
 	
 	ConcertSharedSlate::FPropertyColumnEntry AssignPropertyColumn(
-		TAttribute<TSharedPtr<ConcertSharedSlate::IMultiReplicationStreamEditor>> MultiStreamEditor,
-		FUnifiedClientView& ClientView,
-		const int32 ColumnsSortPriority
+		TAttribute<TSharedPtr<ConcertSharedSlate::IMultiReplicationStreamEditor>> InMultiStreamEditor,
+		FUnifiedClientView& InClientView,
+		FMultiViewOptions& InViewOptions,
+		const int32 InColumnsSortPriority
 		)
 	{
 		using namespace ConcertSharedSlate;
@@ -66,10 +67,11 @@ namespace UE::MultiUserClient::Replication::MultiStreamColumns
 
 			FPropertyColumn_AssignPropertyColumn(
 				TAttribute<TSharedPtr<IMultiReplicationStreamEditor>> MultiStreamEditor,
-				FUnifiedClientView& ClientView UE_LIFETIMEBOUND
+				FUnifiedClientView& ClientView UE_LIFETIMEBOUND,
+				FMultiViewOptions& InViewOptions UE_LIFETIMEBOUND
 				)
 				: MultiStreamEditor(MoveTemp(MultiStreamEditor))
-				, Model(ClientView)
+				, Model(ClientView, InViewOptions)
 				, ClientView(ClientView)
 			{}
 			
@@ -113,8 +115,12 @@ namespace UE::MultiUserClient::Replication::MultiStreamColumns
 			virtual bool CanBeSorted() const override { return true; }
 			virtual bool IsLessThan(const FPropertyTreeRowContext& Left, const FPropertyTreeRowContext& Right) const override
 			{
-				const TOptional<FString> LeftClientDisplayString = SAssignPropertyComboBox::GetDisplayString(ClientView, Left.RowData.GetProperty(), Left.RowData.GetContextObjects());
-				const TOptional<FString> RightClientDisplayString = SAssignPropertyComboBox::GetDisplayString(ClientView, Right.RowData.GetProperty(), Right.RowData.GetContextObjects());
+				const TOptional<FString> LeftClientDisplayString = SAssignPropertyComboBox::GetDisplayString(
+					Model, ClientView, Left.RowData.GetProperty(), Left.RowData.GetContextObjects()
+					);
+				const TOptional<FString> RightClientDisplayString = SAssignPropertyComboBox::GetDisplayString(
+				Model, ClientView, Right.RowData.GetProperty(), Right.RowData.GetContextObjects()
+				);
 			
 				if (LeftClientDisplayString && RightClientDisplayString)
 				{
@@ -135,15 +141,15 @@ namespace UE::MultiUserClient::Replication::MultiStreamColumns
 			FUnifiedClientView& ClientView;
 		};
 		
-		check(MultiStreamEditor.IsBound() || MultiStreamEditor.IsSet());
+		check(InMultiStreamEditor.IsBound() || InMultiStreamEditor.IsSet());
 		return {
 			TReplicationColumnDelegates<FPropertyTreeRowContext>::FCreateColumn::CreateLambda(
-				[MultiStreamEditor = MoveTemp(MultiStreamEditor), &ClientView]()
+				[MultiStreamEditor = MoveTemp(InMultiStreamEditor), &InClientView, &InViewOptions]()
 				{
-					return MakeShared<FPropertyColumn_AssignPropertyColumn>(MultiStreamEditor, ClientView);
+					return MakeShared<FPropertyColumn_AssignPropertyColumn>(MultiStreamEditor, InClientView, InViewOptions);
 				}),
 			AssignPropertyColumnId,
-			{ ColumnsSortPriority }
+			{ InColumnsSortPriority }
 		};
 	}
 }
diff --git a/Engine/Plugins/Developer/Concert/ConcertApp/MultiUserClient/Source/MultiUserClient/Private/Widgets/ActiveSession/Replication/Client/Multi/Columns/AssignProperty/SAssignPropertyComboBox.cpp b/Engine/Plugins/Developer/Concert/ConcertApp/MultiUserClient/Source/MultiUserClient/Private/Widgets/ActiveSession/Replication/Client/Multi/Columns/AssignProperty/SAssignPropertyComboBox.cpp
--- a/Engine/Plugins/Developer/Concert/ConcertApp/MultiUserClient/Source/MultiUserClient/Private/Widgets/ActiveSession/Replication/Client/Multi/Columns/AssignProperty/SAssignPropertyComboBox.cpp
+++ b/Engine/Plugins/Developer/Concert/ConcertApp/MultiUserClient/Source/MultiUserClient/Private/Widgets/ActiveSession/Replication/Client/Multi/Columns/AssignProperty/SAssignPropertyComboBox.cpp
@@ -3,7 +3,6 @@
 #include "SAssignPropertyComboBox.h"
 
 #include "AssignPropertyModel.h"
-#include "IConcertClient.h"
 #include "Replication/Client/Online/OnlineClient.h"
 #include "Replication/Client/Online/OnlineClientManager.h"
 #include "Replication/Client/UnifiedClientViewExtensions.h"
@@ -26,32 +25,14 @@ namespace UE::MultiUserClient::Replication::MultiStreamColumns
 {
 	namespace AssignPropertyComboBox
 	{
-		template<typename TLambda> requires std::is_invocable_r_v<EBreakBehavior, TLambda, const FGuid&>
-		static void ForEachDisplayedClient(
-			const FUnifiedClientView& ClientView,
-			const FConcertPropertyChain& DisplayedProperty,
-			const TArray<TSoftObjectPtr<>>& EditedObjects,
-			TLambda&& Callback
-			)
-		{
-			for (const TSoftObjectPtr<>& Object : EditedObjects)
-			{
-				ClientView.GetStreamCache().EnumerateClientsWithObjectAndProperty(
-					Object.ToSoftObjectPath(),
-					DisplayedProperty,
-					[&Callback](const FGuid& ClientId){ return Callback(ClientId); }
-				);
-			}
-		}
-		
 		static TArray<FGuid> GetDisplayedClients(
-			const FUnifiedClientView& ClientView,
+			const FAssignPropertyModel& Model,
 			const FConcertPropertyChain& DisplayedProperty,
 			const TArray<TSoftObjectPtr<>>& EditedObjects
 			)
 		{
 			TArray<FGuid> Clients;
-			ForEachDisplayedClient(ClientView, DisplayedProperty, EditedObjects,
+			Model.ForEachAssignedClient(DisplayedProperty, EditedObjects,
 				[&Clients](const FGuid& ClientId){ Clients.AddUnique(ClientId); return EBreakBehavior::Continue; }
 				);
 			return Clients;
@@ -59,12 +40,13 @@ namespace UE::MultiUserClient::Replication::MultiStreamColumns
 	}
 	
 	TOptional<FString> SAssignPropertyComboBox::GetDisplayString(
+		const FAssignPropertyModel& Model,
 		const FUnifiedClientView& ClientView,
 		const FConcertPropertyChain& DisplayedProperty,
 		const TArray<TSoftObjectPtr<>>& EditedObjects)
 	{
 		using SWidgetType = ConcertSharedSlate::SHorizontalClientList;
-		const TArray<FGuid> Clients = AssignPropertyComboBox::GetDisplayedClients(ClientView, DisplayedProperty, EditedObjects);
+		const TArray<FGuid> Clients = AssignPropertyComboBox::GetDisplayedClients(Model, DisplayedProperty, EditedObjects);
 		
 		const ConcertSharedSlate::FGetClientParenthesesContent GetParenthesesContent =
 			MakeLocalAndOfflineParenthesesContentGetter(ClientView);
@@ -121,7 +103,7 @@ namespace UE::MultiUserClient::Replication::MultiStreamColumns
 	void SAssignPropertyComboBox::RefreshContentBoxContent() const
 	{
 		ClientListWidget->RefreshList(
-			AssignPropertyComboBox::GetDisplayedClients(*ClientView, Property, EditedObjects)
+			AssignPropertyComboBox::GetDisplayedClients(*Model, Property, EditedObjects)
 			);
 	}
 
@@ -226,7 +208,7 @@ namespace UE::MultiUserClient::Replication::MultiStreamColumns
 	{
 		bool bHasOfflineClients = false;
 		int32 NumClients = 0;
-		AssignPropertyComboBox::ForEachDisplayedClient(*ClientView, Property, EditedObjects,
+		Model->ForEachAssignedClient(Property, EditedObjects,
 			[this, &bHasOfflineClients, &NumClients](const FGuid& ClientId)
 			{
 				++NumClients;
diff --git a/Engine/Plugins/Developer/Concert/ConcertApp/MultiUserClient/Source/MultiUserClient/Private/Widgets/ActiveSession/Replication/Client/Multi/Columns/AssignProperty/SAssignPropertyComboBox.h b/Engine/Plugins/Developer/Concert/ConcertApp/MultiUserClient/Source/MultiUserClient/Private/Widgets/ActiveSession/Replication/Client/Multi/Columns/AssignProperty/SAssignPropertyComboBox.h
--- a/Engine/Plugins/Developer/Concert/ConcertApp/MultiUserClient/Source/MultiUserClient/Private/Widgets/ActiveSession/Replication/Client/Multi/Columns/AssignProperty/SAssignPropertyComboBox.h
+++ b/Engine/Plugins/Developer/Concert/ConcertApp/MultiUserClient/Source/MultiUserClient/Private/Widgets/ActiveSession/Replication/Client/Multi/Columns/AssignProperty/SAssignPropertyComboBox.h
@@ -34,6 +34,7 @@ namespace UE::MultiUserClient::Replication::MultiStreamColumns
 
 		/** @return The display string this widget would have with the given state. If unset, no clients are displayed in the combobox.*/
 		static TOptional<FString> GetDisplayString(
+			const FAssignPropertyModel& Model,
 			const FUnifiedClientView& ClientView,
 			const FConcertPropertyChain& DisplayedProperty,
 			const TArray<TSoftObjectPtr<>>& EditedObjects
diff --git a/Engine/Plugins/Developer/Concert/ConcertApp/MultiUserClient/Source/MultiUserClient/Private/Widgets/ActiveSession/Replication/Client/Multi/Columns/AssignedClients/AssignedClientsModel.cpp b/Engine/Plugins/Developer/Concert/ConcertApp/MultiUserClient/Source/MultiUserClient/Private/Widgets/ActiveSession/Replication/Client/Multi/Columns/AssignedClients/AssignedClientsModel.cpp
--- a/Engine/Plugins/Developer/Concert/ConcertApp/MultiUserClient/Source/MultiUserClient/Private/Widgets/ActiveSession/Replication/Client/Multi/Columns/AssignedClients/AssignedClientsModel.cpp
+++ b/Engine/Plugins/Developer/Concert/ConcertApp/MultiUserClient/Source/MultiUserClient/Private/Widgets/ActiveSession/Replication/Client/Multi/Columns/AssignedClients/AssignedClientsModel.cpp
@@ -10,21 +10,27 @@
 
 #include <type_traits>
 
+#include "Widgets/ActiveSession/Replication/Client/Multi/ViewOptions/MultiViewOptions.h"
+
 namespace UE::MultiUserClient::Replication::MultiStreamColumns
 {
 	FAssignedClientsModel::FAssignedClientsModel(
-		const ConcertSharedSlate::IObjectHierarchyModel& ObjectHierarchy,
-		FUnifiedStreamCache& InStreamCache
+		const ConcertSharedSlate::IObjectHierarchyModel& InObjectHierarchy,
+		FUnifiedStreamCache& InStreamCache,
+		FMultiViewOptions& InViewOptions 
 		)
-		: ObjectHierarchy(ObjectHierarchy)
+		: ObjectHierarchy(InObjectHierarchy)
 		, StreamCache(InStreamCache)
+		, ViewOptions(InViewOptions)
 	{
 		StreamCache.OnCacheChanged().AddRaw(this, &FAssignedClientsModel::BroadcastOwnershipChanged);
+		ViewOptions.OnOptionsChanged().AddRaw(this, &FAssignedClientsModel::BroadcastOwnershipChanged);
 	}
 
 	FAssignedClientsModel::~FAssignedClientsModel()
 	{
 		StreamCache.OnCacheChanged().RemoveAll(this);
+		ViewOptions.OnOptionsChanged().RemoveAll(this);
 	}
 
 	TArray<FGuid> FAssignedClientsModel::GetAssignedClients(const FSoftObjectPath& ObjectPath) const
@@ -33,11 +39,14 @@ namespace UE::MultiUserClient::Replication::MultiStreamColumns
 
 		const auto ProcessObject = [this, &ClientsWithOwnership](const FSoftObjectPath& ObjectPath)
 		{
+			const EClientEnumerationMode Mode = ViewOptions.ShouldShowOfflineClients()
+				? EClientEnumerationMode::SkipOfflineClientsThatFullyOverlapWithOnlineClients
+				: EClientEnumerationMode::SkipOfflineClients;
 			StreamCache.EnumerateClientsWithObject(ObjectPath, [&ClientsWithOwnership](const FGuid& ClientId)
 			{
 				ClientsWithOwnership.AddUnique(ClientId);
 				return EBreakBehavior::Continue;
-			});
+			}, Mode);
 		};
 			
 		ProcessObject(ObjectPath);
diff --git a/Engine/Plugins/Developer/Concert/ConcertApp/MultiUserClient/Source/MultiUserClient/Private/Widgets/ActiveSession/Replication/Client/Multi/Columns/AssignedClients/AssignedClientsModel.h b/Engine/Plugins/Developer/Concert/ConcertApp/MultiUserClient/Source/MultiUserClient/Private/Widgets/ActiveSession/Replication/Client/Multi/Columns/AssignedClients/AssignedClientsModel.h
--- a/Engine/Plugins/Developer/Concert/ConcertApp/MultiUserClient/Source/MultiUserClient/Private/Widgets/ActiveSession/Replication/Client/Multi/Columns/AssignedClients/AssignedClientsModel.h
+++ b/Engine/Plugins/Developer/Concert/ConcertApp/MultiUserClient/Source/MultiUserClient/Private/Widgets/ActiveSession/Replication/Client/Multi/Columns/AssignedClients/AssignedClientsModel.h
@@ -7,21 +7,15 @@
 #include "HAL/Platform.h"
 #include "Templates/UnrealTemplate.h"
 
-namespace UE::MultiUserClient::Replication
-{
-	class FUnifiedStreamCache;
-}
-
-namespace UE::MultiUserClient::Replication
-{
-	class FUnifiedClientView;
-}
 
 struct FGuid;
 
 namespace UE::ConcertSharedSlate { class IObjectHierarchyModel; }
 namespace UE::MultiUserClient::Replication
 {
+	class FMultiViewOptions;
+	class FUnifiedStreamCache;
+	class FUnifiedClientView;
 	class FOnlineClientManager;
 	class FReassignObjectPropertiesLogic;
 }
@@ -34,8 +28,9 @@ namespace UE::MultiUserClient::Replication::MultiStreamColumns
 	public:
 		
 		FAssignedClientsModel(
-			const ConcertSharedSlate::IObjectHierarchyModel& ObjectHierarchy UE_LIFETIMEBOUND,
-			FUnifiedStreamCache& InStreamCache UE_LIFETIMEBOUND
+			const ConcertSharedSlate::IObjectHierarchyModel& InObjectHierarchy UE_LIFETIMEBOUND,
+			FUnifiedStreamCache& InStreamCache UE_LIFETIMEBOUND,
+			FMultiViewOptions& InViewOptions UE_LIFETIMEBOUND
 			);
 		~FAssignedClientsModel();
 
@@ -51,6 +46,9 @@ namespace UE::MultiUserClient::Replication::MultiStreamColumns
 		const ConcertSharedSlate::IObjectHierarchyModel& ObjectHierarchy;
 		/** Holds the stream content of both online and offline clients. */
 		FUnifiedStreamCache& StreamCache;
+
+		/** Controls whether offline clients should be considered. */
+		FMultiViewOptions& ViewOptions;
 		
 		/** Broadcast when the result of GetAssignedClients may have changed. */
 		FSimpleMulticastDelegate OnOwnershipChangedDelegate;
diff --git a/Engine/Plugins/Developer/Concert/ConcertApp/MultiUserClient/Source/MultiUserClient/Private/Widgets/ActiveSession/Replication/Client/Multi/Columns/AssignedClients/MultiStreamColumns_AssignedClients.cpp b/Engine/Plugins/Developer/Concert/ConcertApp/MultiUserClient/Source/MultiUserClient/Private/Widgets/ActiveSession/Replication/Client/Multi/Columns/AssignedClients/MultiStreamColumns_AssignedClients.cpp
--- a/Engine/Plugins/Developer/Concert/ConcertApp/MultiUserClient/Source/MultiUserClient/Private/Widgets/ActiveSession/Replication/Client/Multi/Columns/AssignedClients/MultiStreamColumns_AssignedClients.cpp
+++ b/Engine/Plugins/Developer/Concert/ConcertApp/MultiUserClient/Source/MultiUserClient/Private/Widgets/ActiveSession/Replication/Client/Multi/Columns/AssignedClients/MultiStreamColumns_AssignedClients.cpp
@@ -21,11 +21,12 @@ namespace UE::MultiUserClient::Replication::MultiStreamColumns
 	const FName AssignedClientsColumnId(TEXT("AssignedClientsColumn"));
 
 	ConcertSharedSlate::FObjectColumnEntry AssignedClientsColumn(
-		const TSharedRef<IConcertClient>& ConcertClient,
-		TAttribute<TSharedPtr<ConcertSharedSlate::IMultiReplicationStreamEditor>> MultiStreamModelAttribute,
-		const ConcertSharedSlate::IObjectHierarchyModel& ObjectHierarchy,
-		FUnifiedClientView& ClientView,
-		const int32 ColumnsSortPriority
+		const TSharedRef<IConcertClient>& InConcertClient,
+		TAttribute<TSharedPtr<ConcertSharedSlate::IMultiReplicationStreamEditor>> InMultiStreamModelAttribute,
+		const ConcertSharedSlate::IObjectHierarchyModel& InObjectHierarchy,
+		FUnifiedClientView& InClientView,
+		FMultiViewOptions& InViewOptions,
+		const int32 InColumnsSortPriority
 		)
 	{
 		class FObjectColumn_AssignedClients : public ConcertSharedSlate::IObjectTreeColumn
@@ -33,15 +34,16 @@ namespace UE::MultiUserClient::Replication::MultiStreamColumns
 		public:
 
 			FObjectColumn_AssignedClients(
-				const TSharedRef<IConcertClient>& ConcertClient,
-				TAttribute<TSharedPtr<ConcertSharedSlate::IMultiReplicationStreamEditor>> MultiStreamModelAttribute,
-				const ConcertSharedSlate::IObjectHierarchyModel& ObjectHierarchy UE_LIFETIMEBOUND,
-				FUnifiedClientView& ClientView UE_LIFETIMEBOUND
+				const TSharedRef<IConcertClient>& InConcertClient,
+				TAttribute<TSharedPtr<ConcertSharedSlate::IMultiReplicationStreamEditor>> InMultiStreamModelAttribute,
+				const ConcertSharedSlate::IObjectHierarchyModel& InObjectHierarchy UE_LIFETIMEBOUND,
+				FUnifiedClientView& InClientView UE_LIFETIMEBOUND,
+				FMultiViewOptions& InViewOptions UE_LIFETIMEBOUND
 				)
-				: ConcertClient(ConcertClient)
-				, MultiStreamModelAttribute(MoveTemp(MultiStreamModelAttribute))
-				, Model(ObjectHierarchy, ClientView.GetStreamCache())
-				, ClientView(ClientView)
+				: ConcertClient(InConcertClient)
+				, MultiStreamModelAttribute(MoveTemp(InMultiStreamModelAttribute))
+				, Model(InObjectHierarchy, InClientView.GetStreamCache(), InViewOptions)
+				, ClientView(InClientView)
 			{}
 			
 			virtual SHeaderRow::FColumn::FArguments CreateHeaderRowArgs() const override
@@ -112,12 +114,14 @@ namespace UE::MultiUserClient::Replication::MultiStreamColumns
 
 		return {
 			ConcertSharedSlate::TReplicationColumnDelegates<ConcertSharedSlate::FObjectTreeRowContext>::FCreateColumn::CreateLambda(
-				[ConcertClient, MultiStreamModelAttribute = MoveTemp(MultiStreamModelAttribute), &ObjectHierarchy, &ClientView]()
+				[InConcertClient, MultiStreamModelAttribute = MoveTemp(InMultiStreamModelAttribute), &InObjectHierarchy, &InClientView, &InViewOptions]
 				{
-					return MakeShared<FObjectColumn_AssignedClients>(ConcertClient, MultiStreamModelAttribute, ObjectHierarchy, ClientView);
+					return MakeShared<FObjectColumn_AssignedClients>(
+						InConcertClient, MultiStreamModelAttribute, InObjectHierarchy, InClientView, InViewOptions
+						);
 				}),
 			AssignedClientsColumnId,
-			{ ColumnsSortPriority }
+			{ InColumnsSortPriority }
 		};
 	}
 }
diff --git a/Engine/Plugins/Developer/Concert/ConcertApp/MultiUserClient/Source/MultiUserClient/Private/Widgets/ActiveSession/Replication/Client/Multi/Columns/MultiStreamColumns.h b/Engine/Plugins/Developer/Concert/ConcertApp/MultiUserClient/Source/MultiUserClient/Private/Widgets/ActiveSession/Replication/Client/Multi/Columns/MultiStreamColumns.h
--- a/Engine/Plugins/Developer/Concert/ConcertApp/MultiUserClient/Source/MultiUserClient/Private/Widgets/ActiveSession/Replication/Client/Multi/Columns/MultiStreamColumns.h
+++ b/Engine/Plugins/Developer/Concert/ConcertApp/MultiUserClient/Source/MultiUserClient/Private/Widgets/ActiveSession/Replication/Client/Multi/Columns/MultiStreamColumns.h
@@ -13,6 +13,7 @@ class IConcertClient;
 
 namespace UE::MultiUserClient::Replication
 {
+	class FMultiViewOptions;
 	class FMuteChangeTracker;
 	class FReassignObjectPropertiesLogic;
 	class FOnlineClientManager;
@@ -44,44 +45,48 @@ namespace UE::MultiUserClient::Replication::MultiStreamColumns
 	/**
 	 * Mutes and unmutes the object and its child objects.
 	 * 
-	 * @param MuteChangeTracker Tells us the mute state and changes it.
-	 * @param ColumnsSortPriority The order relative to the other columns
+	 * @param InMuteChangeTracker Tells us the mute state and changes it.
+	 * @param InColumnsSortPriority The order relative to the other columns
 	 * @return A checkbox with pause and unpause icons.
 	 */
 	ConcertSharedSlate::FObjectColumnEntry MuteToggleColumn(
-		FMuteChangeTracker& MuteChangeTracker UE_LIFETIMEBOUND,
-		int32 ColumnsSortPriority = static_cast<int32>(EColumnSortOrder::MuteToggle)
+		FMuteChangeTracker& InMuteChangeTracker UE_LIFETIMEBOUND,
+		int32 InColumnsSortPriority = static_cast<int32>(EColumnSortOrder::MuteToggle)
 		);
 
 	/**
 	 * A combo box which displays all current owners for an object and allows bulk reassigning properties.
 	 *
-	 * @param ConcertClient Used to look up client names
-	 * @param MultiStreamModelAttribute Used for requesting resorting the column
-	 * @param ObjectHierarchy Used to display clients when a child object is replicated
-	 * @param ClientView Used to access all clients' stream content for display in the column
-	 * @param ColumnsSortPriority The order relative to the other columns
+	 * @param InConcertClient Used to look up client names
+	 * @param InMultiStreamModelAttribute Used for requesting resorting the column
+	 * @param InObjectHierarchy Used to display clients when a child object is replicated
+	 * @param InClientView Used to access all clients' stream content for display in the column
+	 * @param InViewOptions Controls whether offline clients are displayed
+	 * @param InColumnsSortPriority The order relative to the other columns
 	 */
 	ConcertSharedSlate::FObjectColumnEntry AssignedClientsColumn(
-		const TSharedRef<IConcertClient>& ConcertClient,
-		TAttribute<TSharedPtr<ConcertSharedSlate::IMultiReplicationStreamEditor>> MultiStreamModelAttribute,
-		const ConcertSharedSlate::IObjectHierarchyModel& ObjectHierarchy UE_LIFETIMEBOUND,
-		FUnifiedClientView& ClientView UE_LIFETIMEBOUND,
-		int32 ColumnsSortPriority = static_cast<int32>(EColumnSortOrder::ReassignOwnership)
+		const TSharedRef<IConcertClient>& InConcertClient,
+		TAttribute<TSharedPtr<ConcertSharedSlate::IMultiReplicationStreamEditor>> InMultiStreamModelAttribute,
+		const ConcertSharedSlate::IObjectHierarchyModel& InObjectHierarchy UE_LIFETIMEBOUND,
+		FUnifiedClientView& InClientView UE_LIFETIMEBOUND,
+		FMultiViewOptions& InViewOptions UE_LIFETIMEBOUND,
+		int32 InColumnsSortPriority = static_cast<int32>(EColumnSortOrder::ReassignOwnership)
 		);
 	
 	/**
 	 * Creates a property column which assigns the property to the stream selected in the combo box.
 	 * 
-	 * @param MultiStreamEditor Used to determine the selected objects.
-	 * @param ClientView Used to obtain the clients that can be assigned to and display them
-	 * @param ColumnsSortPriority The order relative to the other columns
+	 * @param InMultiStreamEditor Used to determine the selected objects.
+	 * @param InClientView Used to obtain the clients that can be assigned to and display them
+	 * @param InViewOptions Controls whether offline clients are displayed
+	 * @param InColumnsSortPriority The order relative to the other columns
 	 * 
 	 * @return A column that spawns a combo box for assigning properties
 	 */
 	ConcertSharedSlate::FPropertyColumnEntry AssignPropertyColumn(
-		TAttribute<TSharedPtr<ConcertSharedSlate::IMultiReplicationStreamEditor>> MultiStreamEditor,
-		FUnifiedClientView& ClientView UE_LIFETIMEBOUND,
-		const int32 ColumnsSortPriority = static_cast<int32>(EColumnSortOrder::AssignPropertyColumn)
+		TAttribute<TSharedPtr<ConcertSharedSlate::IMultiReplicationStreamEditor>> InMultiStreamEditor,
+		FUnifiedClientView& InClientView UE_LIFETIMEBOUND,
+		FMultiViewOptions& InViewOptions UE_LIFETIMEBOUND,
+		const int32 InColumnsSortPriority = static_cast<int32>(EColumnSortOrder::AssignPropertyColumn)
 		);
 }
diff --git a/Engine/Plugins/Developer/Concert/ConcertApp/MultiUserClient/Source/MultiUserClient/Private/Widgets/ActiveSession/Replication/Client/Multi/MultiStreamModel.cpp b/Engine/Plugins/Developer/Concert/ConcertApp/MultiUserClient/Source/MultiUserClient/Private/Widgets/ActiveSession/Replication/Client/Multi/MultiStreamModel.cpp
--- a/Engine/Plugins/Developer/Concert/ConcertApp/MultiUserClient/Source/MultiUserClient/Private/Widgets/ActiveSession/Replication/Client/Multi/MultiStreamModel.cpp
+++ b/Engine/Plugins/Developer/Concert/ConcertApp/MultiUserClient/Source/MultiUserClient/Private/Widgets/ActiveSession/Replication/Client/Multi/MultiStreamModel.cpp
@@ -7,6 +7,7 @@
 #include "Replication/Client/Online/OnlineClientManager.h"
 #include "Replication/Editor/Model/IEditableReplicationStreamModel.h"
 #include "Selection/ISelectionModel.h"
+#include "ViewOptions/MultiViewOptions.h"
 
 namespace UE::MultiUserClient::Replication
 {
@@ -14,24 +15,27 @@ namespace UE::MultiUserClient::Replication
 		IOnlineClientSelectionModel& InOnlineClientSelectionModel,
 		IOfflineClientSelectionModel& InOfflineClientSelectionModel,
 		FOnlineClientManager& InOnlineClientManager,
-		FOfflineClientManager& InOfflineClientManager
+		FOfflineClientManager& InOfflineClientManager,
+		FMultiViewOptions& InViewOptions
 		)
 		: OnlineClientSelectionModel(InOnlineClientSelectionModel)
 		, OfflineClientSelectionModel(InOfflineClientSelectionModel)
 		, OnlineClientManager(InOnlineClientManager)
 		, OfflineClientManager(InOfflineClientManager)
+		, ViewOptions(InViewOptions)
 	{
 		OnlineClientSelectionModel.OnSelectionChanged().AddRaw(this, &FMultiStreamModel::RebuildOnlineClients);
 		OfflineClientSelectionModel.OnSelectionChanged().AddRaw(this, &FMultiStreamModel::RebuildOfflineClients);
+		ViewOptions.OnOptionsChanged().AddRaw(this, &FMultiStreamModel::RebuildClients);
 		
-		RebuildOnlineClients();
-		RebuildOfflineClients();
+		RebuildClients();
 	}
 
 	FMultiStreamModel::~FMultiStreamModel()
 	{
 		UnsubscribeFromOnlineClients();
 		UnsubscribeFromOfflineClients();
+		ViewOptions.OnOptionsChanged().RemoveAll(this);
 	}
 
 	void FMultiStreamModel::ForEachDisplayedOnlineClient(TFunctionRef<EBreakBehavior(const FOnlineClient*)> ProcessClient) const
@@ -88,6 +92,12 @@ namespace UE::MultiUserClient::Replication
 	{
 		// It is not safe to iterate through our cached client array because it may contain stale clients that were just removed.
 		UnsubscribeFromOfflineClients();
+		if (!ViewOptions.ShouldShowOfflineClients())
+		{
+			CachedOfflineClients.Reset();
+			OnStreamSetChangedDelegate.Broadcast();
+			return;
+		}
 		
 		TSet<const FOfflineClient*> OfflineClients;
 		OfflineClientSelectionModel.ForEachItem([this, &OfflineClients](FOfflineClient& Client)
diff --git a/Engine/Plugins/Developer/Concert/ConcertApp/MultiUserClient/Source/MultiUserClient/Private/Widgets/ActiveSession/Replication/Client/Multi/MultiStreamModel.h b/Engine/Plugins/Developer/Concert/ConcertApp/MultiUserClient/Source/MultiUserClient/Private/Widgets/ActiveSession/Replication/Client/Multi/MultiStreamModel.h
--- a/Engine/Plugins/Developer/Concert/ConcertApp/MultiUserClient/Source/MultiUserClient/Private/Widgets/ActiveSession/Replication/Client/Multi/MultiStreamModel.h
+++ b/Engine/Plugins/Developer/Concert/ConcertApp/MultiUserClient/Source/MultiUserClient/Private/Widgets/ActiveSession/Replication/Client/Multi/MultiStreamModel.h
@@ -10,15 +10,12 @@
 #include "Templates/Function.h"
 #include "Templates/SharedPointer.h"
 
-namespace UE::MultiUserClient::Replication
-{
-	class FOfflineClientManager;
-}
-
 enum class EBreakBehavior : uint8;
 
 namespace UE::MultiUserClient::Replication
 {
+	class FMultiViewOptions;
+	class FOfflineClientManager;
 	class FOnlineClient;
 	class FOnlineClientManager;
 
@@ -36,7 +33,8 @@ namespace UE::MultiUserClient::Replication
 			IOnlineClientSelectionModel& InOnlineClientSelectionModel UE_LIFETIMEBOUND,
 			IOfflineClientSelectionModel& InOfflineClientSelectionModel UE_LIFETIMEBOUND,
 			FOnlineClientManager& InOnlineClientManager UE_LIFETIMEBOUND,
-			FOfflineClientManager& InOfflineClientManager UE_LIFETIMEBOUND
+			FOfflineClientManager& InOfflineClientManager UE_LIFETIMEBOUND,
+			FMultiViewOptions& InViewOptions UE_LIFETIMEBOUND
 			);
 		~FMultiStreamModel();
 
@@ -62,14 +60,22 @@ namespace UE::MultiUserClient::Replication
 		/** Used to clean up subscriptions when client list changes. */
 		FOfflineClientManager& OfflineClientManager;
 
+		/** Determines whether offline clients should be shown. */
+		FMultiViewOptions& ViewOptions;
+
 		TSet<const FOnlineClient*> CachedOnlineClients;
 		TSet<const FOfflineClient*> CachedOfflineClients;
 
 		FOnStreamExternallyChanged OnStreamsExternallyChanged;
 		FOnStreamSetChanged OnStreamSetChangedDelegate;
-		
+
 		void RebuildOnlineClients();
 		void RebuildOfflineClients();
+		void RebuildClients()
+		{
+			RebuildOnlineClients();
+			RebuildOfflineClients();
+		}
 		
 		void UnsubscribeFromOnlineClients() const;
 		void UnsubscribeFromOfflineClients() const;
diff --git a/Engine/Plugins/Developer/Concert/ConcertApp/MultiUserClient/Source/MultiUserClient/Private/Widgets/ActiveSession/Replication/Client/Multi/SMultiClientView.cpp b/Engine/Plugins/Developer/Concert/ConcertApp/MultiUserClient/Source/MultiUserClient/Private/Widgets/ActiveSession/Replication/Client/Multi/SMultiClientView.cpp
--- a/Engine/Plugins/Developer/Concert/ConcertApp/MultiUserClient/Source/MultiUserClient/Private/Widgets/ActiveSession/Replication/Client/Multi/SMultiClientView.cpp
+++ b/Engine/Plugins/Developer/Concert/ConcertApp/MultiUserClient/Source/MultiUserClient/Private/Widgets/ActiveSession/Replication/Client/Multi/SMultiClientView.cpp
@@ -18,6 +18,7 @@
 #include "Replication/Editor/View/IReplicationStreamEditor.h"
 #include "Replication/Editor/Model/ObjectSource/ActorSelectionSourceModel.h"
 #include "Replication/Stream/Discovery/MultiUserStreamExtender.h"
+#include "ViewOptions/SMultiViewOptions.h"
 #include "Widgets/ActiveSession/Replication/Client/Context/ContextMenuUtils.h"
 #include "Widgets/ActiveSession/Replication/Client/Multi/Columns/MultiStreamColumns.h"
 #include "Widgets/ActiveSession/Replication/Client/PropertySelection/SPropertySelectionComboButton.h"
@@ -47,7 +48,7 @@ namespace UE::MultiUserClient::Replication
 		OnlineClientSelectionModel = &InOnlineClientSelectionModel;
 		OfflineClientSelectionModel = &InOfflineClientSelectionModel;
 		StreamModel = MakeShared<FMultiStreamModel>(
-			InOnlineClientSelectionModel, InOfflineClientSelectionModel, *OnlineClientManager, *OfflineClientManager
+			InOnlineClientSelectionModel, InOfflineClientSelectionModel, *OnlineClientManager, *OfflineClientManager, ViewOptions
 			);
 		
 		OnlineClientManager->OnRemoteClientsChanged().AddSP(this, &SMultiClientView::RebuildClientSubscriptions);
@@ -73,20 +74,17 @@ namespace UE::MultiUserClient::Replication
 			.ForEachObjectInStream(this, &SMultiClientView::EnumerateReplicatedObjectsInStreams)
 			);
 
-		RebuildClientSubscriptions();
-
 		// Changing worlds affects what things are displayed in the editor.
-		HideObjectsNotInEditorWorld.OnRefreshObjects().AddLambda([this]()
-		{
-			StreamEditor->GetEditorBase().Refresh();
-		});
+		HideObjectsNotInEditorWorld.OnRefreshObjects().AddRaw(this, &SMultiClientView::RefreshUI);
+		ViewOptions.OnOptionsChanged().AddRaw(this, &SMultiClientView::RefreshUI);
+		RebuildClientSubscriptions();
 	}
 
 	SMultiClientView::~SMultiClientView()
 	{
 		OnlineClientManager->OnRemoteClientsChanged().RemoveAll(this);
-		CleanClientSubscriptions();
 		UserSelectedProperties->OnPropertySelectionChanged().RemoveAll(this);
+		CleanClientSubscriptions();
 	}
 
 	TSharedRef<SWidget> SMultiClientView::CreateEditorContent(
@@ -119,7 +117,7 @@ namespace UE::MultiUserClient::Replication
 			{
 				ReplicationColumns::Property::LabelColumn(),
 				MultiStreamColumns::AssignPropertyColumn(
-					MultiStreamEditorAttribute, *InMultiUserReplicationManager.GetUnifiedClientView()
+					MultiStreamEditorAttribute, *InMultiUserReplicationManager.GetUnifiedClientView(), ViewOptions
 					)
 			},
 			.CreateCategoryRow = CreateDefaultCategoryGenerator(NameModel),
@@ -159,13 +157,13 @@ namespace UE::MultiUserClient::Replication
 			{
 				MultiStreamColumns::MuteToggleColumn(MuteManager.GetChangeTracker()),
 				MultiStreamColumns::AssignedClientsColumn(
-					InConcertClient, MultiStreamEditorAttribute, *ObjectHierarchy, *InMultiUserReplicationManager.GetUnifiedClientView()
+					InConcertClient, MultiStreamEditorAttribute, *ObjectHierarchy, *InMultiUserReplicationManager.GetUnifiedClientView(), ViewOptions
 					)
 			},
 			.ShouldDisplayObjectDelegate = FShouldDisplayObject::CreateSP(this, &SMultiClientView::ShouldDisplayObject),
 			.MakeObjectRowOverlayWidgetDelegate = FMakeObjectRowOverlayWidget::CreateSP(this, &SMultiClientView::MakeObjectRowOverlayWidget)
 		};
-		ViewerParams.RightOfObjectSearchBar.Widget = SNew(SPresetComboButton, *InConcertClient, *InMultiUserReplicationManager.GetPresetManager());
+		ViewerParams.RightOfObjectSearchBar.Widget = CreateRightOfSearchBarContent(InConcertClient, InMultiUserReplicationManager);
 		StreamEditor = CreateBaseMultiStreamEditor(MoveTemp(Params), MoveTemp(ViewerParams));
 		check(StreamEditor);
 		
@@ -183,6 +181,26 @@ namespace UE::MultiUserClient::Replication
 			];
 	}
 
+	TSharedRef<SWidget> SMultiClientView::CreateRightOfSearchBarContent(
+		const TSharedRef<IConcertClient>& InConcertClient,
+		FMultiUserReplicationManager& InMultiUserReplicationManager
+		)
+	{
+		return SNew(SHorizontalBox)
+
+			+SHorizontalBox::Slot()
+			.AutoWidth()
+			[
+				SNew(SPresetComboButton, *InConcertClient, *InMultiUserReplicationManager.GetPresetManager())
+			]
+
+			+SHorizontalBox::Slot()
+			.AutoWidth()
+			[
+				SNew(SMultiViewOptions, ViewOptions)
+			];
+	}
+
 	TSet<FGuid> SMultiClientView::GetReplicatableClientIds() const
 	{
 		TSet<FGuid> ClientIds;
diff --git a/Engine/Plugins/Developer/Concert/ConcertApp/MultiUserClient/Source/MultiUserClient/Private/Widgets/ActiveSession/Replication/Client/Multi/SMultiClientView.h b/Engine/Plugins/Developer/Concert/ConcertApp/MultiUserClient/Source/MultiUserClient/Private/Widgets/ActiveSession/Replication/Client/Multi/SMultiClientView.h
--- a/Engine/Plugins/Developer/Concert/ConcertApp/MultiUserClient/Source/MultiUserClient/Private/Widgets/ActiveSession/Replication/Client/Multi/SMultiClientView.h
+++ b/Engine/Plugins/Developer/Concert/ConcertApp/MultiUserClient/Source/MultiUserClient/Private/Widgets/ActiveSession/Replication/Client/Multi/SMultiClientView.h
@@ -6,6 +6,8 @@
 #include "Replication/Editor/Model/ObjectSource/IObjectSourceModel.h"
 #include "Replication/Editor/UnrealEditor/HideObjectsNotInWorldLogic.h"
 #include "Selection/SelectionModelFwd.h"
+#include "ViewOptions/MultiViewOptions.h"
+
 #include "Widgets/DeclarativeSyntaxSupport.h"
 #include "Widgets/SCompoundWidget.h"
 
@@ -73,12 +75,16 @@ namespace UE::MultiUserClient::Replication
 		/** Displays the properties for the objects displayed in the top view. */
 		TSharedPtr<ConcertSharedSlate::IMultiObjectPropertyAssignmentView> PropertyAssignmentView;
 
+		/** Controls the content shown in the UI. */
+		FMultiViewOptions ViewOptions;
+
 		/** This logic helps us decide whether an object should be displayed and lets us know that the object list needs to be refreshed (e.g. due to world change). */
 		ConcertClientSharedSlate::FHideObjectsNotInWorldLogic HideObjectsNotInEditorWorld;
 
 		/** Creates this widget's editor content */
 		TSharedRef<SWidget> CreateEditorContent(const TSharedRef<IConcertClient>& InConcertClient, FMultiUserReplicationManager& InMultiUserReplicationManager);
 		TSharedRef<SWidget> CreateNoPropertiesWarning() const;
+		TSharedRef<SWidget> CreateRightOfSearchBarContent(const TSharedRef<IConcertClient>& InConcertClient, FMultiUserReplicationManager& InMultiUserReplicationManager);
 
 		// SClientToolbar attributes
 		/** @return Gets the clients that may be replicating */
diff --git a/Engine/Plugins/Developer/Concert/ConcertApp/MultiUserClient/Source/MultiUserClient/Private/Widgets/ActiveSession/Replication/Client/Multi/ViewOptions/MultiViewOptions.cpp b/Engine/Plugins/Developer/Concert/ConcertApp/MultiUserClient/Source/MultiUserClient/Private/Widgets/ActiveSession/Replication/Client/Multi/ViewOptions/MultiViewOptions.cpp
new file mode 100644
--- /dev/null
+++ b/Engine/Plugins/Developer/Concert/ConcertApp/MultiUserClient/Source/MultiUserClient/Private/Widgets/ActiveSession/Replication/Client/Multi/ViewOptions/MultiViewOptions.cpp
@@ -0,0 +1,8 @@
+﻿// Copyright Epic Games, Inc. All Rights Reserved.
+
+#include "MultiViewOptions.h"
+
+namespace UE::MultiUserClient::Replication
+{
+	
+}
\ No newline at end of file
diff --git a/Engine/Plugins/Developer/Concert/ConcertApp/MultiUserClient/Source/MultiUserClient/Private/Widgets/ActiveSession/Replication/Client/Multi/ViewOptions/MultiViewOptions.h b/Engine/Plugins/Developer/Concert/ConcertApp/MultiUserClient/Source/MultiUserClient/Private/Widgets/ActiveSession/Replication/Client/Multi/ViewOptions/MultiViewOptions.h
new file mode 100644
--- /dev/null
+++ b/Engine/Plugins/Developer/Concert/ConcertApp/MultiUserClient/Source/MultiUserClient/Private/Widgets/ActiveSession/Replication/Client/Multi/ViewOptions/MultiViewOptions.h
@@ -0,0 +1,37 @@
+﻿// Copyright Epic Games, Inc. All Rights Reserved.
+
+#pragma once
+
+#include "Delegates/Delegate.h"
+
+namespace UE::MultiUserClient::Replication
+{
+	/** The view options for the SMultiClientView. */
+	class FMultiViewOptions
+	{
+	public:
+
+		bool ShouldShowOfflineClients() const { return bShowOfflineClients; }
+		void SetShouldShowOfflineClients(bool bValue)
+		{
+			if (bShowOfflineClients != bValue)
+			{
+				bShowOfflineClients = bValue;
+				OnOptionsChangedDelegate.Broadcast();
+			}
+		}
+		void ToggleShouldShowOfflineClients() { SetShouldShowOfflineClients(!ShouldShowOfflineClients()); }
+
+		/** Broadcasts when a view option changes. */
+		FSimpleMulticastDelegate& OnOptionsChanged() { return OnOptionsChangedDelegate; }
+		
+	private:
+
+		/** Whether offline clients should be shown. */
+		bool bShowOfflineClients = true;
+
+		/** Broadcasts when a view option changes. */
+		FSimpleMulticastDelegate OnOptionsChangedDelegate;
+	};
+}
+
diff --git a/Engine/Plugins/Developer/Concert/ConcertApp/MultiUserClient/Source/MultiUserClient/Private/Widgets/ActiveSession/Replication/Client/Multi/ViewOptions/SMultiViewOptions.cpp b/Engine/Plugins/Developer/Concert/ConcertApp/MultiUserClient/Source/MultiUserClient/Private/Widgets/ActiveSession/Replication/Client/Multi/ViewOptions/SMultiViewOptions.cpp
new file mode 100644
--- /dev/null
+++ b/Engine/Plugins/Developer/Concert/ConcertApp/MultiUserClient/Source/MultiUserClient/Private/Widgets/ActiveSession/Replication/Client/Multi/ViewOptions/SMultiViewOptions.cpp
@@ -0,0 +1,55 @@
+﻿// Copyright Epic Games, Inc. All Rights Reserved.
+
+#include "SMultiViewOptions.h"
+
+#include "Framework/MultiBox/MultiBoxBuilder.h"
+#include "MultiViewOptions.h"
+#include "Styling/AppStyle.h"
+#include "Styling/SlateColor.h"
+#include "Widgets/Images/SImage.h"
+#include "Widgets/Input/SComboButton.h"
+
+#define LOCTEXT_NAMESPACE "SMultiViewOptions"
+
+namespace UE::MultiUserClient::Replication
+{
+	void SMultiViewOptions::Construct(const FArguments& InArgs, FMultiViewOptions& InViewOptions)
+	{
+		ViewOptions = &InViewOptions;
+		ChildSlot
+		[
+			SNew(SComboButton)
+			.ComboButtonStyle(FAppStyle::Get(), "SimpleComboButtonWithIcon") // Use the tool bar item style for this button
+			.OnGetMenuContent(this, &SMultiViewOptions::GetViewButtonContent)
+			.HasDownArrow(false)
+			.ButtonContent()
+			[
+				SNew(SImage)
+				.ColorAndOpacity(FSlateColor::UseForeground())
+				.Image(FAppStyle::Get().GetBrush("Icons.Settings"))
+			]
+		];
+	}
+
+	TSharedRef<SWidget> SMultiViewOptions::GetViewButtonContent() const
+	{
+		FMenuBuilder MenuBuilder(true, nullptr);
+		
+		MenuBuilder.AddMenuEntry(
+			LOCTEXT("ShowOfflineClients.Label", "Show offline clients"),
+			LOCTEXT("ShowOfflineClients.ToolTipText", "Controls whether properties for clients rejoining the session are displayed."),
+			FSlateIcon(),
+			FUIAction(
+				FExecuteAction::CreateLambda([this]{ ViewOptions->ToggleShouldShowOfflineClients(); }),
+				FCanExecuteAction::CreateLambda([]{ return true; }),
+				FIsActionChecked::CreateLambda([this]{ return ViewOptions->ShouldShowOfflineClients(); })
+				),
+			NAME_None,
+			EUserInterfaceActionType::ToggleButton
+		);
+
+		return MenuBuilder.MakeWidget();
+	}
+}
+
+#undef LOCTEXT_NAMESPACE
diff --git a/Engine/Plugins/Developer/Concert/ConcertApp/MultiUserClient/Source/MultiUserClient/Private/Widgets/ActiveSession/Replication/Client/Multi/ViewOptions/SMultiViewOptions.h b/Engine/Plugins/Developer/Concert/ConcertApp/MultiUserClient/Source/MultiUserClient/Private/Widgets/ActiveSession/Replication/Client/Multi/ViewOptions/SMultiViewOptions.h
new file mode 100644
--- /dev/null
+++ b/Engine/Plugins/Developer/Concert/ConcertApp/MultiUserClient/Source/MultiUserClient/Private/Widgets/ActiveSession/Replication/Client/Multi/ViewOptions/SMultiViewOptions.h
@@ -0,0 +1,32 @@
+﻿// Copyright Epic Games, Inc. All Rights Reserved.
+
+#pragma once
+
+#include "HAL/Platform.h"
+#include "Widgets/DeclarativeSyntaxSupport.h"
+#include "Widgets/SCompoundWidget.h"
+
+namespace UE::MultiUserClient::Replication
+{
+	class FMultiViewOptions;
+
+	/** Displays FMultiViewOptions in a combo button. */
+	class SMultiViewOptions : public SCompoundWidget
+    {
+    public:
+
+		SLATE_BEGIN_ARGS(SMultiViewOptions){}
+		SLATE_END_ARGS()
+
+		void Construct(const FArguments& InArgs, FMultiViewOptions& InViewOptions UE_LIFETIMEBOUND);
+
+	private:
+
+		/** The view options being displayed and mutated by this widget. */
+		FMultiViewOptions* ViewOptions = nullptr;
+
+		/** @return The menu content for the combo button */
+		TSharedRef<SWidget> GetViewButtonContent() const;
+    };
+}
+
