Expose Blueprint API for querying offline MU replication client info

diff --git a/Engine/Plugins/Developer/Concert/ConcertApp/MultiUserClient/Source/MultiUserClient/Private/Replication/MultiUserReplicationManager.cpp b/Engine/Plugins/Developer/Concert/ConcertApp/MultiUserClient/Source/MultiUserClient/Private/Replication/MultiUserReplicationManager.cpp
--- a/Engine/Plugins/Developer/Concert/ConcertApp/MultiUserClient/Source/MultiUserClient/Private/Replication/MultiUserReplicationManager.cpp
+++ b/Engine/Plugins/Developer/Concert/ConcertApp/MultiUserClient/Source/MultiUserClient/Private/Replication/MultiUserReplicationManager.cpp
@@ -6,6 +6,8 @@
 #include "IConcertSyncClient.h"
 #include "Replication/ChangeOperationTypes.h"
 #include "Replication/IConcertClientReplicationManager.h"
+#include "Replication/IOfflineReplicationClient.h"
+#include "Replication/Misc/StreamAndAuthorityPredictionUtils.h"
 
 #include "Containers/Ticker.h"
 #include "Framework/Application/SlateApplication.h"
@@ -15,6 +17,25 @@
 
 #define LOCTEXT_NAMESPACE "FMultiUserReplicationManager"
 
+namespace UE::MultiUserClient::Replication::PrivateReplicationManager
+{
+	class FOfflineClientAdapter : public IOfflineReplicationClient
+	{
+		const FOfflineClient& Client;
+	public:
+		
+		FOfflineClientAdapter(const FOfflineClient& Client)
+			: Client(Client)
+		{}
+
+		//~ Begin IOfflineReplicationClient Interface
+		virtual const FConcertClientInfo& GetClientInfo() const override { return Client.GetClientInfo(); }
+		virtual const FGuid& GetLastAssociatedEndpoint() const override { return Client.GetLastAssociatedEndpoint(); }
+		virtual const FConcertBaseStreamInfo& GetPredictedStream() const override { return Client.GetPredictedStream(); }
+		//~ End IOfflineReplicationClient Interface
+	};
+}
+
 namespace UE::MultiUserClient::Replication
 {
 	FMultiUserReplicationManager::FMultiUserReplicationManager(TSharedRef<IConcertSyncClient> InClient)
@@ -231,6 +252,36 @@ namespace UE::MultiUserClient::Replication
 		}
 		return FExternalClientChangeRequestHandler::MakeFailedOperation(EChangeStreamOperationResult::NotInSession, EChangeAuthorityOperationResult::NotInSession);
 	}
+	
+	void FMultiUserReplicationManager::ForEachOfflineClient(TFunctionRef<EBreakBehavior(const IOfflineReplicationClient&)> Callback) const
+	{
+		if (ensureMsgf(IsInGameThread(), TEXT("To simplify implementation, only calls from game thread are allowed."))
+			&& ConnectedState)
+		{
+			ConnectedState->OfflineClientManager.ForEachClient([&Callback](FOfflineClient& Client)
+			{
+				return Callback(PrivateReplicationManager::FOfflineClientAdapter(Client));
+			});
+		}
+	}
+
+	bool FMultiUserReplicationManager::FindOfflineClient(const FGuid& ClientId, TFunctionRef<void(const IOfflineReplicationClient&)> Callback) const
+	{
+		if (!ensureMsgf(IsInGameThread(), TEXT("To simplify implementation, only calls from game thread are allowed."))
+			|| !ConnectedState)
+		{
+			return false;
+		}
+
+		const FOfflineClient* OfflineClient = ConnectedState->OfflineClientManager.FindClient(ClientId);
+		if (OfflineClient)
+		{
+			Callback(PrivateReplicationManager::FOfflineClientAdapter(*OfflineClient));
+			return true;
+		}
+		
+		return false;
+	}
 
 	FMultiUserReplicationManager::FConnectedState::FConnectedState(TSharedRef<IConcertSyncClient> InClient, FReplicationDiscoveryContainer& InDiscoveryContainer)
 		: Client(InClient)
diff --git a/Engine/Plugins/Developer/Concert/ConcertApp/MultiUserClient/Source/MultiUserClient/Private/Replication/MultiUserReplicationManager.h b/Engine/Plugins/Developer/Concert/ConcertApp/MultiUserClient/Source/MultiUserClient/Private/Replication/MultiUserReplicationManager.h
--- a/Engine/Plugins/Developer/Concert/ConcertApp/MultiUserClient/Source/MultiUserClient/Private/Replication/MultiUserReplicationManager.h
+++ b/Engine/Plugins/Developer/Concert/ConcertApp/MultiUserClient/Source/MultiUserClient/Private/Replication/MultiUserReplicationManager.h
@@ -99,6 +99,8 @@ namespace UE::MultiUserClient::Replication
 		virtual void RegisterReplicationDiscoverer(TSharedRef<IReplicationDiscoverer> Discoverer) override;
 		virtual void RemoveReplicationDiscoverer(const TSharedRef<IReplicationDiscoverer>& Discoverer) override;
 		virtual TSharedRef<IClientChangeOperation> EnqueueChanges(const FGuid& ClientId, TAttribute<FChangeClientReplicationRequest> SubmissionParams) override;
+		virtual void ForEachOfflineClient(TFunctionRef<EBreakBehavior(const IOfflineReplicationClient&)> Callback) const override;
+		virtual bool FindOfflineClient(const FGuid& ClientId, TFunctionRef<void(const IOfflineReplicationClient&)> Callback) const override;
 		virtual FOnServerStateChanged& OnStreamServerStateChanged() override { return OnStreamServerStateChangedDelegate; }
 		virtual FOnServerStateChanged& OnAuthorityServerStateChanged() override { return OnAuthorityServerStateChangedDelegate; }
 		//~ End IMultiUserReplication Interface
diff --git a/Engine/Plugins/Developer/Concert/ConcertApp/MultiUserClient/Source/MultiUserClient/Public/Replication/IMultiUserReplication.h b/Engine/Plugins/Developer/Concert/ConcertApp/MultiUserClient/Source/MultiUserClient/Public/Replication/IMultiUserReplication.h
--- a/Engine/Plugins/Developer/Concert/ConcertApp/MultiUserClient/Source/MultiUserClient/Public/Replication/IMultiUserReplication.h
+++ b/Engine/Plugins/Developer/Concert/ConcertApp/MultiUserClient/Source/MultiUserClient/Public/Replication/IMultiUserReplication.h
@@ -5,13 +5,18 @@
 #include "Misc/Attribute.h"
 #include "Templates/SharedPointer.h"
 
-struct FGuid;
+enum class EBreakBehavior : uint8;
+
+struct FConcertClientInfo;
 struct FConcertStreamFrequencySettings;
 struct FConcertObjectReplicationMap;
+struct FGuid;
+template<typename OptionalType> struct TOptional;
 
 namespace UE::MultiUserClient
 {
 	class IClientChangeOperation;
+	class IOfflineReplicationClient;
 	class IReplicationDiscoverer;
 	
 	struct FChangeClientReplicationRequest;
@@ -24,22 +29,24 @@ namespace UE::MultiUserClient
 		/**
 		 * @return Gets the last known server map of objects registered for replication for a given client.
 		 * 
-		 * This server state is regularly polled whilst the local client state should always be in synch.
-		 * This function must be called from the game thread.
+		 * This server state is regularly polled whilst the local client state should always be in sync.
+		 * @note This function must be called from the game thread.
+		 * @warning Do not keep any latent reference to the returned pointer. The pointed to memory can be reallocated when the client disconnects.
 		 */
 		virtual const FConcertObjectReplicationMap* FindReplicationMapForClient(const FGuid& ClientId) const = 0;
 
 		/**
 		 * @return Gets the last known server object replication frequencies.
 		 * 
-		 * This server state is regularly polled whilst the local client state should always be in synch.
-		 * This function must be called from the game thread.
+		 * This server state is regularly polled whilst the local client state should always be in sync.
+		 * @note This function must be called from the game thread.
+		 * @warning Do not keep any latent reference to the returned pointer. The pointed to memory can be reallocated when the client disconnects.
 		 */
 		virtual const FConcertStreamFrequencySettings* FindReplicationFrequenciesForClient(const FGuid& ClientId) const = 0;
 		
 		/**
 		 * @return Whether the local editor instance thinks the client has authority over the properties it has registered to ObjectPath.
-		 * This function must be called from the game thread.
+		 * @note This function must be called from the game thread.
 		 */
 		virtual bool IsReplicatingObject(const FGuid& ClientId, const FSoftObjectPath& ObjectPath) const = 0;
 
@@ -50,10 +57,13 @@ namespace UE::MultiUserClient
 		 * - When an user adds an object via Add Actor button
 		 * - When an UObject is added to the world via a transaction (run on the client machine that adds the UObject)
 		 * This function must be called from the game thread.
-		 * 
+		 * @note This function must be called from the game thread.
 		 */
 		virtual void RegisterReplicationDiscoverer(TSharedRef<IReplicationDiscoverer> Discoverer) = 0;
-		/** Unregisters a previously registered discoverer. This function must be called from the game thread. */
+		/**
+		 * Unregisters a previously registered discoverer. 
+		 * @note This function must be called from the game thread.
+		 */
 		virtual void RemoveReplicationDiscoverer(const TSharedRef<IReplicationDiscoverer>& Discoverer) = 0;
 
 		/**
@@ -64,13 +74,39 @@ namespace UE::MultiUserClient
 		 * The authority state specifies which of the registered objects should actually be sending data.
 		 * The stream change is requested first and is followed by the authority change.
 		 *
-		 * This function must be called from the game thread.
-		 *
 		 * @param ClientId The client for which to change authority
 		 * @param SubmissionParams Once the request is ready to be sent to the server, this attribute is used to generate the change request
+		 * @note This function must be called from the game thread.
 		 */
 		virtual TSharedRef<IClientChangeOperation> EnqueueChanges(const FGuid& ClientId, TAttribute<FChangeClientReplicationRequest> SubmissionParams) = 0;
-
+		
+		/******************** Offline client ********************/
+		
+		/**
+		 * Iterates over every user who was previously connected to the current session but is now offline. 
+		 * The callback function is invoked for each offline client with the last endpoint ID the user had.
+		 * 
+		 * Endpoint IDs are associated with the "same" user by matching equal display names and device names.
+		 * 
+		 * @param Callback The callback function to invoke for each offline client.
+		 * 
+		 * @warning The client instance is only guaranteed to be valid for the duration of the call - do not keep any reference to it!
+		 * @note This function must be called from the game thread.
+		 */
+		virtual void ForEachOfflineClient(TFunctionRef<EBreakBehavior(const IOfflineReplicationClient&)> Callback) const = 0;
+		
+		/**
+		 * Finds an offline client by an endpoint Id that was associated with the user in the past.
+		 *
+		 * @return Whether Callback was invoked.
+		 *
+		 * @warning The client instance is only guaranteed to be valid for the duration of the call - do not keep any reference to it!
+		 * @note This function must be called from the game thread.
+		 */
+		virtual bool FindOfflineClient(const FGuid& ClientId, TFunctionRef<void(const IOfflineReplicationClient&)> Callback) const = 0;
+		
+		/******************** Events ********************/
+		
 		DECLARE_MULTICAST_DELEGATE_OneParam(FOnServerStateChanged, const FGuid& /*EndpointId*/);
 		/** @return Delegate that triggers when the given client's known server state has changed. */
 		virtual FOnServerStateChanged& OnStreamServerStateChanged() = 0;
diff --git a/Engine/Plugins/Developer/Concert/ConcertApp/MultiUserClient/Source/MultiUserClient/Public/Replication/IOfflineReplicationClient.h b/Engine/Plugins/Developer/Concert/ConcertApp/MultiUserClient/Source/MultiUserClient/Public/Replication/IOfflineReplicationClient.h
new file mode 100644
--- /dev/null
+++ b/Engine/Plugins/Developer/Concert/ConcertApp/MultiUserClient/Source/MultiUserClient/Public/Replication/IOfflineReplicationClient.h
@@ -0,0 +1,25 @@
+﻿// Copyright Epic Games, Inc. All Rights Reserved.
+
+#pragma once
+
+struct FConcertClientInfo;
+struct FConcertBaseStreamInfo;
+struct FGuid;
+
+namespace UE::MultiUserClient
+{
+	/** Interface for querying information about an offline replication client. */
+	class IOfflineReplicationClient
+	{
+	public:
+
+		/** @return The client info this client had when they were last connected. */
+		virtual const FConcertClientInfo& GetClientInfo() const = 0;
+		/** @return The endpoint id this client had when they were last connected. */
+		virtual const FGuid& GetLastAssociatedEndpoint() const = 0;
+		/** @return The stream content that is expected to be reclaimed by this client when rejoining the session.  */
+		virtual const FConcertBaseStreamInfo& GetPredictedStream() const = 0;
+
+		virtual ~IOfflineReplicationClient() = default;
+	};
+}
diff --git a/Engine/Plugins/Developer/Concert/ConcertApp/MultiUserClient/Source/MultiUserClientLibrary/Private/Replication/MultiUserReplicationSubsystem.cpp b/Engine/Plugins/Developer/Concert/ConcertApp/MultiUserClient/Source/MultiUserClientLibrary/Private/Replication/MultiUserReplicationSubsystem.cpp
--- a/Engine/Plugins/Developer/Concert/ConcertApp/MultiUserClient/Source/MultiUserClientLibrary/Private/Replication/MultiUserReplicationSubsystem.cpp
+++ b/Engine/Plugins/Developer/Concert/ConcertApp/MultiUserClient/Source/MultiUserClientLibrary/Private/Replication/MultiUserReplicationSubsystem.cpp
@@ -5,11 +5,13 @@
 #include "Replication/Async/ChangeClientBlueprintParams.h"
 
 #if WITH_CONCERT
-#include "UObjectAdapterReplicationDiscoverer.h"
+#include "Data/MultiUserClientDisplayInfo.h"
 #include "IMultiUserClientModule.h"
 #include "Replication/Data/ObjectReplicationMap.h"
 #include "Replication/Data/ReplicationFrequencySettings.h"
 #include "Replication/IMultiUserReplication.h"
+#include "Replication/IOfflineReplicationClient.h"
+#include "UObjectAdapterReplicationDiscoverer.h"
 
 #include "Algo/Transform.h"
 #endif
@@ -124,6 +126,28 @@ TArray<FSoftObjectPath> UMultiUserReplicationSubsystem::GetReplicatedObjects(con
 	return {};
 }
 
+TArray<FMultiUserClientDisplayInfo> UMultiUserReplicationSubsystem::GetOwningOfflineClients(const FSoftObjectPath& ObjectPath) const
+{
+	TArray<FMultiUserClientDisplayInfo> Result;
+	
+#if WITH_CONCERT
+	UE::MultiUserClient::IMultiUserReplication* ReplicationInterface = IMultiUserClientModule::Get().GetReplication();
+	if (ensureMsgf(ReplicationInterface, TEXT("We expected it to always be valid.")))
+	{
+		ReplicationInterface->ForEachOfflineClient([&ObjectPath, &Result](const UE::MultiUserClient::IOfflineReplicationClient& Client)
+		{
+			if (Client.GetPredictedStream().ReplicationMap.HasProperties(ObjectPath))
+			{
+				Result.Emplace(Client.GetClientInfo());
+			}
+			return EBreakBehavior::Continue;
+		});
+	}
+#endif
+
+	return Result;
+}
+
 void UMultiUserReplicationSubsystem::Initialize(FSubsystemCollectionBase& Collection)
 {
 	Super::Initialize(Collection);  
diff --git a/Engine/Plugins/Developer/Concert/ConcertApp/MultiUserClient/Source/MultiUserClientLibrary/Public/Data/MultiUserClientDisplayInfo.h b/Engine/Plugins/Developer/Concert/ConcertApp/MultiUserClient/Source/MultiUserClientLibrary/Public/Data/MultiUserClientDisplayInfo.h
new file mode 100644
--- /dev/null
+++ b/Engine/Plugins/Developer/Concert/ConcertApp/MultiUserClient/Source/MultiUserClientLibrary/Public/Data/MultiUserClientDisplayInfo.h
@@ -0,0 +1,22 @@
+﻿// Copyright Epic Games, Inc. All Rights Reserved.
+
+#pragma once
+#include "ConcertMessageData.h"
+
+#include "MultiUserClientDisplayInfo.generated.h"
+
+/** Blueprint copy of FConcertClientInfo. Can describe offline clients, too. */
+USTRUCT(BlueprintType)
+struct FMultiUserClientDisplayInfo
+{
+	GENERATED_BODY()
+
+	/** Holds the display name of the user that owns this instance. */
+	UPROPERTY(VisibleAnywhere, BlueprintReadOnly, Category = "Client Info")
+	FString DisplayName;
+
+	FMultiUserClientDisplayInfo() = default;
+	explicit FMultiUserClientDisplayInfo(const FConcertClientInfo& ClientInfo)
+		: DisplayName(ClientInfo.DisplayName)
+	{}
+};
\ No newline at end of file
diff --git a/Engine/Plugins/Developer/Concert/ConcertApp/MultiUserClient/Source/MultiUserClientLibrary/Public/Replication/MultiUserReplicationSubsystem.h b/Engine/Plugins/Developer/Concert/ConcertApp/MultiUserClient/Source/MultiUserClientLibrary/Public/Replication/MultiUserReplicationSubsystem.h
--- a/Engine/Plugins/Developer/Concert/ConcertApp/MultiUserClient/Source/MultiUserClientLibrary/Public/Replication/MultiUserReplicationSubsystem.h
+++ b/Engine/Plugins/Developer/Concert/ConcertApp/MultiUserClient/Source/MultiUserClientLibrary/Public/Replication/MultiUserReplicationSubsystem.h
@@ -3,7 +3,10 @@
 #pragma once
 
 #include "Subsystems/EngineSubsystem.h"
+
 #include "ConcertPropertyChainWrapper.h"
+#include "Data/MultiUserClientDisplayInfo.h"
+
 #include "MultiUserReplicationSubsystem.generated.h"
 
 struct FMultiUserObjectReplicationSettings;
@@ -73,6 +76,23 @@ public:
 	 */
 	UFUNCTION(BlueprintPure, Category = "Multi-user")
 	TArray<FSoftObjectPath> GetReplicatedObjects(const FGuid& ClientId) const;
+
+	/**
+	 * A list of offline clients that, upon rejoining a session, will attempt to reclaim properties 
+	 * they previously registered for an object, regardless of whether the client left gracefully or due to a crash.
+	 * 
+	 * By default, when a client disconnects (either gracefully or due to a crash) and later rejoins a session,
+	 * the client attempts to re-register the properties it had previously registered for the object.
+	 * 
+	 * @return A list of client descriptions representing offline clients that will attempt to reclaim properties 
+	 * associated with the object when they rejoin.
+	 */
+	UFUNCTION(BlueprintPure, Category = "Multi-user", meta = (Keywords = "Owning Offline Reclaim Join Rejoin Client Find Stream Registered"))
+	TArray<FMultiUserClientDisplayInfo> GetOwningOfflineClients(const FSoftObjectPath& ObjectPath) const;
+
+	/** @return Whether any offline clients will try to register properties for ObjectPath upon rejoining.*/
+	UFUNCTION(BlueprintPure, Category = "Multi-user", meta = (Keywords = "Owning Offline Reclaim Join Rejoin Client Find Stream Registered"))
+	bool IsOwnedByOfflineClient(const FSoftObjectPath& ObjectPath) const { return !GetOwningOfflineClients(ObjectPath).IsEmpty(); }
 	
 	//~ Begin USubsystem Interface
 	virtual void Initialize(FSubsystemCollectionBase& Collection) override;
