Fix compiler error on Mac due to lambda shadowing name from surrounding class.

diff --git a/Engine/Plugins/Developer/Concert/ConcertApp/MultiUserClient/Source/MultiUserClient/Private/Replication/MultiUserReplicationManager.cpp b/Engine/Plugins/Developer/Concert/ConcertApp/MultiUserClient/Source/MultiUserClient/Private/Replication/MultiUserReplicationManager.cpp
--- a/Engine/Plugins/Developer/Concert/ConcertApp/MultiUserClient/Source/MultiUserClient/Private/Replication/MultiUserReplicationManager.cpp
+++ b/Engine/Plugins/Developer/Concert/ConcertApp/MultiUserClient/Source/MultiUserClient/Private/Replication/MultiUserReplicationManager.cpp
@@ -258,9 +258,10 @@ namespace UE::MultiUserClient::Replication
 		if (ensureMsgf(IsInGameThread(), TEXT("To simplify implementation, only calls from game thread are allowed."))
 			&& ConnectedState)
 		{
-			ConnectedState->OfflineClientManager.ForEachClient([&Callback](FOfflineClient& Client)
+			ConnectedState->OfflineClientManager.ForEachClient([&Callback](const FOfflineClient& OfflineClient)
 			{
-				return Callback(PrivateReplicationManager::FOfflineClientAdapter(Client));
+				const PrivateReplicationManager::FOfflineClientAdapter OfflineClientAdapter(OfflineClient);
+				return Callback(OfflineClientAdapter);
 			});
 		}
 	}
