Warn when when clients have duplicate names and are using MU replication

diff --git a/Engine/Plugins/Developer/Concert/ConcertApp/MultiUserClient/Source/MultiUserClient/Private/Replication/Misc/Notification/DuplicateUserNotifier.cpp b/Engine/Plugins/Developer/Concert/ConcertApp/MultiUserClient/Source/MultiUserClient/Private/Replication/Misc/Notification/DuplicateUserNotifier.cpp
new file mode 100644
--- /dev/null
+++ b/Engine/Plugins/Developer/Concert/ConcertApp/MultiUserClient/Source/MultiUserClient/Private/Replication/Misc/Notification/DuplicateUserNotifier.cpp
@@ -0,0 +1,122 @@
+﻿// Copyright Epic Games, Inc. All Rights Reserved.
+
+#include "DuplicateUserNotifier.h"
+
+#include "ConcertLogGlobal.h"
+#include "ConcertMessageData.h"
+#include "IConcertClient.h"
+#include "Replication/Client/Online/OnlineClientManager.h"
+#include "Replication/Messages/RestoreContent.h"
+#include "Replication/Misc/GlobalAuthorityCache.h"
+#include "Replication/Misc/StreamAndAuthorityPredictionUtils.h"
+
+#include "HAL/PlatformMisc.h"
+#include "Framework/Application/SlateApplication.h"
+#include "Framework/Notifications/NotificationManager.h"
+#include "Widgets/Notifications/SNotificationList.h"
+
+#define LOCTEXT_NAMESPACE "FDuplicateUserNotifier"
+
+namespace UE::MultiUserClient::Replication::PrivateDuplicateUserNotifier
+{
+	static void ShowNotification(const FConcertClientInfo& ClientInfo)
+	{
+		check(IsInGameThread());
+		if (!FSlateApplication::IsInitialized())
+		{
+			return;
+		}
+			
+		FSlateNotificationManager& NotificationManager = FSlateNotificationManager::Get();
+		FNotificationInfo NotificationInfo(LOCTEXT("DuplicateClient.Main", "Duplicate Client Name"));
+		NotificationInfo.SubText = FText::Format(
+			LOCTEXT("DuplicateClient.SubTextFmt", "There are 2 clients with display name {0} and device name {1} in the session.\n"
+				"Some replication features will not work as intended.\n\n"
+				"Most likely, you have launched 2 editors on the same machine.\n"
+				"Try setting a display name using different -CONCERTDISPLAYNAME flag values for each instance."
+			),
+			FText::FromString(ClientInfo.DisplayName),
+			FText::FromString(ClientInfo.DeviceName)
+			);
+		NotificationInfo.bFireAndForget = true;
+		NotificationInfo.ExpireDuration = 8.f;
+		NotificationManager.AddNotification(NotificationInfo)
+			->SetCompletionState(SNotificationItem::CS_Fail);
+	}
+
+	
+	static bool HasNameConflictWithLogicallyDuplicateClient(
+		const FOnlineClientManager& ClientManager,
+		const IConcertClientSession& Session
+		)
+	{
+		const FConcertClientInfo& ClientInfo = Session.GetLocalClientInfo();
+		const FGuid& LocalId = Session.GetSessionClientEndpointId();
+
+		bool bLocalClientHasContent = !ClientManager.GetLocalClient().GetStreamSynchronizer().GetServerState().IsEmpty();
+		bool bHasDuplicateClientName = false;
+		ClientManager.ForEachClient(
+			[&Session, &ClientInfo, &LocalId, &bHasDuplicateClientName, bLocalClientHasContent]
+			(const FOnlineClient& Client)
+			{
+				FConcertSessionClientInfo OtherInfo;
+				const bool bSuccess = Session.FindSessionClient(Client.GetEndpointId(), OtherInfo);
+				
+				bHasDuplicateClientName |= bSuccess
+					&& Client.GetEndpointId() != LocalId
+					&& ConcertSyncCore::Replication::AreLogicallySameClients(ClientInfo, OtherInfo.ClientInfo)
+					// Only warn if either I or the other client is using replication.
+					&& (bLocalClientHasContent || !Client.GetStreamSynchronizer().GetServerState().IsEmpty());
+				return bHasDuplicateClientName ? EBreakBehavior::Break : EBreakBehavior::Continue;
+			});
+
+		return bHasDuplicateClientName;
+	}
+}
+
+namespace UE::MultiUserClient::Replication
+{
+	FDuplicateUserNotifier::FDuplicateUserNotifier(
+		IConcertClient& InClient, 
+		FOnlineClientManager& InOnlineClientManager 
+		)
+		: Client(InClient)
+		, OnlineClientManager(InOnlineClientManager)
+	{
+		OnlineClientManager.GetAuthorityCache().OnCacheChanged().AddRaw(this, &FDuplicateUserNotifier::OnClientContentChanged);
+	}
+
+	FDuplicateUserNotifier::~FDuplicateUserNotifier()
+	{
+		OnlineClientManager.GetAuthorityCache().OnCacheChanged().RemoveAll(this);
+	}
+
+	void FDuplicateUserNotifier::OnClientContentChanged(const FGuid&)
+	{
+		const TSharedPtr<IConcertClientSession> Session = Client.GetCurrentSession();
+		if (ensureMsgf(Session, TEXT("We are already supposed to have been destroyed."))
+			&& PrivateDuplicateUserNotifier::HasNameConflictWithLogicallyDuplicateClient(OnlineClientManager, *Session))
+		{
+			ShowWarning();
+		}
+	}
+
+	void FDuplicateUserNotifier::ShowWarning()
+	{
+		if (!bHasWarnedUser)
+		{
+			bHasWarnedUser = true;
+			OnlineClientManager.GetAuthorityCache().OnCacheChanged().RemoveAll(this);
+		
+			const FConcertClientInfo& ClientInfo = Client.GetClientInfo();
+			UE_LOG(LogConcert, Error,
+				TEXT("There are 2 clients with display name %s and device name %s in the session. Some replication features won't work as intended."),
+				*ClientInfo.DisplayName,
+				*ClientInfo.DeviceName
+				);
+			PrivateDuplicateUserNotifier::ShowNotification(ClientInfo);
+		}
+	}
+}
+
+#undef LOCTEXT_NAMESPACE
\ No newline at end of file
diff --git a/Engine/Plugins/Developer/Concert/ConcertApp/MultiUserClient/Source/MultiUserClient/Private/Replication/Misc/Notification/DuplicateUserNotifier.h b/Engine/Plugins/Developer/Concert/ConcertApp/MultiUserClient/Source/MultiUserClient/Private/Replication/Misc/Notification/DuplicateUserNotifier.h
new file mode 100644
--- /dev/null
+++ b/Engine/Plugins/Developer/Concert/ConcertApp/MultiUserClient/Source/MultiUserClient/Private/Replication/Misc/Notification/DuplicateUserNotifier.h
@@ -0,0 +1,49 @@
+﻿// Copyright Epic Games, Inc. All Rights Reserved.
+
+#pragma once
+
+#include "HAL/Platform.h"
+#include "Replication/Client/Online/RemoteClient.h"
+#include "Templates/UnrealTemplate.h"
+
+class IConcertClient;
+struct FConcertReplication_RestoreContent_Response;
+
+namespace UE::MultiUserClient::Replication
+{
+	class FOnlineClientManager;
+	
+	/**
+	 * Informs the user if they are using the same name as another client.
+	 * The message is only shown if the local or other, duplicate client use replication, i.e. register properties.
+	 * 
+	 * Certain features, like restoring the content that the client had when they last left, will not working correctly with non-unique display names.
+	 */
+	class FDuplicateUserNotifier : public FNoncopyable
+	{
+	public:
+		
+		FDuplicateUserNotifier(
+			IConcertClient& InClient UE_LIFETIMEBOUND,
+			FOnlineClientManager& InOnlineClientManager UE_LIFETIMEBOUND
+			);
+		~FDuplicateUserNotifier();
+
+	private:
+
+		/** Used to get client display information. */
+		IConcertClient& Client;
+		/** Used to detect when a client's state changes. */
+		FOnlineClientManager& OnlineClientManager; 
+
+		/** Whether the user was already warned in this session (will re-warn if they leave and join another session). */
+		bool bHasWarnedUser = false;
+
+		/** Invokes when any client's content changes. */
+		void OnClientContentChanged(const FGuid& ClientId);
+
+		/** Shows the warning and prevents it from being shown again in the active session. */
+		void ShowWarning();
+	};
+}
+
diff --git a/Engine/Plugins/Developer/Concert/ConcertApp/MultiUserClient/Source/MultiUserClient/Private/Replication/Misc/Notification/ReplicationUserNotifier.cpp b/Engine/Plugins/Developer/Concert/ConcertApp/MultiUserClient/Source/MultiUserClient/Private/Replication/Misc/Notification/ReplicationUserNotifier.cpp
--- a/Engine/Plugins/Developer/Concert/ConcertApp/MultiUserClient/Source/MultiUserClient/Private/Replication/Misc/Notification/ReplicationUserNotifier.cpp
+++ b/Engine/Plugins/Developer/Concert/ConcertApp/MultiUserClient/Source/MultiUserClient/Private/Replication/Misc/Notification/ReplicationUserNotifier.cpp
@@ -4,8 +4,13 @@
 
 namespace UE::MultiUserClient::Replication
 {
-	FReplicationUserNotifier::FReplicationUserNotifier(FOnlineClientManager& InReplicationClientManager, FMuteStateManager& InMuteManager)
+	FReplicationUserNotifier::FReplicationUserNotifier(
+		IConcertClient& InClient,
+		FOnlineClientManager& InReplicationClientManager,
+		FMuteStateManager& InMuteManager
+		)
 		: SubmissionNotifier(InReplicationClientManager)
 		, MutingNotifier(InMuteManager)
+		, DuplicateClientNameNotifier(InClient, InReplicationClientManager)
 	{}
 }
diff --git a/Engine/Plugins/Developer/Concert/ConcertApp/MultiUserClient/Source/MultiUserClient/Private/Replication/Misc/Notification/ReplicationUserNotifier.h b/Engine/Plugins/Developer/Concert/ConcertApp/MultiUserClient/Source/MultiUserClient/Private/Replication/Misc/Notification/ReplicationUserNotifier.h
--- a/Engine/Plugins/Developer/Concert/ConcertApp/MultiUserClient/Source/MultiUserClient/Private/Replication/Misc/Notification/ReplicationUserNotifier.h
+++ b/Engine/Plugins/Developer/Concert/ConcertApp/MultiUserClient/Source/MultiUserClient/Private/Replication/Misc/Notification/ReplicationUserNotifier.h
@@ -2,6 +2,7 @@
 
 #pragma once
 
+#include "DuplicateUserNotifier.h"
 #include "MutingNotifier.h"
 #include "Replication/Submission/Notification/SubmissionNotifier.h"
 
@@ -17,7 +18,11 @@ namespace UE::MultiUserClient::Replication
 	{
 	public:
 
-		FReplicationUserNotifier(FOnlineClientManager& InReplicationClientManager UE_LIFETIMEBOUND, FMuteStateManager& InMuteManager UE_LIFETIMEBOUND);
+		FReplicationUserNotifier(
+			IConcertClient& InClient UE_LIFETIMEBOUND,
+			FOnlineClientManager& InReplicationClientManager UE_LIFETIMEBOUND,
+			FMuteStateManager& InMuteManager UE_LIFETIMEBOUND
+			);
 
 	private:
 		
@@ -26,6 +31,8 @@ namespace UE::MultiUserClient::Replication
 		
 		/** Informs the user of things that went wrong with muting requests. */
 		const FMutingNotifier MutingNotifier;
+		/** Informs the user if they are using the same name as another client. */
+		const FDuplicateUserNotifier DuplicateClientNameNotifier;
 	};
 }
 
diff --git a/Engine/Plugins/Developer/Concert/ConcertApp/MultiUserClient/Source/MultiUserClient/Private/Replication/MultiUserReplicationManager.cpp b/Engine/Plugins/Developer/Concert/ConcertApp/MultiUserClient/Source/MultiUserClient/Private/Replication/MultiUserReplicationManager.cpp
--- a/Engine/Plugins/Developer/Concert/ConcertApp/MultiUserClient/Source/MultiUserClient/Private/Replication/MultiUserReplicationManager.cpp
+++ b/Engine/Plugins/Developer/Concert/ConcertApp/MultiUserClient/Source/MultiUserClient/Private/Replication/MultiUserReplicationManager.cpp
@@ -136,29 +136,7 @@ namespace UE::MultiUserClient::Replication
 			{
 				.Flags = EConcertReplicationRestoreContentFlags::All | EConcertReplicationRestoreContentFlags::ValidateUniqueClient
 			}
-		)
-		.Next([ClientInfo = Client->GetConcertClient()->GetClientInfo()](FConcertReplication_RestoreContent_Response&& Response)
-		{
-			UE_LOG(LogConcert, Log, TEXT("Content restoration completed with result '%s'"), *ConcertSyncCore::LexToString(Response.ErrorCode));
-			const bool bIsTimeout = !IsInGameThread() || Response.ErrorCode == EConcertReplicationRestoreErrorCode::Timeout;
-			if (Response.IsSuccess() || bIsTimeout || !FSlateApplication::IsInitialized())
-			{
-				return;
-			}
-
-			FSlateNotificationManager& NotificationManager = FSlateNotificationManager::Get();
-			FNotificationInfo NotificationInfo(LOCTEXT("RestoreFailed.Main", "Replication Content Restore"));
-			NotificationInfo.SubText = FText::Format(
-				LOCTEXT("RestoreFailed.SubTextFmt", "Display name {0} and device name {1} already taken by another client in session."),
-				FText::FromString(ClientInfo.DisplayName),
-				FText::FromString(ClientInfo.DeviceName)
-				);
-			NotificationInfo.bFireAndForget = true;
-			NotificationInfo.bUseSuccessFailIcons = true;
-			NotificationInfo.ExpireDuration = 4.f;
-			NotificationManager.AddNotification(NotificationInfo)
-				->SetCompletionState(SNotificationItem::CS_Fail);
-		});
+		);
 	}
 
 	void FMultiUserReplicationManager::SetupClientConnectionEvents()
@@ -295,7 +273,7 @@ namespace UE::MultiUserClient::Replication
 		, PropertySelector(OnlineClientManager, OfflineClientManager)
 		, ChangeLevelHandler(OnlineClientManager.GetLocalClient().GetClientEditModel().Get())
 		, PreventReplicatedPropertyTransaction(*InClient, OnlineClientManager, MuteManager)
-		, UserNotifier(OnlineClientManager, MuteManager)
+		, UserNotifier(*InClient->GetConcertClient(), OnlineClientManager, MuteManager)
 	{}
 }
 
