Fix crash due to destruction widget destruction order.
ViewOptions member was being passed to child column widget but ViewOptions was destroyed before the child widget.

diff --git a/Engine/Plugins/Developer/Concert/ConcertApp/MultiUserClient/Source/MultiUserClient/Private/Widgets/ActiveSession/Replication/Client/Multi/SMultiClientView.cpp b/Engine/Plugins/Developer/Concert/ConcertApp/MultiUserClient/Source/MultiUserClient/Private/Widgets/ActiveSession/Replication/Client/Multi/SMultiClientView.cpp
--- a/Engine/Plugins/Developer/Concert/ConcertApp/MultiUserClient/Source/MultiUserClient/Private/Widgets/ActiveSession/Replication/Client/Multi/SMultiClientView.cpp
+++ b/Engine/Plugins/Developer/Concert/ConcertApp/MultiUserClient/Source/MultiUserClient/Private/Widgets/ActiveSession/Replication/Client/Multi/SMultiClientView.cpp
@@ -82,6 +82,16 @@ namespace UE::MultiUserClient::Replication
 
 	SMultiClientView::~SMultiClientView()
 	{
+		// Some of StreamEditor's columns reference some of our members.
+		// Hence, StreamEditor needs to be destroyed before ~SCompoundWidget destroys it.
+		// E.g. MultiStreamColumns::AssignedClientsColumn references ViewOptions so make sure ViewOptions is destroyed after the widget.
+		ChildSlot.DetachWidget();
+		StreamEditor.Reset();
+		// These objects depend on ViewOptions. These resets are strictly not needed to be safe but doing this explicitly protects us in case
+		// somebody moves declaration order of properties without realising dependency order.
+		StreamModel.Reset();
+		PropertyAssignmentView.Reset();
+		
 		OnlineClientManager->OnRemoteClientsChanged().RemoveAll(this);
 		UserSelectedProperties->OnPropertySelectionChanged().RemoveAll(this);
 		CleanClientSubscriptions();
diff --git a/Engine/Plugins/Developer/Concert/ConcertApp/MultiUserClient/Source/MultiUserClient/Private/Widgets/ActiveSession/Replication/Client/Multi/SMultiClientView.h b/Engine/Plugins/Developer/Concert/ConcertApp/MultiUserClient/Source/MultiUserClient/Private/Widgets/ActiveSession/Replication/Client/Multi/SMultiClientView.h
--- a/Engine/Plugins/Developer/Concert/ConcertApp/MultiUserClient/Source/MultiUserClient/Private/Widgets/ActiveSession/Replication/Client/Multi/SMultiClientView.h
+++ b/Engine/Plugins/Developer/Concert/ConcertApp/MultiUserClient/Source/MultiUserClient/Private/Widgets/ActiveSession/Replication/Client/Multi/SMultiClientView.h
@@ -53,14 +53,20 @@ namespace UE::MultiUserClient::Replication
 		TSharedPtr<IConcertClient> ConcertClient;
 		/** Keeps track of the properties that the user has selected to iterate on. */
 		FUserPropertySelector* UserSelectedProperties = nullptr;
-
+		
 		// These are used to know when to refresh the UI.
 		FOnlineClientManager* OnlineClientManager = nullptr;
 		FOfflineClientManager* OfflineClientManager = nullptr;
 		IOnlineClientSelectionModel* OnlineClientSelectionModel = nullptr;
 		IOfflineClientSelectionModel* OfflineClientSelectionModel = nullptr;
 		
-		/** Combines the clients */
+		/**
+		 * Controls the content shown in the UI.
+		 * Important: Some systems keep a reference to ViewOptions. Evaluate destruction order if you move the member ordering.
+		 */
+		FMultiViewOptions ViewOptions;
+		
+		/** Combines the online and offline clients. */
 		TSharedPtr<FMultiStreamModel> StreamModel;
 		/** Displayed in the UI. */
 		TSharedPtr<ConcertSharedSlate::IMultiReplicationStreamEditor> StreamEditor;
@@ -74,10 +80,7 @@ namespace UE::MultiUserClient::Replication
 		TSharedPtr<SPropertySelectionComboButton> PropertySelectionButton;
 		/** Displays the properties for the objects displayed in the top view. */
 		TSharedPtr<ConcertSharedSlate::IMultiObjectPropertyAssignmentView> PropertyAssignmentView;
-
-		/** Controls the content shown in the UI. */
-		FMultiViewOptions ViewOptions;
-
+		
 		/** This logic helps us decide whether an object should be displayed and lets us know that the object list needs to be refreshed (e.g. due to world change). */
 		ConcertClientSharedSlate::FHideObjectsNotInWorldLogic HideObjectsNotInEditorWorld;
 
