Refresh the list of known offline endpoints whenever an activity produced. This fixes offline client info not being displayed correctly in certain cases.

diff --git a/Engine/Plugins/Developer/Concert/ConcertApp/MultiUserClient/Source/MultiUserClient/Private/Replication/Client/Offline/OfflineClientManager.cpp b/Engine/Plugins/Developer/Concert/ConcertApp/MultiUserClient/Source/MultiUserClient/Private/Replication/Client/Offline/OfflineClientManager.cpp
--- a/Engine/Plugins/Developer/Concert/ConcertApp/MultiUserClient/Source/MultiUserClient/Private/Replication/Client/Offline/OfflineClientManager.cpp
+++ b/Engine/Plugins/Developer/Concert/ConcertApp/MultiUserClient/Source/MultiUserClient/Private/Replication/Client/Offline/OfflineClientManager.cpp
@@ -17,6 +17,7 @@ namespace UE::MultiUserClient::Replication
 		: ClientInstance(ClientInstance)
 		, OnlineClientManager(OnlineClientManager)
 	{
+		ClientInstance.GetWorkspace()->OnActivityAddedOrUpdated().AddRaw(this, &FOfflineClientManager::OnActivityAddedOrProduced);
 		OnlineClientManager.OnRemoteClientsChanged().AddRaw(this, &FOfflineClientManager::RefreshOfflineClients);
 		RefreshOfflineClients();
 	}
@@ -46,6 +47,13 @@ namespace UE::MultiUserClient::Replication
 			);
 	}
 
+	void FOfflineClientManager::OnActivityAddedOrProduced(const FConcertClientInfo&, const FConcertSyncActivity&, const FStructOnScope&)
+	{
+		const TSharedPtr<IConcertClientWorkspace> Workspace = ClientInstance.GetWorkspace();
+		check(Workspace);
+		EndpointCache.UpdateEndpoints(*Workspace);
+	}
+
 	void FOfflineClientManager::RefreshOfflineClients()
 	{
 		const TSharedPtr<IConcertClientWorkspace> Workspace = ClientInstance.GetWorkspace();
diff --git a/Engine/Plugins/Developer/Concert/ConcertApp/MultiUserClient/Source/MultiUserClient/Private/Replication/Client/Offline/OfflineClientManager.h b/Engine/Plugins/Developer/Concert/ConcertApp/MultiUserClient/Source/MultiUserClient/Private/Replication/Client/Offline/OfflineClientManager.h
--- a/Engine/Plugins/Developer/Concert/ConcertApp/MultiUserClient/Source/MultiUserClient/Private/Replication/Client/Offline/OfflineClientManager.h
+++ b/Engine/Plugins/Developer/Concert/ConcertApp/MultiUserClient/Source/MultiUserClient/Private/Replication/Client/Offline/OfflineClientManager.h
@@ -13,6 +13,7 @@
 
 #include <type_traits>
 
+struct FConcertSyncActivity;
 class IConcertClientWorkspace;
 class IConcertSyncClient;
 
@@ -78,6 +79,9 @@ namespace UE::MultiUserClient::Replication
 		FOfflineClientDelegate OnPostClientAddedDelegate;
 		/** Called just before an offline client is about to be removed. Called before OnClientsChanged. */
 		FOfflineClientDelegate OnPreClientRemovedDelegate;
+
+		/** Updates the list of known endpoints in response to an activity being added. */
+		void OnActivityAddedOrProduced(const FConcertClientInfo&, const FConcertSyncActivity&, const FStructOnScope&);
 		
 		/** Refreshes Clients based on which clients are online. */
 		void RefreshOfflineClients();
