Fix forgetting to unsubscribe from OnActivityAddedOrUpdated

diff --git a/Engine/Plugins/Developer/Concert/ConcertApp/MultiUserClient/Source/MultiUserClient/Private/Replication/Client/Offline/OfflineClientManager.cpp b/Engine/Plugins/Developer/Concert/ConcertApp/MultiUserClient/Source/MultiUserClient/Private/Replication/Client/Offline/OfflineClientManager.cpp
--- a/Engine/Plugins/Developer/Concert/ConcertApp/MultiUserClient/Source/MultiUserClient/Private/Replication/Client/Offline/OfflineClientManager.cpp
+++ b/Engine/Plugins/Developer/Concert/ConcertApp/MultiUserClient/Source/MultiUserClient/Private/Replication/Client/Offline/OfflineClientManager.cpp
@@ -24,6 +24,7 @@ namespace UE::MultiUserClient::Replication
 
 	FOfflineClientManager::~FOfflineClientManager()
 	{
+		ClientInstance.GetWorkspace()->OnActivityAddedOrUpdated().RemoveAll(this);
 		OnlineClientManager.OnRemoteClientsChanged().RemoveAll(this);
 	}
 
