Offline client state now is refreshed when a replication activity is produced. This fixes the UI not showing an offline client the first time the client leaves the session.

diff --git a/Engine/Plugins/Developer/Concert/ConcertApp/MultiUserClient/Source/MultiUserClient/Private/Replication/Client/Offline/OfflineClient.cpp b/Engine/Plugins/Developer/Concert/ConcertApp/MultiUserClient/Source/MultiUserClient/Private/Replication/Client/Offline/OfflineClient.cpp
--- a/Engine/Plugins/Developer/Concert/ConcertApp/MultiUserClient/Source/MultiUserClient/Private/Replication/Client/Offline/OfflineClient.cpp
+++ b/Engine/Plugins/Developer/Concert/ConcertApp/MultiUserClient/Source/MultiUserClient/Private/Replication/Client/Offline/OfflineClient.cpp
@@ -21,5 +21,10 @@ namespace UE::MultiUserClient::Replication
 					)
 				)
 			)
-	{}
+	{
+		ContentPredictor.OnPredictionChanged().AddLambda([this]
+		{
+			OnStreamPredictionChangedDelegate.Broadcast();
+		});
+	}
 }
diff --git a/Engine/Plugins/Developer/Concert/ConcertApp/MultiUserClient/Source/MultiUserClient/Private/Replication/Client/Offline/RejoinStreamAndAuthorityPredictor.cpp b/Engine/Plugins/Developer/Concert/ConcertApp/MultiUserClient/Source/MultiUserClient/Private/Replication/Client/Offline/RejoinStreamAndAuthorityPredictor.cpp
--- a/Engine/Plugins/Developer/Concert/ConcertApp/MultiUserClient/Source/MultiUserClient/Private/Replication/Client/Offline/RejoinStreamAndAuthorityPredictor.cpp
+++ b/Engine/Plugins/Developer/Concert/ConcertApp/MultiUserClient/Source/MultiUserClient/Private/Replication/Client/Offline/RejoinStreamAndAuthorityPredictor.cpp
@@ -3,9 +3,11 @@
 #include "RejoinStreamAndAuthorityPredictor.h"
 
 #include "EndpointCache.h"
+#include "IConcertClientWorkspace.h"
 #include "Replication/Messages/ReplicationActivity.h"
 
 #include "Replication/Misc/ClientPredictionUtils.h"
+#include "Replication/Misc/StreamAndAuthorityPredictionUtils.h"
 
 namespace UE::MultiUserClient::Replication
 {
@@ -54,14 +56,42 @@ namespace UE::MultiUserClient::Replication
 			}
 		}
 	}
-	
-	FRejoinStreamAndAuthorityPredictor::FRejoinStreamAndAuthorityPredictor(const IConcertClientWorkspace& InWorkspace, FConcertClientInfo InClientInfo)
-		: ClientInfo(MoveTemp(InClientInfo))
+
+	FRejoinStreamAndAuthorityPredictor::FRejoinStreamAndAuthorityPredictor(IConcertClientWorkspace& InWorkspace, FConcertClientInfo InClientInfo)
+		: Workspace(InWorkspace)
+		, ClientInfo(MoveTemp(InClientInfo))
 		, PredictedStream({ .Identifier = MultiUserStreamID })
 	{
-		// For now, we only look at the FConcertSyncReplicationPayload_LeaveReplication activity. It can be looked up once.
-		// In the future, we'll add an event for FConcertReplication_PutState_Requests, which can affect offline clients, invoked by presets.
-		// Then, we'll have to listen for activity changes.
-		Private::AnalyzeActivityHistory(InWorkspace, ClientInfo, PredictedStream, PredictedAuthority);
+		// Whenever a leave replication activity is received, update the predicted stream.
+		// Usually when a remote client leaves, a FOfflineClient is created is response to the session's client list changing but the leave replication
+		// activity is only received later. Hence, we need to listen for changes.
+		Workspace.OnActivityAddedOrUpdated().AddRaw(this, &FRejoinStreamAndAuthorityPredictor::OnActivityAddedOrUpdated);
+		
+		AnalyzeHistory();
+	}
+
+	FRejoinStreamAndAuthorityPredictor::~FRejoinStreamAndAuthorityPredictor()
+	{
+		Workspace.OnActivityAddedOrUpdated().RemoveAll(this);
+	}
+
+	void FRejoinStreamAndAuthorityPredictor::OnActivityAddedOrUpdated(
+		const FConcertClientInfo& ActivityClientInfo,
+		const FConcertSyncActivity& Activity,
+		const FStructOnScope& Summary
+		)
+	{
+		if (ConcertSyncCore::Replication::AreLogicallySameClients(ClientInfo, ActivityClientInfo)
+			&& Activity.EventType == EConcertSyncActivityEventType::Replication)
+		{
+			AnalyzeHistory();
+		}
+	}
+
+	
+	void FRejoinStreamAndAuthorityPredictor::AnalyzeHistory()
+	{
+		Private::AnalyzeActivityHistory(Workspace, ClientInfo, PredictedStream, PredictedAuthority);
+		OnPredictionChangedDelegate.Broadcast();
 	}
 }
diff --git a/Engine/Plugins/Developer/Concert/ConcertApp/MultiUserClient/Source/MultiUserClient/Private/Replication/Client/Offline/RejoinStreamAndAuthorityPredictor.h b/Engine/Plugins/Developer/Concert/ConcertApp/MultiUserClient/Source/MultiUserClient/Private/Replication/Client/Offline/RejoinStreamAndAuthorityPredictor.h
--- a/Engine/Plugins/Developer/Concert/ConcertApp/MultiUserClient/Source/MultiUserClient/Private/Replication/Client/Offline/RejoinStreamAndAuthorityPredictor.h
+++ b/Engine/Plugins/Developer/Concert/ConcertApp/MultiUserClient/Source/MultiUserClient/Private/Replication/Client/Offline/RejoinStreamAndAuthorityPredictor.h
@@ -3,12 +3,14 @@
 #pragma once
 
 #include "ConcertMessageData.h"
+#include "ConcertSyncSessionTypes.h"
 #include "Replication/Authority/IClientAuthoritySynchronizer.h"
 #include "Replication/Data/ReplicationStream.h"
 #include "Replication/Stream/IClientStreamSynchronizer.h"
 #include "Replication/Stream/MultiUserStreamId.h"
 
 #include "Containers/Set.h"
+#include "Delegates/Delegate.h"
 #include "HAL/Platform.h"
 
 class IConcertClientWorkspace;
@@ -26,7 +28,8 @@ namespace UE::MultiUserClient::Replication
 	{
 	public:
 		
-		FRejoinStreamAndAuthorityPredictor(const IConcertClientWorkspace& InWorkspace UE_LIFETIMEBOUND, FConcertClientInfo InClientInfo);
+		FRejoinStreamAndAuthorityPredictor(IConcertClientWorkspace& InWorkspace UE_LIFETIMEBOUND, FConcertClientInfo InClientInfo);
+		~FRejoinStreamAndAuthorityPredictor();
 
 		FORCEINLINE const FConcertBaseStreamInfo& GetPredictedStream() const { return PredictedStream; }
 		
@@ -41,15 +44,31 @@ namespace UE::MultiUserClient::Replication
 		virtual bool HasAuthorityOver(const FSoftObjectPath& ObjectPath) const override { return PredictedAuthority.Contains(ObjectPath); }
 		//~ End IClientAuthoritySynchronizer Interface
 
+		/** Broadcasts when PredictedStream changes. */
+		FSimpleMulticastDelegate& OnPredictionChanged() { return OnPredictionChangedDelegate; }
+
 	private:
 
+		/** Used to listen for activity changes. */
+		IConcertClientWorkspace& Workspace;
 		/** Client for which we're predicting the state. */
 		const FConcertClientInfo ClientInfo;
 
 		/** The stream content the local client thinks the offline client will have. */
 		FConcertBaseStreamInfo PredictedStream;
 		/** The most up to date server state of the remote client's authority. */
-		TSet<FSoftObjectPath> PredictedAuthority; 
+		TSet<FSoftObjectPath> PredictedAuthority;
+
+		/** Broadcasts when PredictedStream changes. */
+		FSimpleMulticastDelegate OnPredictionChangedDelegate;
+
+		/** Refreshes the prediction state if a new leave replication event is produced. */
+		void OnActivityAddedOrUpdated(
+			const FConcertClientInfo& ActivityClientInfo, const FConcertSyncActivity& Activity, const FStructOnScope& Summary
+			);
+
+		/** Refreshes PredicatedStream. */
+		void AnalyzeHistory();
 	};
 }
 
