Context menu ?Reassign this to? and ?Reassign children to? was compacted to ?Reassign all to?

diff --git a/Engine/Plugins/Developer/Concert/ConcertApp/MultiUserClient/Source/MultiUserClient/Private/Widgets/ActiveSession/Replication/Client/Context/ContextMenuUtils_Reassignment.cpp b/Engine/Plugins/Developer/Concert/ConcertApp/MultiUserClient/Source/MultiUserClient/Private/Widgets/ActiveSession/Replication/Client/Context/ContextMenuUtils_Reassignment.cpp
--- a/Engine/Plugins/Developer/Concert/ConcertApp/MultiUserClient/Source/MultiUserClient/Private/Widgets/ActiveSession/Replication/Client/Context/ContextMenuUtils_Reassignment.cpp
+++ b/Engine/Plugins/Developer/Concert/ConcertApp/MultiUserClient/Source/MultiUserClient/Private/Widgets/ActiveSession/Replication/Client/Context/ContextMenuUtils_Reassignment.cpp
@@ -84,21 +84,17 @@ namespace UE::MultiUserClient::Replication::ContextMenuUtils
 	{
 		const TArray<const FOnlineClient*> SortedClients = ClientUtils::GetSortedClientList(ConcertClient, ReplicationManager);
 		
-		MenuBuilder.BeginSection(TEXT("Reassign.This"), LOCTEXT("Reassign.This", "Reassign this to"));
-		Private::AddReassignSection(MenuBuilder, SortedClients, Private::FInlineObjectPathArray{ ContextObject.GetUniqueID() }, ConcertClient, ReassignmentLogic, MultiStreamEditor);
-		MenuBuilder.EndSection();
-
-		// Do not distract the user with more options if children have no assigned properties
-		const bool bHasChildrenWithProperties = ReassignmentLogic.IsAnyObjectOwned(Private::GetChildrenOfManagedObject(ObjectHierarchy, ContextObject));
-		if (bHasChildrenWithProperties)
-		{
-			MenuBuilder.BeginSection(TEXT("Reassign.Children"), LOCTEXT("Reassign.Children", "Reassign children to"));
-			Private::AddReassignSection(MenuBuilder, SortedClients, TAttribute<Private::FInlineObjectPathArray>::CreateLambda([ContextObject, &ObjectHierarchy]()
+		MenuBuilder.BeginSection(TEXT("Reassign.All"), LOCTEXT("Reassign.All", "Reassign all to"));
+		Private::AddReassignSection(MenuBuilder, SortedClients,
+			TAttribute<Private::FInlineObjectPathArray>::CreateLambda([ContextObject, &ObjectHierarchy]()
 			{
-				return Private::GetChildrenOfManagedObject(ObjectHierarchy, ContextObject);
-			}), ConcertClient, ReassignmentLogic, MultiStreamEditor);
-			MenuBuilder.EndSection();
-		}
+				Private::FInlineObjectPathArray Result = Private::GetChildrenOfManagedObject(ObjectHierarchy, ContextObject);
+				Result.Add(ContextObject.GetUniqueID());
+				return Result;
+			}),
+			ConcertClient, ReassignmentLogic, MultiStreamEditor
+			);
+		MenuBuilder.EndSection();
 	}
 }
 
