Fix: Reassign context menu shows ticks next to clients that do not own the object

diff --git a/Engine/Plugins/Developer/Concert/ConcertApp/MultiUserClient/Source/MultiUserClient/Private/Replication/Submission/MultiEdit/ReassignObjectPropertiesLogic.cpp b/Engine/Plugins/Developer/Concert/ConcertApp/MultiUserClient/Source/MultiUserClient/Private/Replication/Submission/MultiEdit/ReassignObjectPropertiesLogic.cpp
--- a/Engine/Plugins/Developer/Concert/ConcertApp/MultiUserClient/Source/MultiUserClient/Private/Replication/Submission/MultiEdit/ReassignObjectPropertiesLogic.cpp
+++ b/Engine/Plugins/Developer/Concert/ConcertApp/MultiUserClient/Source/MultiUserClient/Private/Replication/Submission/MultiEdit/ReassignObjectPropertiesLogic.cpp
@@ -308,7 +308,7 @@ namespace UE::MultiUserClient::Replication
 		{
 			ReassignObjectProperties::EnumerateOwningOnlineClients(ClientManager, Objects[i],[&bOwnsAny, &TargetClientId](const FGuid& ClientId)
 			{
-				bOwnsAny = true;
+				bOwnsAny |= ClientId == TargetClientId;
 				return bOwnsAny ? EBreakBehavior::Break : EBreakBehavior::Continue;
 			});
 		}
