Fix: Reassigning property and then clearing it on the reassigned to client's machine removes the object

diff --git a/Engine/Plugins/Developer/Concert/ConcertApp/MultiUserClient/Source/MultiUserClient/Private/Replication/Client/UnifiedClientView.cpp b/Engine/Plugins/Developer/Concert/ConcertApp/MultiUserClient/Source/MultiUserClient/Private/Replication/Client/UnifiedClientView.cpp
--- a/Engine/Plugins/Developer/Concert/ConcertApp/MultiUserClient/Source/MultiUserClient/Private/Replication/Client/UnifiedClientView.cpp
+++ b/Engine/Plugins/Developer/Concert/ConcertApp/MultiUserClient/Source/MultiUserClient/Private/Replication/Client/UnifiedClientView.cpp
@@ -56,6 +56,11 @@ namespace UE::MultiUserClient::Replication
 		return Result;
 	}
 
+	FGuid FUnifiedClientView::GetLocalClient() const
+	{
+		return OnlineClientManager.GetLocalClient().GetEndpointId();
+	}
+
 	TOptional<FConcertClientInfo> FUnifiedClientView::GetClientInfoByEndpoint(const FGuid& EndpointId) const
 	{
 		FConcertClientInfo ClientInfo;
diff --git a/Engine/Plugins/Developer/Concert/ConcertApp/MultiUserClient/Source/MultiUserClient/Private/Replication/Client/UnifiedClientView.h b/Engine/Plugins/Developer/Concert/ConcertApp/MultiUserClient/Source/MultiUserClient/Private/Replication/Client/UnifiedClientView.h
--- a/Engine/Plugins/Developer/Concert/ConcertApp/MultiUserClient/Source/MultiUserClient/Private/Replication/Client/UnifiedClientView.h
+++ b/Engine/Plugins/Developer/Concert/ConcertApp/MultiUserClient/Source/MultiUserClient/Private/Replication/Client/UnifiedClientView.h
@@ -59,6 +59,7 @@ namespace UE::MultiUserClient::Replication
 
 		TArray<FGuid> GetClients() const;
 		TArray<FGuid> GetOnlineClients() const;
+		FGuid GetLocalClient() const;
 		
 		/** @return Client info associated with EndpointId. */
 		TOptional<FConcertClientInfo> GetClientInfoByEndpoint(const FGuid& EndpointId) const;
diff --git a/Engine/Plugins/Developer/Concert/ConcertApp/MultiUserClient/Source/MultiUserClient/Private/Widgets/ActiveSession/Replication/Client/Multi/Columns/AssignProperty/AssignPropertyModel.cpp b/Engine/Plugins/Developer/Concert/ConcertApp/MultiUserClient/Source/MultiUserClient/Private/Widgets/ActiveSession/Replication/Client/Multi/Columns/AssignProperty/AssignPropertyModel.cpp
--- a/Engine/Plugins/Developer/Concert/ConcertApp/MultiUserClient/Source/MultiUserClient/Private/Widgets/ActiveSession/Replication/Client/Multi/Columns/AssignProperty/AssignPropertyModel.cpp
+++ b/Engine/Plugins/Developer/Concert/ConcertApp/MultiUserClient/Source/MultiUserClient/Private/Widgets/ActiveSession/Replication/Client/Multi/Columns/AssignProperty/AssignPropertyModel.cpp
@@ -12,6 +12,8 @@
 
 #include <type_traits>
 
+#include "Misc/ObjectUtils.h"
+
 #define LOCTEXT_NAMESPACE "FAssignPropertyModel"
 
 namespace UE::MultiUserClient::Replication::MultiStreamColumns::AssignPropertyModel
@@ -36,11 +38,33 @@ namespace UE::MultiUserClient::Replication::MultiStreamColumns::AssignPropertyMo
 			ClientEditModel.AddProperties(ObjectPath, AddedProperties);
 		}
 	}
+
+	static void AddOwningActorIfHierarchyIsEmpty(ConcertSharedSlate::IEditableReplicationStreamModel& ClientEditModel, const FSoftObjectPath& ObjectPath)
+	{
+		const TOptional<FSoftObjectPath> ActorPath = ConcertSyncCore::GetActorPathIn(ObjectPath);
+		UObject* ResolvedObject = ActorPath ? ActorPath->ResolveObject() : nullptr;
+		if (!ResolvedObject || ClientEditModel.ContainsObjects({ *ActorPath }))
+		{
+			return;
+		}
+
+		bool bHasChildren = false;
+		ClientEditModel.ForEachSubobject(*ActorPath, [&bHasChildren](auto)
+		{
+			bHasChildren = true;
+			return bHasChildren ? EBreakBehavior::Break : EBreakBehavior::Continue;
+		});
+		if (!bHasChildren)
+		{
+			ClientEditModel.AddObjects({ ResolvedObject });
+		}
+	}
 	
 	static void RemovePropertiesFromClient(
 		ConcertSharedSlate::IEditableReplicationStreamModel& ClientEditModel,
 		TConstArrayView<TSoftObjectPtr<>> Objects,
-		const FConcertPropertyChain& Property
+		const FConcertPropertyChain& Property,
+		bool bIsLocalClient
 		)
 	{
 		for (const TSoftObjectPtr<>& Object : Objects)
@@ -60,11 +84,20 @@ namespace UE::MultiUserClient::Replication::MultiStreamColumns::AssignPropertyMo
 			// If the user does not want the actor anymore, they should click it and delete it.
 			const UClass* ObjectClass = ClassPath.IsValid() ? ClassPath.TryLoadClass<UObject>() : nullptr;
 			UE_CLOG(ClassPath.IsValid() && !ObjectClass, LogConcert, Warning, TEXT("SAssignPropertyComboBox: Failed to resolve class %s"), *ClassPath.ToString());
-			const bool bIsTopLevelObject = ObjectClass && !ObjectClass->IsChildOf<AActor>();
-			if (bIsTopLevelObject)
+			const bool bIsSubobject = ObjectClass && !ObjectClass->IsChildOf<AActor>();
+			if (bIsSubobject)
 			{
 				ClientEditModel.RemoveObjects({ ObjectPath });
 			}
+
+			// Scenario: 1. We had nothing assigned, 2. Remote client assigns some property of some component to us 3. Now, the property is cleared.
+			// The remote assignment op from step 2 does not add the owning actor.
+			// If the local client is clearing the property, we'd now remove the last object from the hierarchy, thus removing it from the UI.
+			// That feels unnatural. To prevent it, add the owning actor to keep the hierarchy in the UI.
+			if (bIsLocalClient)
+			{
+				AddOwningActorIfHierarchyIsEmpty(ClientEditModel, ObjectPath);
+			}
 		}
 	}
 	
@@ -82,7 +115,8 @@ namespace UE::MultiUserClient::Replication::MultiStreamColumns::AssignPropertyMo
 			const TSharedPtr<ConcertSharedSlate::IEditableReplicationStreamModel> Stream = ClientView.GetEditableClientStreamById(EndpointId);
 			if (Stream && ShouldRemoveFromClient(EndpointId))
 			{
-				RemovePropertiesFromClient(*Stream, Objects, Property);
+				const bool bIsLocalClient = EndpointId == ClientView.GetLocalClient();
+				RemovePropertiesFromClient(*Stream, Objects, Property, bIsLocalClient);
 			}
 			return EBreakBehavior::Continue;
 		});
