Introduce UMultiUserReplicationSubsystem::GetOfflineClientIds and UMultiUserReplicationSubsystem::OnOfflineClientsChanged.
This allows VCam's MultiUserModifier to respond directly to the offline client list changing.

diff --git a/Engine/Plugins/Developer/Concert/ConcertApp/MultiUserClient/Source/MultiUserClient/Private/Replication/MultiUserReplicationManager.cpp b/Engine/Plugins/Developer/Concert/ConcertApp/MultiUserClient/Source/MultiUserClient/Private/Replication/MultiUserReplicationManager.cpp
--- a/Engine/Plugins/Developer/Concert/ConcertApp/MultiUserClient/Source/MultiUserClient/Private/Replication/MultiUserReplicationManager.cpp
+++ b/Engine/Plugins/Developer/Concert/ConcertApp/MultiUserClient/Source/MultiUserClient/Private/Replication/MultiUserReplicationManager.cpp
@@ -141,9 +141,18 @@ namespace UE::MultiUserClient::Replication
 
 	void FMultiUserReplicationManager::SetupClientConnectionEvents()
 	{
-		FOnlineClientManager& ClientManager = ConnectedState->OnlineClientManager;
-		ClientManager.ForEachClient([this](FOnlineClient& InClient){ SetupClientDelegates(InClient); return EBreakBehavior::Continue; });
-		ClientManager.OnPostRemoteClientAdded().AddRaw(this, &FMultiUserReplicationManager::OnReplicationClientConnected);
+		FOnlineClientManager& OnlineClientManager = ConnectedState->OnlineClientManager;
+		OnlineClientManager.ForEachClient([this](FOnlineClient& InClient){ SetupClientDelegates(InClient); return EBreakBehavior::Continue; });
+		OnlineClientManager.OnPostRemoteClientAdded().AddRaw(this, &FMultiUserReplicationManager::OnReplicationClientConnected);
+
+		FOfflineClientManager& OfflineClientManager = ConnectedState->OfflineClientManager;
+		OfflineClientManager.OnClientsChanged().AddRaw(this, &FMultiUserReplicationManager::OnInternalOfflineClientsChanged);
+	}
+
+	void FMultiUserReplicationManager::SetupClientDelegates(FOnlineClient& InClient) const
+	{
+		InClient.GetStreamSynchronizer().OnServerStreamChanged().AddRaw(this, &FMultiUserReplicationManager::OnClientStreamServerStateChanged, InClient.GetEndpointId());
+		InClient.GetAuthoritySynchronizer().OnServerAuthorityChanged().AddRaw(this, &FMultiUserReplicationManager::OnClientAuthorityServerStateChanged, InClient.GetEndpointId());
 	}
 
 	void FMultiUserReplicationManager::OnClientStreamServerStateChanged(const FGuid EndpointId) const
@@ -158,12 +167,6 @@ namespace UE::MultiUserClient::Replication
 		OnAuthorityServerStateChangedDelegate.Broadcast(EndpointId);
 	}
 
-	void FMultiUserReplicationManager::SetupClientDelegates(FOnlineClient& InClient) const
-	{
-		InClient.GetStreamSynchronizer().OnServerStreamChanged().AddRaw(this, &FMultiUserReplicationManager::OnClientStreamServerStateChanged, InClient.GetEndpointId());
-		InClient.GetAuthoritySynchronizer().OnServerAuthorityChanged().AddRaw(this, &FMultiUserReplicationManager::OnClientAuthorityServerStateChanged, InClient.GetEndpointId());
-	}
-
 	const FConcertObjectReplicationMap* FMultiUserReplicationManager::FindReplicationMapForClient(const FGuid& ClientId) const
 	{
 		if (ConnectedState && ensureMsgf(IsInGameThread(), TEXT("To simplify implementation, only calls from game thread are allowed.")))
diff --git a/Engine/Plugins/Developer/Concert/ConcertApp/MultiUserClient/Source/MultiUserClient/Private/Replication/MultiUserReplicationManager.h b/Engine/Plugins/Developer/Concert/ConcertApp/MultiUserClient/Source/MultiUserClient/Private/Replication/MultiUserReplicationManager.h
--- a/Engine/Plugins/Developer/Concert/ConcertApp/MultiUserClient/Source/MultiUserClient/Private/Replication/MultiUserReplicationManager.h
+++ b/Engine/Plugins/Developer/Concert/ConcertApp/MultiUserClient/Source/MultiUserClient/Private/Replication/MultiUserReplicationManager.h
@@ -103,6 +103,7 @@ namespace UE::MultiUserClient::Replication
 		virtual bool FindOfflineClient(const FGuid& ClientId, TFunctionRef<void(const IOfflineReplicationClient&)> Callback) const override;
 		virtual FOnServerStateChanged& OnStreamServerStateChanged() override { return OnStreamServerStateChangedDelegate; }
 		virtual FOnServerStateChanged& OnAuthorityServerStateChanged() override { return OnAuthorityServerStateChangedDelegate; }
+		virtual FOnOfflineClientsChanged OnOfflineClientsChanged() override { return OnOfflineClientsChangedDelegate; }
 		//~ End IMultiUserReplication Interface
 
 	private:
@@ -172,6 +173,8 @@ namespace UE::MultiUserClient::Replication
 		FOnServerStateChanged OnStreamServerStateChangedDelegate;
 		/** Triggers when a client's known server state has changed. */
 		FOnServerStateChanged OnAuthorityServerStateChangedDelegate;
+		/** Triggers when the offline clients have changed. */
+		FOnOfflineClientsChanged OnOfflineClientsChangedDelegate;
 
 		/** Callback into Concert for when client connection has changed. */
 		void OnSessionConnectionChanged(IConcertClientSession& ConcertClientSession, EConcertConnectionStatus ConcertConnectionStatus);
@@ -187,10 +190,13 @@ namespace UE::MultiUserClient::Replication
 
 		/** Sets up delegates for implementing the broadcasting of OnStreamServerStateChangedDelegate and OnAuthorityServerStateChangedDelegate. */
 		void SetupClientConnectionEvents();
-		void OnClientStreamServerStateChanged(const FGuid EndpointId) const;
-		void OnClientAuthorityServerStateChanged(const FGuid EndpointId) const;
+		void OnInternalOfflineClientsChanged() const { OnOfflineClientsChangedDelegate.Broadcast(); }
 		void OnReplicationClientConnected(FRemoteClient& RemoteClient) const { SetupClientDelegates(RemoteClient); }
+
+		/** Sets up delegates for reacting to content changes. */
 		void SetupClientDelegates(FOnlineClient& InClient) const;
+		void OnClientStreamServerStateChanged(const FGuid EndpointId) const;
+		void OnClientAuthorityServerStateChanged(const FGuid EndpointId) const;
 	};
 }
 
diff --git a/Engine/Plugins/Developer/Concert/ConcertApp/MultiUserClient/Source/MultiUserClient/Public/Replication/IMultiUserReplication.h b/Engine/Plugins/Developer/Concert/ConcertApp/MultiUserClient/Source/MultiUserClient/Public/Replication/IMultiUserReplication.h
--- a/Engine/Plugins/Developer/Concert/ConcertApp/MultiUserClient/Source/MultiUserClient/Public/Replication/IMultiUserReplication.h
+++ b/Engine/Plugins/Developer/Concert/ConcertApp/MultiUserClient/Source/MultiUserClient/Public/Replication/IMultiUserReplication.h
@@ -113,6 +113,10 @@ namespace UE::MultiUserClient
 		/** @return Delegate that triggers when the given client's known server state has changed. */
 		virtual FOnServerStateChanged& OnAuthorityServerStateChanged() = 0;
 
+		DECLARE_MULTICAST_DELEGATE(FOnOfflineClientsChanged);
+		/** @return Delegate that triggers when the endpoints considered offline have changed. */
+		virtual FOnOfflineClientsChanged OnOfflineClientsChanged() = 0;
+
 		virtual ~IMultiUserReplication() = default;
 	};
 }
diff --git a/Engine/Plugins/Developer/Concert/ConcertApp/MultiUserClient/Source/MultiUserClientLibrary/Private/Replication/MultiUserReplicationSubsystem.cpp b/Engine/Plugins/Developer/Concert/ConcertApp/MultiUserClient/Source/MultiUserClientLibrary/Private/Replication/MultiUserReplicationSubsystem.cpp
--- a/Engine/Plugins/Developer/Concert/ConcertApp/MultiUserClient/Source/MultiUserClientLibrary/Private/Replication/MultiUserReplicationSubsystem.cpp
+++ b/Engine/Plugins/Developer/Concert/ConcertApp/MultiUserClient/Source/MultiUserClientLibrary/Private/Replication/MultiUserReplicationSubsystem.cpp
@@ -148,6 +148,25 @@ TArray<FMultiUserClientDisplayInfo> UMultiUserReplicationSubsystem::GetOwningOff
 	return Result;
 }
 
+TArray<FGuid> UMultiUserReplicationSubsystem::GetOfflineClientIds() const
+{
+	TArray<FGuid> Clients;
+	
+#if WITH_CONCERT
+	UE::MultiUserClient::IMultiUserReplication* ReplicationInterface = IMultiUserClientModule::Get().GetReplication();
+	if (ensureMsgf(ReplicationInterface, TEXT("We expected it to always be valid.")))
+	{
+		ReplicationInterface->ForEachOfflineClient([&Clients](const UE::MultiUserClient::IOfflineReplicationClient& Client)
+		{
+			Clients.Add(Client.GetLastAssociatedEndpoint());
+			return EBreakBehavior::Continue;
+		});
+	}
+#endif
+	
+	return Clients;
+}
+
 void UMultiUserReplicationSubsystem::Initialize(FSubsystemCollectionBase& Collection)
 {
 	Super::Initialize(Collection);  
@@ -161,6 +180,7 @@ void UMultiUserReplicationSubsystem::Initialize(FSubsystemCollectionBase& Collec
 
 		ReplicationInterface->OnStreamServerStateChanged().AddUObject(this, &UMultiUserReplicationSubsystem::OnClientStreamsChanged);
 		ReplicationInterface->OnAuthorityServerStateChanged().AddUObject(this, &UMultiUserReplicationSubsystem::OnClientAuthorityChanged);
+		ReplicationInterface->OnOfflineClientsChanged().AddUObject(this, &UMultiUserReplicationSubsystem::OnInternalOfflineClientsChanged);
 	}
 #endif
 }
@@ -180,6 +200,7 @@ void UMultiUserReplicationSubsystem::Deinitialize()
 			
 			ReplicationInterface->OnStreamServerStateChanged().RemoveAll(this);
 			ReplicationInterface->OnAuthorityServerStateChanged().RemoveAll(this);
+			ReplicationInterface->OnOfflineClientsChanged().RemoveAll(this);
 		}
 	}
 #endif
diff --git a/Engine/Plugins/Developer/Concert/ConcertApp/MultiUserClient/Source/MultiUserClientLibrary/Public/Replication/MultiUserReplicationSubsystem.h b/Engine/Plugins/Developer/Concert/ConcertApp/MultiUserClient/Source/MultiUserClientLibrary/Public/Replication/MultiUserReplicationSubsystem.h
--- a/Engine/Plugins/Developer/Concert/ConcertApp/MultiUserClient/Source/MultiUserClientLibrary/Public/Replication/MultiUserReplicationSubsystem.h
+++ b/Engine/Plugins/Developer/Concert/ConcertApp/MultiUserClient/Source/MultiUserClientLibrary/Public/Replication/MultiUserReplicationSubsystem.h
@@ -23,6 +23,7 @@ public:
 	// This would be the right place to expose additional MU specific replication functions in the future
 	
 	DECLARE_DYNAMIC_MULTICAST_DELEGATE_OneParam(FOnServerStateChanged, const FGuid&, EndpointId);
+	DECLARE_DYNAMIC_MULTICAST_DELEGATE(FOnOfflineClientsChanged);
 	
 	/**
 	 * @return Whether the client is replicating the object.
@@ -87,13 +88,17 @@ public:
 	 * @return A list of client descriptions representing offline clients that will attempt to reclaim properties 
 	 * associated with the object when they rejoin.
 	 */
-	UFUNCTION(BlueprintPure, Category = "Multi-user", meta = (Keywords = "Owning Offline Reclaim Join Rejoin Client Find Stream Registered"))
+	UFUNCTION(BlueprintPure, Category = "Multi-user", meta = (Keywords = "Owning Offline Disconnected  Reclaim Join Rejoin Client Find Stream Registered"))
 	TArray<FMultiUserClientDisplayInfo> GetOwningOfflineClients(const FSoftObjectPath& ObjectPath) const;
 
 	/** @return Whether any offline clients will try to register properties for ObjectPath upon rejoining.*/
-	UFUNCTION(BlueprintPure, Category = "Multi-user", meta = (Keywords = "Owning Offline Reclaim Join Rejoin Client Find Stream Registered"))
+	UFUNCTION(BlueprintPure, Category = "Multi-user", meta = (Keywords = "Owning Offline Disconnected  Reclaim Join Rejoin Client Find Stream Registered"))
 	bool IsOwnedByOfflineClient(const FSoftObjectPath& ObjectPath) const { return !GetOwningOfflineClients(ObjectPath).IsEmpty(); }
-	
+
+	/** @return The list of offline clients. Each entry is an endpoint ID that a user had in the past. */
+	UFUNCTION(BlueprintPure, Category = "Multi-user", meta = (Keywords = "Get Offline Disconnected Clients"))
+	TArray<FGuid> GetOfflineClientIds() const;
+
 	//~ Begin USubsystem Interface
 	virtual void Initialize(FSubsystemCollectionBase& Collection) override;
 	virtual void Deinitialize() override;
@@ -113,6 +118,10 @@ private:
 	UPROPERTY(BlueprintAssignable, Category = "Multi-user")
 	FOnServerStateChanged OnClientAuthorityServerStateChanged;
 
+	/** Event triggered when the local list of offline clients has changed. */
+	UPROPERTY(BlueprintAssignable, Category = "Multi-user")
+	FOnOfflineClientsChanged OnOfflineClientsChanged;
+
 	/**
 	 * This is used only for when the adds an object through the Add button in the UI.
 	 * 
@@ -125,4 +134,5 @@ private:
 
 	void OnClientStreamsChanged(const FGuid& EndpointId) const { OnClientStreamServerStateChanged.Broadcast(EndpointId); }
 	void OnClientAuthorityChanged(const FGuid& EndpointId) const { OnClientAuthorityServerStateChanged.Broadcast(EndpointId); }
+	void OnInternalOfflineClientsChanged() const { OnOfflineClientsChanged.Broadcast(); }
 };
