The game visibility column in the Levels tab will now be
- shown by MultiUserClient upon joining a session,
- be hidden by MultiUserClient upon leaving a session if the column was hidden when they joined

diff --git a/Engine/Plugins/Developer/Concert/ConcertApp/MultiUserClient/Source/MultiUserClient/MultiUserClient.Build.cs b/Engine/Plugins/Developer/Concert/ConcertApp/MultiUserClient/Source/MultiUserClient/MultiUserClient.Build.cs
--- a/Engine/Plugins/Developer/Concert/ConcertApp/MultiUserClient/Source/MultiUserClient/MultiUserClient.Build.cs
+++ b/Engine/Plugins/Developer/Concert/ConcertApp/MultiUserClient/Source/MultiUserClient/MultiUserClient.Build.cs
@@ -58,6 +58,7 @@ namespace UnrealBuildTool.Rules
 						"Sequencer",
 						"ToolMenus",
 						"UnrealEd",
+						"WorldBrowser",
 						"WorkspaceMenuStructure",
 					}
 				);
diff --git a/Engine/Plugins/Developer/Concert/ConcertApp/MultiUserClient/Source/MultiUserClient/Private/Misc/GameVisibilityColumnToggler.cpp b/Engine/Plugins/Developer/Concert/ConcertApp/MultiUserClient/Source/MultiUserClient/Private/Misc/GameVisibilityColumnToggler.cpp
new file mode 100644
--- /dev/null
+++ b/Engine/Plugins/Developer/Concert/ConcertApp/MultiUserClient/Source/MultiUserClient/Private/Misc/GameVisibilityColumnToggler.cpp
@@ -0,0 +1,52 @@
+﻿// Copyright Epic Games, Inc. All Rights Reserved.
+
+#include "GameVisibilityColumnToggler.h"
+
+#if WITH_EDITOR
+
+#include "IConcertSyncClient.h"
+
+#include "ILevelEditor.h"
+#include "IWorldHierarchy.h"
+#include "LevelEditor.h"
+#include "WorldHierarchyColumns.h"
+
+namespace UE::MultiUserClient
+{
+	FGameVisibilityColumnToggler::FGameVisibilityColumnToggler(TSharedRef<IConcertSyncClient> InMultiUserClient)
+		: MultiUserClient(InMultiUserClient)
+	{
+		MultiUserClient->OnSyncSessionStartup().AddRaw(this, &FGameVisibilityColumnToggler::OnStartSession);
+		MultiUserClient->OnSyncSessionShutdown().AddRaw(this, &FGameVisibilityColumnToggler::OnStopSession);
+	}
+
+	FGameVisibilityColumnToggler::~FGameVisibilityColumnToggler()
+	{
+		MultiUserClient->OnSyncSessionStartup().RemoveAll(this);
+		MultiUserClient->OnSyncSessionShutdown().RemoveAll(this);
+	}
+
+	void FGameVisibilityColumnToggler::OnStartSession(const IConcertSyncClient*)
+	{
+		FLevelEditorModule& LevelEditorModule = FModuleManager::GetModuleChecked<FLevelEditorModule>("LevelEditor");
+		const TSharedPtr<ILevelEditor> LevelEditor = LevelEditorModule.GetFirstLevelEditor();
+		if (const TSharedPtr<WorldHierarchy::IWorldHierarchy> LevelsTab = LevelEditor->GetWorldHierarchy().Pin())
+		{
+			bHideVisibilityColumnOnSessionLeave = !LevelsTab->IsColumnVisible(WorldHierarchy::HierarchyColumns::ColumnID_GameVisibility);
+			LevelsTab->SetColumnVisible(WorldHierarchy::HierarchyColumns::ColumnID_GameVisibility, true);
+		}
+	}
+
+	void FGameVisibilityColumnToggler::OnStopSession(const IConcertSyncClient*)
+	{
+		FLevelEditorModule& LevelEditorModule = FModuleManager::GetModuleChecked<FLevelEditorModule>("LevelEditor");
+		const TSharedPtr<ILevelEditor> LevelEditor = LevelEditorModule.GetFirstLevelEditor();
+		if (const TSharedPtr<WorldHierarchy::IWorldHierarchy> LevelsTab = LevelEditor->GetWorldHierarchy().Pin();
+			LevelsTab && bHideVisibilityColumnOnSessionLeave)
+		{
+			LevelsTab->SetColumnVisible(WorldHierarchy::HierarchyColumns::ColumnID_GameVisibility, false);
+		}
+	}
+}
+
+#endif
\ No newline at end of file
diff --git a/Engine/Plugins/Developer/Concert/ConcertApp/MultiUserClient/Source/MultiUserClient/Private/Misc/GameVisibilityColumnToggler.h b/Engine/Plugins/Developer/Concert/ConcertApp/MultiUserClient/Source/MultiUserClient/Private/Misc/GameVisibilityColumnToggler.h
new file mode 100644
--- /dev/null
+++ b/Engine/Plugins/Developer/Concert/ConcertApp/MultiUserClient/Source/MultiUserClient/Private/Misc/GameVisibilityColumnToggler.h
@@ -0,0 +1,41 @@
+﻿// Copyright Epic Games, Inc. All Rights Reserved.
+
+#pragma once
+
+#if WITH_EDITOR
+
+#include "Templates/SharedPointer.h"
+#include "Templates/UnrealTemplate.h"
+
+class IConcertSyncClient;
+
+namespace UE::MultiUserClient
+{
+	/**
+	 * Toggles the GameVisibility column in the Levels tab when joining a MU session.
+	 * 
+	 * This is done to improve the UX of VP users: it allows them to show & hide levels on connected clients launched with -game,
+	 * such as nDisplay nodes.
+	 */
+	class FGameVisibilityColumnToggler : FNoncopyable
+	{
+	public:
+		
+		FGameVisibilityColumnToggler(TSharedRef<IConcertSyncClient> InMultiUserClient);
+		~FGameVisibilityColumnToggler();
+		
+	private:
+
+		/** The client to which we listen for session changes to*/
+		TSharedRef<IConcertSyncClient> MultiUserClient;
+
+		/** This is false if the visibility column had been shown when the client joined the session. */
+		bool bHideVisibilityColumnOnSessionLeave = true;
+
+		void OnStartSession(const IConcertSyncClient* Client);
+		void OnStopSession(const IConcertSyncClient* Client);
+	};
+}
+
+#endif
+
diff --git a/Engine/Plugins/Developer/Concert/ConcertApp/MultiUserClient/Source/MultiUserClient/Private/MultiUserClientModule.cpp b/Engine/Plugins/Developer/Concert/ConcertApp/MultiUserClient/Source/MultiUserClient/Private/MultiUserClientModule.cpp
--- a/Engine/Plugins/Developer/Concert/ConcertApp/MultiUserClient/Source/MultiUserClient/Private/MultiUserClientModule.cpp
+++ b/Engine/Plugins/Developer/Concert/ConcertApp/MultiUserClient/Source/MultiUserClient/Private/MultiUserClientModule.cpp
@@ -52,6 +52,7 @@
 #include "DetailLayoutBuilder.h"
 #include "DetailCategoryBuilder.h"
 #include "DetailWidgetRow.h"
+#include "Misc/GameVisibilityColumnToggler.h"
 
 
 #include "Widgets/SConcertBrowser.h"
@@ -629,10 +630,17 @@ public:
 		// Hook UI elements in the tool bar (and setup commands).
 		RegisterUI();
 		RegisterLogging();
+#if WITH_EDITOR
+		VisibilityColumnToggler = MakeUnique<UE::MultiUserClient::FGameVisibilityColumnToggler>(MultiUserClient.ToSharedRef());
+#endif
 	}
 
 	virtual void ShutdownModule() override
 	{
+#if WITH_EDITOR
+		VisibilityColumnToggler.Reset();
+#endif
+		
 		// Unhook AppPreExit and call it
 		FCoreDelegates::OnPreExit.RemoveAll(this);
 		HandleAppPreExit();
@@ -1411,6 +1419,11 @@ private:
 	/** Interacts with the replication system on behalf of Multi-User. */
 	TSharedPtr<UE::MultiUserClient::Replication::FMultiUserReplicationManager> ReplicationManager;
 
+#if WITH_EDITOR
+	/** Toggles the game visibility column in the Levels tab when joining & leaving sessions. */
+	TUniquePtr<UE::MultiUserClient::FGameVisibilityColumnToggler> VisibilityColumnToggler;
+#endif
+
 	/** True if the tab spawners have been registered for this module */
 	bool bHasRegisteredTabSpawners = false;
 
