Fix typo causing event to be returned by value instead of by reference

diff --git a/Engine/Plugins/Developer/Concert/ConcertApp/MultiUserClient/Source/MultiUserClient/Private/Replication/MultiUserReplicationManager.h b/Engine/Plugins/Developer/Concert/ConcertApp/MultiUserClient/Source/MultiUserClient/Private/Replication/MultiUserReplicationManager.h
--- a/Engine/Plugins/Developer/Concert/ConcertApp/MultiUserClient/Source/MultiUserClient/Private/Replication/MultiUserReplicationManager.h
+++ b/Engine/Plugins/Developer/Concert/ConcertApp/MultiUserClient/Source/MultiUserClient/Private/Replication/MultiUserReplicationManager.h
@@ -103,7 +103,7 @@ namespace UE::MultiUserClient::Replication
 		virtual bool FindOfflineClient(const FGuid& ClientId, TFunctionRef<void(const IOfflineReplicationClient&)> Callback) const override;
 		virtual FOnServerStateChanged& OnStreamServerStateChanged() override { return OnStreamServerStateChangedDelegate; }
 		virtual FOnServerStateChanged& OnAuthorityServerStateChanged() override { return OnAuthorityServerStateChangedDelegate; }
-		virtual FOnOfflineClientsChanged OnOfflineClientsChanged() override { return OnOfflineClientsChangedDelegate; }
+		virtual FOnOfflineClientsChanged& OnOfflineClientsChanged() override { return OnOfflineClientsChangedDelegate; }
 		//~ End IMultiUserReplication Interface
 
 	private:
diff --git a/Engine/Plugins/Developer/Concert/ConcertApp/MultiUserClient/Source/MultiUserClient/Public/Replication/IMultiUserReplication.h b/Engine/Plugins/Developer/Concert/ConcertApp/MultiUserClient/Source/MultiUserClient/Public/Replication/IMultiUserReplication.h
--- a/Engine/Plugins/Developer/Concert/ConcertApp/MultiUserClient/Source/MultiUserClient/Public/Replication/IMultiUserReplication.h
+++ b/Engine/Plugins/Developer/Concert/ConcertApp/MultiUserClient/Source/MultiUserClient/Public/Replication/IMultiUserReplication.h
@@ -115,7 +115,7 @@ namespace UE::MultiUserClient
 
 		DECLARE_MULTICAST_DELEGATE(FOnOfflineClientsChanged);
 		/** @return Delegate that triggers when the endpoints considered offline have changed. */
-		virtual FOnOfflineClientsChanged OnOfflineClientsChanged() = 0;
+		virtual FOnOfflineClientsChanged& OnOfflineClientsChanged() = 0;
 
 		virtual ~IMultiUserReplication() = default;
 	};
