Make UMultiUserReplicationSubsystem::GetRegisteredObjects also return objects for offline clients

diff --git a/Engine/Plugins/Developer/Concert/ConcertApp/MultiUserClient/Source/MultiUserClient/Private/Replication/MultiUserReplicationManager.cpp b/Engine/Plugins/Developer/Concert/ConcertApp/MultiUserClient/Source/MultiUserClient/Private/Replication/MultiUserReplicationManager.cpp
--- a/Engine/Plugins/Developer/Concert/ConcertApp/MultiUserClient/Source/MultiUserClient/Private/Replication/MultiUserReplicationManager.cpp
+++ b/Engine/Plugins/Developer/Concert/ConcertApp/MultiUserClient/Source/MultiUserClient/Private/Replication/MultiUserReplicationManager.cpp
@@ -171,11 +171,19 @@ namespace UE::MultiUserClient::Replication
 	{
 		if (ConnectedState && ensureMsgf(IsInGameThread(), TEXT("To simplify implementation, only calls from game thread are allowed.")))
 		{
-			const FOnlineClient* ReplicationClient = ConnectedState->OnlineClientManager.FindClient(ClientId);
-			return ReplicationClient
-				? &ReplicationClient->GetStreamSynchronizer().GetServerState()
-				: nullptr;
+			if (const FOnlineClient* OnlineClient = ConnectedState->OnlineClientManager.FindClient(ClientId))
+			{
+				return &OnlineClient->GetStreamSynchronizer().GetServerState();
+			}
+
+			if (const FOfflineClient* OfflineClient = ConnectedState->OfflineClientManager.FindClient(ClientId))
+			{
+				return &OfflineClient->GetPredictedStream().ReplicationMap;
+			}
+
+			return nullptr;
 		}
+		
 		return nullptr;
 	}
 
@@ -183,10 +191,17 @@ namespace UE::MultiUserClient::Replication
 	{
 		if (ConnectedState && ensureMsgf(IsInGameThread(), TEXT("To simplify implementation, only calls from game thread are allowed.")))
 		{
-			const FOnlineClient* ReplicationClient = ConnectedState->OnlineClientManager.FindClient(ClientId);
-			return ReplicationClient
-				? &ReplicationClient->GetStreamSynchronizer().GetFrequencySettings()
-				: nullptr;
+			if (const FOnlineClient* OnlineClient = ConnectedState->OnlineClientManager.FindClient(ClientId))
+			{
+				return &OnlineClient->GetStreamSynchronizer().GetFrequencySettings();
+			}
+
+			if (const FOfflineClient* OfflineClient = ConnectedState->OfflineClientManager.FindClient(ClientId))
+			{
+				return &OfflineClient->GetPredictedStream().FrequencySettings;
+			}
+
+			return nullptr;
 		}
 		return nullptr;
 	}
diff --git a/Engine/Plugins/Developer/Concert/ConcertApp/MultiUserClient/Source/MultiUserClient/Public/Replication/IMultiUserReplication.h b/Engine/Plugins/Developer/Concert/ConcertApp/MultiUserClient/Source/MultiUserClient/Public/Replication/IMultiUserReplication.h
--- a/Engine/Plugins/Developer/Concert/ConcertApp/MultiUserClient/Source/MultiUserClient/Public/Replication/IMultiUserReplication.h
+++ b/Engine/Plugins/Developer/Concert/ConcertApp/MultiUserClient/Source/MultiUserClient/Public/Replication/IMultiUserReplication.h
@@ -27,7 +27,7 @@ namespace UE::MultiUserClient
 	public:
 
 		/**
-		 * @return Gets the last known server map of objects registered for replication for a given client.
+		 * @return Gets the last known server map of objects registered for replication for an online or offline client.
 		 * 
 		 * This server state is regularly polled whilst the local client state should always be in sync.
 		 * @note This function must be called from the game thread.
@@ -36,7 +36,7 @@ namespace UE::MultiUserClient
 		virtual const FConcertObjectReplicationMap* FindReplicationMapForClient(const FGuid& ClientId) const = 0;
 
 		/**
-		 * @return Gets the last known server object replication frequencies.
+		 * @return Gets the last known server object replication frequencies for an online or offline .
 		 * 
 		 * This server state is regularly polled whilst the local client state should always be in sync.
 		 * @note This function must be called from the game thread.
diff --git a/Engine/Plugins/Developer/Concert/ConcertApp/MultiUserClient/Source/MultiUserClientLibrary/Public/Replication/MultiUserReplicationSubsystem.h b/Engine/Plugins/Developer/Concert/ConcertApp/MultiUserClient/Source/MultiUserClientLibrary/Public/Replication/MultiUserReplicationSubsystem.h
--- a/Engine/Plugins/Developer/Concert/ConcertApp/MultiUserClient/Source/MultiUserClientLibrary/Public/Replication/MultiUserReplicationSubsystem.h
+++ b/Engine/Plugins/Developer/Concert/ConcertApp/MultiUserClient/Source/MultiUserClientLibrary/Public/Replication/MultiUserReplicationSubsystem.h
@@ -50,11 +50,13 @@ public:
 	TArray<FConcertPropertyChainWrapper> GetPropertiesRegisteredToObject(const FGuid& ClientId, const FSoftObjectPath& ObjectPath) const;
 
 	/**
-	 * Gets the objects the client has registered with the server.
+	 * Gets the objects the online or offline client has registered with the server.
 	 *
-	 * Just because an object is replicated, it does not mean that the object is being replicated.
-	 * Objects must be registered with the server so registration is just the first step.
-	 * GetReplicatedObjects() will always contain GetRegisteredObjects().
+	 * Just because an object is returned here, it does not mean that the object is being replicated:
+	 * - If ClientId is an offline client, then the object is not being replicated by that client.
+	 * - If ClientId is an online client, then the object(s) may only be registered with the server.
+	 *
+	 * To find out which objects are actually being replicated, use GetReplicatedObjects(), which will always contain GetRegisteredObjects().
 	 * 
 	 * @param ClientId The client of which to get the registered objects. 
 	 * @return The objects the client has
