Refactor UMultiUserReplicationSubsystem::GetOwningOfflineClients to return FGuid instead. A follow up CL will add a way to query the client display info.

diff --git a/Engine/Plugins/Developer/Concert/ConcertApp/MultiUserClient/Source/MultiUserClientLibrary/Private/Replication/MultiUserReplicationSubsystem.cpp b/Engine/Plugins/Developer/Concert/ConcertApp/MultiUserClient/Source/MultiUserClientLibrary/Private/Replication/MultiUserReplicationSubsystem.cpp
--- a/Engine/Plugins/Developer/Concert/ConcertApp/MultiUserClient/Source/MultiUserClientLibrary/Private/Replication/MultiUserReplicationSubsystem.cpp
+++ b/Engine/Plugins/Developer/Concert/ConcertApp/MultiUserClient/Source/MultiUserClientLibrary/Private/Replication/MultiUserReplicationSubsystem.cpp
@@ -126,9 +126,9 @@ TArray<FSoftObjectPath> UMultiUserReplicationSubsystem::GetReplicatedObjects(con
 	return {};
 }
 
-TArray<FMultiUserClientDisplayInfo> UMultiUserReplicationSubsystem::GetOwningOfflineClients(const FSoftObjectPath& ObjectPath) const
+TArray<FGuid> UMultiUserReplicationSubsystem::GetOwningOfflineClients(const FSoftObjectPath& ObjectPath) const
 {
-	TArray<FMultiUserClientDisplayInfo> Result;
+	TArray<FGuid> Result;
 	
 #if WITH_CONCERT
 	UE::MultiUserClient::IMultiUserReplication* ReplicationInterface = IMultiUserClientModule::Get().GetReplication();
@@ -138,7 +138,7 @@ TArray<FMultiUserClientDisplayInfo> UMultiUserReplicationSubsystem::GetOwningOff
 		{
 			if (Client.GetPredictedStream().ReplicationMap.HasProperties(ObjectPath))
 			{
-				Result.Emplace(Client.GetClientInfo());
+				Result.Emplace(Client.GetLastAssociatedEndpoint());
 			}
 			return EBreakBehavior::Continue;
 		});
diff --git a/Engine/Plugins/Developer/Concert/ConcertApp/MultiUserClient/Source/MultiUserClientLibrary/Public/Replication/MultiUserReplicationSubsystem.h b/Engine/Plugins/Developer/Concert/ConcertApp/MultiUserClient/Source/MultiUserClientLibrary/Public/Replication/MultiUserReplicationSubsystem.h
--- a/Engine/Plugins/Developer/Concert/ConcertApp/MultiUserClient/Source/MultiUserClientLibrary/Public/Replication/MultiUserReplicationSubsystem.h
+++ b/Engine/Plugins/Developer/Concert/ConcertApp/MultiUserClient/Source/MultiUserClientLibrary/Public/Replication/MultiUserReplicationSubsystem.h
@@ -91,7 +91,7 @@ public:
 	 * associated with the object when they rejoin.
 	 */
 	UFUNCTION(BlueprintPure, Category = "Multi-user", meta = (Keywords = "Owning Offline Disconnected  Reclaim Join Rejoin Client Find Stream Registered"))
-	TArray<FMultiUserClientDisplayInfo> GetOwningOfflineClients(const FSoftObjectPath& ObjectPath) const;
+	TArray<FGuid> GetOwningOfflineClients(const FSoftObjectPath& ObjectPath) const;
 
 	/** @return Whether any offline clients will try to register properties for ObjectPath upon rejoining.*/
 	UFUNCTION(BlueprintPure, Category = "Multi-user", meta = (Keywords = "Owning Offline Disconnected  Reclaim Join Rejoin Client Find Stream Registered"))
@@ -100,6 +100,7 @@ public:
 	/** @return The list of offline clients. Each entry is an endpoint ID that a user had in the past. */
 	UFUNCTION(BlueprintPure, Category = "Multi-user", meta = (Keywords = "Get Offline Disconnected Clients"))
 	TArray<FGuid> GetOfflineClientIds() const;
+	
 
 	//~ Begin USubsystem Interface
 	virtual void Initialize(FSubsystemCollectionBase& Collection) override;
