Introduce FMultiUserClientInfo constructor

diff --git a/Engine/Plugins/Developer/Concert/ConcertApp/MultiUserClient/Source/MultiUserClientLibrary/Private/MultiUserClientStatics.cpp b/Engine/Plugins/Developer/Concert/ConcertApp/MultiUserClient/Source/MultiUserClientLibrary/Private/MultiUserClientStatics.cpp
--- a/Engine/Plugins/Developer/Concert/ConcertApp/MultiUserClient/Source/MultiUserClientLibrary/Private/MultiUserClientStatics.cpp
+++ b/Engine/Plugins/Developer/Concert/ConcertApp/MultiUserClient/Source/MultiUserClientLibrary/Private/MultiUserClientStatics.cpp
@@ -27,14 +27,9 @@ namespace UE::MultiUserClientLibrary
 {
 	FMultiUserClientInfo ConvertClientInfo(const FGuid& ClientEndpointId, const FConcertClientInfo& ClientInfo)
 	{
-		FMultiUserClientInfo Result;
-		Result.ClientEndpointId = ClientEndpointId;
-		Result.DisplayName = ClientInfo.DisplayName;
-		Result.AvatarColor = ClientInfo.AvatarColor;
-		Result.Tags = ClientInfo.Tags;
-		return Result;
+		return FMultiUserClientInfo(ClientEndpointId, ClientInfo);
 	}
-
+	
 	FMultiUserConnectionError ConvertConnectionError(FConcertConnectionError Error)
 	{
 		FMultiUserConnectionError MUError;
@@ -80,6 +75,19 @@ namespace UE::MultiUserClientLibrary
 } // namespace UE::MultiUserClientLibrary
 #endif
 
+#if WITH_CONCERT
+FMultiUserClientInfo::FMultiUserClientInfo(const FConcertSessionClientInfo& ClientInfo)
+	: FMultiUserClientInfo(ClientInfo.ClientEndpointId, ClientInfo.ClientInfo)
+{}
+
+FMultiUserClientInfo::FMultiUserClientInfo(const FGuid& ClientId, const FConcertClientInfo& ClientInfo)
+	: ClientEndpointId(ClientId)
+	, DisplayName(ClientInfo.DisplayName)
+	, AvatarColor(ClientInfo.AvatarColor)
+	, Tags(ClientInfo.Tags)
+{}
+#endif
+
 UMultiUserClientStatics::UMultiUserClientStatics(const FObjectInitializer& ObjectInitializer)
 	: Super(ObjectInitializer)
 {
@@ -265,7 +273,7 @@ FMultiUserClientInfo UMultiUserClientStatics::GetLocalMultiUserClientInfo()
 
 			FGuid LocalClientEndpointId = ClientSession ? ClientSession->GetSessionClientEndpointId() : FGuid();
 			const FConcertClientInfo& LocalClientInfo = ClientSession ? ClientSession->GetLocalClientInfo() : ConcertClient->GetClientInfo();
-			ClientInfo = UE::MultiUserClientLibrary::ConvertClientInfo(LocalClientEndpointId, LocalClientInfo);
+			ClientInfo = FMultiUserClientInfo(LocalClientEndpointId, LocalClientInfo);
 		}
 	}
 #endif
@@ -315,7 +323,7 @@ bool UMultiUserClientStatics::GetMultiUserClientInfoByName(const FString& Client
 			const FConcertClientInfo& LocalClientInfo = ClientSession ? ClientSession->GetLocalClientInfo() : ConcertClient->GetClientInfo();
 			if (ClientName == LocalClientInfo.DisplayName)
 			{
-				ClientInfo = UE::MultiUserClientLibrary::ConvertClientInfo(ClientSession ? ClientSession->GetSessionClientEndpointId() : FGuid(), LocalClientInfo);
+				ClientInfo = FMultiUserClientInfo(ClientSession ? ClientSession->GetSessionClientEndpointId() : FGuid(), LocalClientInfo);
 				return true;
 			}
 
@@ -326,7 +334,7 @@ bool UMultiUserClientStatics::GetMultiUserClientInfoByName(const FString& Client
 				{
 					if (SessionClient.ClientInfo.DisplayName == ClientName)
 					{
-						ClientInfo = UE::MultiUserClientLibrary::ConvertClientInfo(SessionClient.ClientEndpointId, SessionClient.ClientInfo);
+						ClientInfo = FMultiUserClientInfo(SessionClient.ClientEndpointId, SessionClient.ClientInfo);
 						return true;
 					}
 				}
@@ -352,7 +360,7 @@ bool UMultiUserClientStatics::GetRemoteMultiUserClientInfos(TArray<FMultiUserCli
 				const TArray<FConcertSessionClientInfo> SessionClients = ClientSession->GetSessionClients();
 				for (const FConcertSessionClientInfo& SessionClient : SessionClients)
 				{
-					ClientInfos.Add(UE::MultiUserClientLibrary::ConvertClientInfo(SessionClient.ClientEndpointId, SessionClient.ClientInfo));
+					ClientInfos.Add(FMultiUserClientInfo(SessionClient.ClientEndpointId, SessionClient.ClientInfo));
 				}
 
 				return ClientInfos.Num() > 0;
diff --git a/Engine/Plugins/Developer/Concert/ConcertApp/MultiUserClient/Source/MultiUserClientLibrary/Private/MultiUserSubsystem.cpp b/Engine/Plugins/Developer/Concert/ConcertApp/MultiUserClient/Source/MultiUserClientLibrary/Private/MultiUserSubsystem.cpp
--- a/Engine/Plugins/Developer/Concert/ConcertApp/MultiUserClient/Source/MultiUserClientLibrary/Private/MultiUserSubsystem.cpp
+++ b/Engine/Plugins/Developer/Concert/ConcertApp/MultiUserClient/Source/MultiUserClientLibrary/Private/MultiUserSubsystem.cpp
@@ -146,7 +146,7 @@ struct FConcertManager
 		
 		SubSystem->OnSessionClientChanged.Broadcast(
 			MultiUserClientLibrary::ConvertClientStatus(Status),
-			MultiUserClientLibrary::ConvertClientInfo(ClientInfo.ClientEndpointId, ClientInfo.ClientInfo)
+			FMultiUserClientInfo(ClientInfo.ClientEndpointId, ClientInfo.ClientInfo)
 			);
 	}
 
diff --git a/Engine/Plugins/Developer/Concert/ConcertApp/MultiUserClient/Source/MultiUserClientLibrary/Public/MultiUserClientStatics.h b/Engine/Plugins/Developer/Concert/ConcertApp/MultiUserClient/Source/MultiUserClientLibrary/Public/MultiUserClientStatics.h
--- a/Engine/Plugins/Developer/Concert/ConcertApp/MultiUserClient/Source/MultiUserClientLibrary/Public/MultiUserClientStatics.h
+++ b/Engine/Plugins/Developer/Concert/ConcertApp/MultiUserClient/Source/MultiUserClientLibrary/Public/MultiUserClientStatics.h
@@ -11,6 +11,7 @@ class UConcertClientConfig;
 enum class EConcertClientStatus : uint8;
 enum class EConcertConnectionStatus : uint8;
 struct FConcertServerInfo;
+struct FConcertSessionClientInfo;
 struct FConcertSessionInfo;
 struct FConcertClientInfo;
 struct FConcertConnectionError;
@@ -43,6 +44,12 @@ struct FMultiUserClientInfo
 	/** Holds an array of tags that can be used for grouping and categorizing. */
 	UPROPERTY(VisibleAnywhere, BlueprintReadOnly, AdvancedDisplay, Category = "Client Info")
 	TArray<FName> Tags;
+
+	FMultiUserClientInfo() = default;
+#if WITH_CONCERT
+	explicit FMultiUserClientInfo(const FConcertSessionClientInfo& ClientInfo);
+	explicit FMultiUserClientInfo(const FGuid& ClientId, const FConcertClientInfo& ClientInfo);
+#endif
 };
 
 /**
@@ -178,6 +185,7 @@ enum class EMultiUserClientStatus : uint8
 namespace UE::MultiUserClientLibrary
 {
 #if WITH_CONCERT
+	UE_DEPRECATED(5.5, "Use FMultiUserClientInfo contructor instead.")
 	MULTIUSERCLIENTLIBRARY_API FMultiUserClientInfo ConvertClientInfo(const FGuid& ClientEndpointId, const FConcertClientInfo& ClientInfo);
 	MULTIUSERCLIENTLIBRARY_API FMultiUserConnectionError ConvertConnectionError(FConcertConnectionError Error);
 	MULTIUSERCLIENTLIBRARY_API UConcertClientConfig* ModifyClientConfig(const FMultiUserClientConfig& InClientConfig);
