Introduce UMultiUserReplicationSubsystem::OnOfflineClientContentChanged for detecting when the content of an offline replication client has changed

diff --git a/Engine/Plugins/Developer/Concert/ConcertApp/MultiUserClient/Source/MultiUserClient/Private/Replication/Client/Offline/OfflineClientManager.cpp b/Engine/Plugins/Developer/Concert/ConcertApp/MultiUserClient/Source/MultiUserClient/Private/Replication/Client/Offline/OfflineClientManager.cpp
--- a/Engine/Plugins/Developer/Concert/ConcertApp/MultiUserClient/Source/MultiUserClient/Private/Replication/Client/Offline/OfflineClientManager.cpp
+++ b/Engine/Plugins/Developer/Concert/ConcertApp/MultiUserClient/Source/MultiUserClient/Private/Replication/Client/Offline/OfflineClientManager.cpp
@@ -91,11 +91,14 @@ namespace UE::MultiUserClient::Replication
 			else
 			{
 				bChanged = true;
-				
+
+				const FGuid LastId = EndpointCache.GetLastAssociatedEndpoint(i);
 				const int32 AddedIndex = Clients.Emplace(
-					MakeUnique<FOfflineClient>(Workspace, KnownEndpointInfo, EndpointCache.GetLastAssociatedEndpoint(i))
+					MakeUnique<FOfflineClient>(Workspace, KnownEndpointInfo, LastId)
 					);
-				OnPostClientAddedDelegate.Broadcast(*Clients[AddedIndex].Get());
+				FOfflineClient& Client = *Clients[AddedIndex].Get();
+				Client.OnStreamPredictionChanged().AddLambda([this, &Client]{ OnClientContentChangedDelegate.Broadcast(Client); });
+				OnPostClientAddedDelegate.Broadcast(Client);
 			}
 		}
 		return bChanged;
diff --git a/Engine/Plugins/Developer/Concert/ConcertApp/MultiUserClient/Source/MultiUserClient/Private/Replication/Client/Offline/OfflineClientManager.h b/Engine/Plugins/Developer/Concert/ConcertApp/MultiUserClient/Source/MultiUserClient/Private/Replication/Client/Offline/OfflineClientManager.h
--- a/Engine/Plugins/Developer/Concert/ConcertApp/MultiUserClient/Source/MultiUserClient/Private/Replication/Client/Offline/OfflineClientManager.h
+++ b/Engine/Plugins/Developer/Concert/ConcertApp/MultiUserClient/Source/MultiUserClient/Private/Replication/Client/Offline/OfflineClientManager.h
@@ -59,6 +59,8 @@ namespace UE::MultiUserClient::Replication
 		FOfflineClientDelegate& OnPostClientAdded() { return OnPostClientAddedDelegate; }
 		/** Called just before an offline client is about to be removed. Called before OnClientsChanged. */
 		FOfflineClientDelegate& OnPreClientRemoved() { return OnPreClientRemovedDelegate; }
+		/** Broadcasts after an offline client's stream content has changed. Also called after OnClientsChangedDelegate. */
+		FOfflineClientDelegate& OnClientContentChanged() { return OnClientContentChangedDelegate; }
 
 	private:
 
@@ -75,10 +77,12 @@ namespace UE::MultiUserClient::Replication
 
 		/** Broadcasts when Clients is changed. */
 		FOfflineClientsChanged OnClientsChangedDelegate;
-		/** Called just after an offline client is has been added. Called before OnClientsChanged. */
+		/** Broadcasts just after an offline client is has been added. Called before OnClientsChanged. */
 		FOfflineClientDelegate OnPostClientAddedDelegate;
-		/** Called just before an offline client is about to be removed. Called before OnClientsChanged. */
+		/** Broadcasts just before an offline client is about to be removed. Called before OnClientsChanged. */
 		FOfflineClientDelegate OnPreClientRemovedDelegate;
+		/** Broadcasts after an offline client's stream content has changed. Not called as part of OnPostClientAddedDelegate. */
+		FOfflineClientDelegate OnClientContentChangedDelegate;
 
 		/** Updates the list of known endpoints in response to an activity being added. */
 		void OnActivityAddedOrProduced(const FConcertClientInfo&, const FConcertSyncActivity&, const FStructOnScope&);
diff --git a/Engine/Plugins/Developer/Concert/ConcertApp/MultiUserClient/Source/MultiUserClient/Private/Replication/MultiUserReplicationManager.cpp b/Engine/Plugins/Developer/Concert/ConcertApp/MultiUserClient/Source/MultiUserClient/Private/Replication/MultiUserReplicationManager.cpp
--- a/Engine/Plugins/Developer/Concert/ConcertApp/MultiUserClient/Source/MultiUserClient/Private/Replication/MultiUserReplicationManager.cpp
+++ b/Engine/Plugins/Developer/Concert/ConcertApp/MultiUserClient/Source/MultiUserClient/Private/Replication/MultiUserReplicationManager.cpp
@@ -147,6 +147,7 @@ namespace UE::MultiUserClient::Replication
 
 		FOfflineClientManager& OfflineClientManager = ConnectedState->OfflineClientManager;
 		OfflineClientManager.OnClientsChanged().AddRaw(this, &FMultiUserReplicationManager::OnInternalOfflineClientsChanged);
+		OfflineClientManager.OnClientContentChanged().AddRaw(this, &FMultiUserReplicationManager::OnInternalOfflineClientContentChanged);
 	}
 
 	void FMultiUserReplicationManager::SetupClientDelegates(FOnlineClient& InClient) const
diff --git a/Engine/Plugins/Developer/Concert/ConcertApp/MultiUserClient/Source/MultiUserClient/Private/Replication/MultiUserReplicationManager.h b/Engine/Plugins/Developer/Concert/ConcertApp/MultiUserClient/Source/MultiUserClient/Private/Replication/MultiUserReplicationManager.h
--- a/Engine/Plugins/Developer/Concert/ConcertApp/MultiUserClient/Source/MultiUserClient/Private/Replication/MultiUserReplicationManager.h
+++ b/Engine/Plugins/Developer/Concert/ConcertApp/MultiUserClient/Source/MultiUserClient/Private/Replication/MultiUserReplicationManager.h
@@ -104,6 +104,7 @@ namespace UE::MultiUserClient::Replication
 		virtual FOnServerStateChanged& OnStreamServerStateChanged() override { return OnStreamServerStateChangedDelegate; }
 		virtual FOnServerStateChanged& OnAuthorityServerStateChanged() override { return OnAuthorityServerStateChangedDelegate; }
 		virtual FOnOfflineClientsChanged& OnOfflineClientsChanged() override { return OnOfflineClientsChangedDelegate; }
+		virtual FOnServerStateChanged& OnOfflineClientContentChanged() override { return OnOfflineClientContentChangedDelegate; }
 		//~ End IMultiUserReplication Interface
 
 	private:
@@ -138,7 +139,7 @@ namespace UE::MultiUserClient::Replication
 			FOnlineClientManager OnlineClientManager;
 			/** Keeps track of clients that were once in the session but are no longer to be able to display their old settings in the UI. */
 			FOfflineClientManager OfflineClientManager;
-			/** Adapter abstraction that allows systems to query online and offline clients with an unified interface. */
+			/** Adapter abstraction that allows systems to query online and offline clients with a unified interface. */
 			FUnifiedClientView UnifiedClientView;
 			
 			/** Interacts with the mute global server mute system. */
@@ -175,6 +176,8 @@ namespace UE::MultiUserClient::Replication
 		FOnServerStateChanged OnAuthorityServerStateChangedDelegate;
 		/** Triggers when the offline clients have changed. */
 		FOnOfflineClientsChanged OnOfflineClientsChangedDelegate;
+		/** Triggers when the content of an offline client has changed. Also triggered as part of OnOfflineClientsChangedDelegate. */
+		FOnServerStateChanged OnOfflineClientContentChangedDelegate;
 
 		/** Callback into Concert for when client connection has changed. */
 		void OnSessionConnectionChanged(IConcertClientSession& ConcertClientSession, EConcertConnectionStatus ConcertConnectionStatus);
@@ -191,6 +194,7 @@ namespace UE::MultiUserClient::Replication
 		/** Sets up delegates for implementing the broadcasting of OnStreamServerStateChangedDelegate and OnAuthorityServerStateChangedDelegate. */
 		void SetupClientConnectionEvents();
 		void OnInternalOfflineClientsChanged() const { OnOfflineClientsChangedDelegate.Broadcast(); }
+		void OnInternalOfflineClientContentChanged(FOfflineClient& OfflineClient) const { OnOfflineClientContentChangedDelegate.Broadcast(OfflineClient.GetLastAssociatedEndpoint()); }
 		void OnReplicationClientConnected(FRemoteClient& RemoteClient) const { SetupClientDelegates(RemoteClient); }
 
 		/** Sets up delegates for reacting to content changes. */
diff --git a/Engine/Plugins/Developer/Concert/ConcertApp/MultiUserClient/Source/MultiUserClient/Public/Replication/IMultiUserReplication.h b/Engine/Plugins/Developer/Concert/ConcertApp/MultiUserClient/Source/MultiUserClient/Public/Replication/IMultiUserReplication.h
--- a/Engine/Plugins/Developer/Concert/ConcertApp/MultiUserClient/Source/MultiUserClient/Public/Replication/IMultiUserReplication.h
+++ b/Engine/Plugins/Developer/Concert/ConcertApp/MultiUserClient/Source/MultiUserClient/Public/Replication/IMultiUserReplication.h
@@ -117,6 +117,9 @@ namespace UE::MultiUserClient
 		/** @return Delegate that triggers when the endpoints considered offline have changed. */
 		virtual FOnOfflineClientsChanged& OnOfflineClientsChanged() = 0;
 
+		/** @return Delegates that triggers when the content of an offline client has changed. Not called as part of OnOfflineClientsChanged. */
+		virtual FOnServerStateChanged& OnOfflineClientContentChanged() = 0;
+
 		virtual ~IMultiUserReplication() = default;
 	};
 }
diff --git a/Engine/Plugins/Developer/Concert/ConcertApp/MultiUserClient/Source/MultiUserClientLibrary/Private/Replication/MultiUserReplicationSubsystem.cpp b/Engine/Plugins/Developer/Concert/ConcertApp/MultiUserClient/Source/MultiUserClientLibrary/Private/Replication/MultiUserReplicationSubsystem.cpp
--- a/Engine/Plugins/Developer/Concert/ConcertApp/MultiUserClient/Source/MultiUserClientLibrary/Private/Replication/MultiUserReplicationSubsystem.cpp
+++ b/Engine/Plugins/Developer/Concert/ConcertApp/MultiUserClient/Source/MultiUserClientLibrary/Private/Replication/MultiUserReplicationSubsystem.cpp
@@ -178,9 +178,10 @@ void UMultiUserReplicationSubsystem::Initialize(FSubsystemCollectionBase& Collec
 		UObjectAdapter = MakeShared<UE::MultiUserClientLibrary::FUObjectAdapterReplicationDiscoverer>();
 		ReplicationInterface->RegisterReplicationDiscoverer(UObjectAdapter.ToSharedRef());
 
-		ReplicationInterface->OnStreamServerStateChanged().AddUObject(this, &UMultiUserReplicationSubsystem::OnClientStreamsChanged);
-		ReplicationInterface->OnAuthorityServerStateChanged().AddUObject(this, &UMultiUserReplicationSubsystem::OnClientAuthorityChanged);
-		ReplicationInterface->OnOfflineClientsChanged().AddUObject(this, &UMultiUserReplicationSubsystem::OnInternalOfflineClientsChanged);
+		ReplicationInterface->OnStreamServerStateChanged().AddUObject(this, &UMultiUserReplicationSubsystem::BroadcastStreamsChanged);
+		ReplicationInterface->OnAuthorityServerStateChanged().AddUObject(this, &UMultiUserReplicationSubsystem::BroadcastAuthorityChanged);
+		ReplicationInterface->OnOfflineClientsChanged().AddUObject(this, &UMultiUserReplicationSubsystem::BroadcastOfflineClientsChanged);
+		ReplicationInterface->OnOfflineClientContentChanged().AddUObject(this, &UMultiUserReplicationSubsystem::BroadcastOfflineClientContentChanged);
 	}
 #endif
 }
diff --git a/Engine/Plugins/Developer/Concert/ConcertApp/MultiUserClient/Source/MultiUserClientLibrary/Public/Replication/MultiUserReplicationSubsystem.h b/Engine/Plugins/Developer/Concert/ConcertApp/MultiUserClient/Source/MultiUserClientLibrary/Public/Replication/MultiUserReplicationSubsystem.h
--- a/Engine/Plugins/Developer/Concert/ConcertApp/MultiUserClient/Source/MultiUserClientLibrary/Public/Replication/MultiUserReplicationSubsystem.h
+++ b/Engine/Plugins/Developer/Concert/ConcertApp/MultiUserClient/Source/MultiUserClientLibrary/Public/Replication/MultiUserReplicationSubsystem.h
@@ -100,7 +100,6 @@ public:
 	/** @return The list of offline clients. Each entry is an endpoint ID that a user had in the past. */
 	UFUNCTION(BlueprintPure, Category = "Multi-user", meta = (Keywords = "Get Offline Disconnected Clients"))
 	TArray<FGuid> GetOfflineClientIds() const;
-	
 
 	//~ Begin USubsystem Interface
 	virtual void Initialize(FSubsystemCollectionBase& Collection) override;
@@ -125,6 +124,10 @@ private:
 	UPROPERTY(BlueprintAssignable, Category = "Multi-user")
 	FOnOfflineClientsChanged OnOfflineClientsChanged;
 
+	/** Event triggered when the content that an offline client will attempt to re-claim changes. */
+	UPROPERTY(BlueprintAssignable, Category = "Multi-user")
+	FOnServerStateChanged OnOfflineClientContentChanged;
+
 	/**
 	 * This is used only for when the adds an object through the Add button in the UI.
 	 * 
@@ -135,7 +138,8 @@ private:
 	 */
 	TSharedPtr<UE::MultiUserClientLibrary::FUObjectAdapterReplicationDiscoverer> UObjectAdapter;
 
-	void OnClientStreamsChanged(const FGuid& EndpointId) const { OnClientStreamServerStateChanged.Broadcast(EndpointId); }
-	void OnClientAuthorityChanged(const FGuid& EndpointId) const { OnClientAuthorityServerStateChanged.Broadcast(EndpointId); }
-	void OnInternalOfflineClientsChanged() const { OnOfflineClientsChanged.Broadcast(); }
+	void BroadcastStreamsChanged(const FGuid& EndpointId) const { OnClientStreamServerStateChanged.Broadcast(EndpointId); }
+	void BroadcastAuthorityChanged(const FGuid& EndpointId) const { OnClientAuthorityServerStateChanged.Broadcast(EndpointId); }
+	void BroadcastOfflineClientsChanged() const { OnOfflineClientsChanged.Broadcast(); }
+	void BroadcastOfflineClientContentChanged(const FGuid& EndpointId) const { OnOfflineClientContentChanged.Broadcast(EndpointId); }
 };
