Only trigger FRejoinStreamAndAuthorityPredictor::OnPredictionChangedDelegate when the content has actually changed.

diff --git a/Engine/Plugins/Developer/Concert/ConcertApp/MultiUserClient/Source/MultiUserClient/Private/Replication/Client/Offline/RejoinStreamAndAuthorityPredictor.cpp b/Engine/Plugins/Developer/Concert/ConcertApp/MultiUserClient/Source/MultiUserClient/Private/Replication/Client/Offline/RejoinStreamAndAuthorityPredictor.cpp
--- a/Engine/Plugins/Developer/Concert/ConcertApp/MultiUserClient/Source/MultiUserClient/Private/Replication/Client/Offline/RejoinStreamAndAuthorityPredictor.cpp
+++ b/Engine/Plugins/Developer/Concert/ConcertApp/MultiUserClient/Source/MultiUserClient/Private/Replication/Client/Offline/RejoinStreamAndAuthorityPredictor.cpp
@@ -87,11 +87,16 @@ namespace UE::MultiUserClient::Replication
 			AnalyzeHistory();
 		}
 	}
-
 	
 	void FRejoinStreamAndAuthorityPredictor::AnalyzeHistory()
 	{
-		Private::AnalyzeActivityHistory(Workspace, ClientInfo, PredictedStream, PredictedAuthority);
-		OnPredictionChangedDelegate.Broadcast();
+		FConcertBaseStreamInfo Prediction;
+		Private::AnalyzeActivityHistory(Workspace, ClientInfo, Prediction, PredictedAuthority);
+
+		if (Prediction != PredictedStream)
+		{
+			PredictedStream = MoveTemp(Prediction);
+			OnPredictionChangedDelegate.Broadcast();
+		}
 	}
 }
