Check for editor before hooking into the column visibility on the level browser. Also check for level editor before unregistering as we may be exiting.

diff --git a/Engine/Plugins/Developer/Concert/ConcertApp/MultiUserClient/Source/MultiUserClient/Private/Misc/GameVisibilityColumnToggler.cpp b/Engine/Plugins/Developer/Concert/ConcertApp/MultiUserClient/Source/MultiUserClient/Private/Misc/GameVisibilityColumnToggler.cpp
--- a/Engine/Plugins/Developer/Concert/ConcertApp/MultiUserClient/Source/MultiUserClient/Private/Misc/GameVisibilityColumnToggler.cpp
+++ b/Engine/Plugins/Developer/Concert/ConcertApp/MultiUserClient/Source/MultiUserClient/Private/Misc/GameVisibilityColumnToggler.cpp
@@ -41,12 +41,15 @@ namespace UE::MultiUserClient
 	{
 		FLevelEditorModule& LevelEditorModule = FModuleManager::GetModuleChecked<FLevelEditorModule>("LevelEditor");
 		const TSharedPtr<ILevelEditor> LevelEditor = LevelEditorModule.GetFirstLevelEditor();
-		if (const TSharedPtr<WorldHierarchy::IWorldHierarchy> LevelsTab = LevelEditor->GetWorldHierarchy().Pin();
-			LevelsTab && bHideVisibilityColumnOnSessionLeave)
+		if (LevelEditor.IsValid())
 		{
-			LevelsTab->SetColumnVisible(WorldHierarchy::HierarchyColumns::ColumnID_GameVisibility, false);
+			if (const TSharedPtr<WorldHierarchy::IWorldHierarchy> LevelsTab = LevelEditor->GetWorldHierarchy().Pin();
+				LevelsTab && bHideVisibilityColumnOnSessionLeave)
+			{
+				LevelsTab->SetColumnVisible(WorldHierarchy::HierarchyColumns::ColumnID_GameVisibility, false);
+			}
 		}
 	}
 }
 
-#endif
\ No newline at end of file
+#endif
diff --git a/Engine/Plugins/Developer/Concert/ConcertApp/MultiUserClient/Source/MultiUserClient/Private/MultiUserClientModule.cpp b/Engine/Plugins/Developer/Concert/ConcertApp/MultiUserClient/Source/MultiUserClient/Private/MultiUserClientModule.cpp
--- a/Engine/Plugins/Developer/Concert/ConcertApp/MultiUserClient/Source/MultiUserClient/Private/MultiUserClientModule.cpp
+++ b/Engine/Plugins/Developer/Concert/ConcertApp/MultiUserClient/Source/MultiUserClient/Private/MultiUserClientModule.cpp
@@ -631,7 +631,10 @@ public:
 		RegisterUI();
 		RegisterLogging();
 #if WITH_EDITOR
-		VisibilityColumnToggler = MakeUnique<UE::MultiUserClient::FGameVisibilityColumnToggler>(MultiUserClient.ToSharedRef());
+		if (GIsEditor)
+		{
+			VisibilityColumnToggler = MakeUnique<UE::MultiUserClient::FGameVisibilityColumnToggler>(MultiUserClient.ToSharedRef());
+		}
 #endif
 	}
 
