Add analytics to replication usage.

diff --git a/Engine/Plugins/Developer/Concert/ConcertApp/MultiUserClient/Source/MultiUserClient/Private/Replication/Misc/AnalyticsHandler.cpp b/Engine/Plugins/Developer/Concert/ConcertApp/MultiUserClient/Source/MultiUserClient/Private/Replication/Misc/AnalyticsHandler.cpp
new file mode 100644
--- /dev/null
+++ b/Engine/Plugins/Developer/Concert/ConcertApp/MultiUserClient/Source/MultiUserClient/Private/Replication/Misc/AnalyticsHandler.cpp
@@ -0,0 +1,54 @@
+// Copyright Epic Games, Inc. All Rights Reserved.
+#include "AnalyticsHandler.h"
+
+#include "EngineAnalytics.h"
+#include "IConcertClient.h"
+#include "Misc/EBreakBehavior.h"
+#include "Replication/Client/Online/OnlineClientManager.h"
+
+namespace UE::MultiUserClient::Replication
+{
+	FAnalyticsHandler::FAnalyticsHandler( IConcertClient& InClient, FOnlineClientManager& InOnlineClientManager )
+		:
+		  Client(InClient)
+		, OnlineClientManager(InOnlineClientManager)
+	{
+		OnlineClientManager.OnPostRemoteClientAdded().AddRaw(this, &FAnalyticsHandler::OnClientAdded);
+	}
+	FAnalyticsHandler::~FAnalyticsHandler()
+	{
+		OnlineClientManager.OnPostRemoteClientAdded().RemoveAll(this);
+		OnlineClientManager.ForEachClient(
+			[this](FOnlineClient& OnlineClient)
+			{
+				OnlineClient.OnModelChanged().RemoveAll(this);
+				return EBreakBehavior::Continue;
+			});
+	}
+
+	void FAnalyticsHandler::OnClientContentChanged()
+	{
+		if (!bSentReplicationAnalyticsData)
+		{
+			if (!FEngineAnalytics::IsAvailable())
+			{
+				return;
+			}
+
+			TSharedPtr<IConcertClientSession> CurrentSession = Client.GetCurrentSession();
+			check(CurrentSession.IsValid());
+
+			const FConcertSessionInfo& SessionInfo = CurrentSession->GetSessionInfo();
+			TArray<FAnalyticsEventAttribute> EventAttributes;
+			EventAttributes.Add(FAnalyticsEventAttribute(TEXT("SessionID"), SessionInfo.SessionId.ToString()));
+
+			FEngineAnalytics::GetProvider().RecordEvent(TEXT("Usage.MultiUser.ReplicationUsed"), EventAttributes);
+			bSentReplicationAnalyticsData = true;
+		}
+	}
+
+	void FAnalyticsHandler::OnClientAdded(FRemoteClient& InRemoteClient)
+	{
+		InRemoteClient.OnModelChanged().AddRaw(this, &FAnalyticsHandler::OnClientContentChanged);
+	}
+}
diff --git a/Engine/Plugins/Developer/Concert/ConcertApp/MultiUserClient/Source/MultiUserClient/Private/Replication/Misc/AnalyticsHandler.h b/Engine/Plugins/Developer/Concert/ConcertApp/MultiUserClient/Source/MultiUserClient/Private/Replication/Misc/AnalyticsHandler.h
new file mode 100644
--- /dev/null
+++ b/Engine/Plugins/Developer/Concert/ConcertApp/MultiUserClient/Source/MultiUserClient/Private/Replication/Misc/AnalyticsHandler.h
@@ -0,0 +1,41 @@
+// Copyright Epic Games, Inc. All Rights Reserved.
+
+#pragma once
+
+#include "HAL/Platform.h"
+#include "Templates/UnrealTemplate.h"
+
+class IConcertClient;
+
+namespace UE::MultiUserClient::Replication
+{
+	class FOnlineClientManager;
+	class FRemoteClient;
+
+	/**
+	 * Handles any changes to the replication data model and send an analytics event with session id.
+	 */
+	class FAnalyticsHandler : public FNoncopyable
+	{
+	public:
+		FAnalyticsHandler(IConcertClient& InClient UE_LIFETIMEBOUND,
+						  FOnlineClientManager& InReplicationClientManager UE_LIFETIMEBOUND);
+		~FAnalyticsHandler();
+
+	private:
+		/** Used to get client session information. */
+		IConcertClient& Client;
+
+		/** Used to detect when a client's state changes. */
+		FOnlineClientManager& OnlineClientManager;
+
+		/** Called when a remote client joins. */
+		void OnClientAdded(FRemoteClient& InRemoteClient);
+
+		/** Called when remote client changes the replication model. */
+		void OnClientContentChanged();
+
+		/** Have we sent our analytics event for replication. */
+		bool bSentReplicationAnalyticsData = false;
+	};
+}
diff --git a/Engine/Plugins/Developer/Concert/ConcertApp/MultiUserClient/Source/MultiUserClient/Private/Replication/MultiUserReplicationManager.cpp b/Engine/Plugins/Developer/Concert/ConcertApp/MultiUserClient/Source/MultiUserClient/Private/Replication/MultiUserReplicationManager.cpp
--- a/Engine/Plugins/Developer/Concert/ConcertApp/MultiUserClient/Source/MultiUserClient/Private/Replication/MultiUserReplicationManager.cpp
+++ b/Engine/Plugins/Developer/Concert/ConcertApp/MultiUserClient/Source/MultiUserClient/Private/Replication/MultiUserReplicationManager.cpp
@@ -4,16 +4,14 @@
 
 #include "ConcertLogGlobal.h"
 #include "IConcertSyncClient.h"
+#include "Misc/AnalyticsHandler.h"
 #include "Replication/ChangeOperationTypes.h"
 #include "Replication/IConcertClientReplicationManager.h"
 #include "Replication/IOfflineReplicationClient.h"
 #include "Replication/Misc/StreamAndAuthorityPredictionUtils.h"
 
 #include "Containers/Ticker.h"
-#include "Framework/Application/SlateApplication.h"
-#include "Framework/Notifications/NotificationManager.h"
 #include "UObject/Package.h"
-#include "Widgets/Notifications/SNotificationList.h"
 
 #define LOCTEXT_NAMESPACE "FMultiUserReplicationManager"
 
@@ -281,10 +279,17 @@ namespace UE::MultiUserClient::Replication
 		return false;
 	}
 
-	FMultiUserReplicationManager::FConnectedState::FConnectedState(TSharedRef<IConcertSyncClient> InClient, FReplicationDiscoveryContainer& InDiscoveryContainer)
+	FMultiUserReplicationManager::FConnectedState::FConnectedState(
+		TSharedRef<IConcertSyncClient> InClient, FReplicationDiscoveryContainer& InDiscoveryContainer
+	)
 		: Client(InClient)
 		, QueryService(*InClient)
-		, OnlineClientManager(InClient, InClient->GetConcertClient()->GetCurrentSession().ToSharedRef(), InDiscoveryContainer, QueryService.GetStreamAndAuthorityQueryService())
+		, OnlineClientManager(
+			  InClient,
+			  InClient->GetConcertClient()->GetCurrentSession().ToSharedRef(),
+			  InDiscoveryContainer,
+			  QueryService.GetStreamAndAuthorityQueryService()
+		  )
 		, OfflineClientManager(*InClient, OnlineClientManager)
 		, UnifiedClientView(*InClient, OnlineClientManager, OfflineClientManager)
 		, MuteManager(*InClient, QueryService.GetMuteStateQueryService(), OnlineClientManager.GetAuthorityCache())
@@ -293,7 +298,8 @@ namespace UE::MultiUserClient::Replication
 		, ChangeLevelHandler(OnlineClientManager.GetLocalClient().GetClientEditModel().Get())
 		, PreventReplicatedPropertyTransaction(*InClient, OnlineClientManager, MuteManager)
 		, UserNotifier(*InClient->GetConcertClient(), OnlineClientManager, MuteManager)
+		, AnalyticsHandler(*InClient->GetConcertClient(), OnlineClientManager)
 	{}
 }
 
-#undef LOCTEXT_NAMESPACE
\ No newline at end of file
+#undef LOCTEXT_NAMESPACE
diff --git a/Engine/Plugins/Developer/Concert/ConcertApp/MultiUserClient/Source/MultiUserClient/Private/Replication/MultiUserReplicationManager.h b/Engine/Plugins/Developer/Concert/ConcertApp/MultiUserClient/Source/MultiUserClient/Private/Replication/MultiUserReplicationManager.h
--- a/Engine/Plugins/Developer/Concert/ConcertApp/MultiUserClient/Source/MultiUserClient/Private/Replication/MultiUserReplicationManager.h
+++ b/Engine/Plugins/Developer/Concert/ConcertApp/MultiUserClient/Source/MultiUserClient/Private/Replication/MultiUserReplicationManager.h
@@ -6,6 +6,7 @@
 #include "IConcertSession.h"
 #include "Client/UnifiedClientView.h"
 #include "Client/Offline/OfflineClientManager.h"
+#include "Misc/AnalyticsHandler.h"
 #include "Misc/ChangeLevelHandler.h"
 #include "Misc/Notification/ReplicationUserNotifier.h"
 #include "Misc/PreventReplicatedPropertyTransaction.h"
@@ -158,7 +159,9 @@ namespace UE::MultiUserClient::Replication
 			FPreventReplicatedPropertyTransaction PreventReplicatedPropertyTransaction;
 			/** This system notifies users when requests go wrong */
 			FReplicationUserNotifier UserNotifier;
-			
+			/** Sends analytics data about replication use. */
+			FAnalyticsHandler AnalyticsHandler;
+
 			FConnectedState(TSharedRef<IConcertSyncClient> InClient, FReplicationDiscoveryContainer& InDiscoveryContainer);
 		};
 		/** Set when connected to a replication session. */
