The user now automatically receives ownership over replicated properties that they add to their selection.

diff --git a/Engine/Plugins/Developer/Concert/ConcertApp/MultiUserClient/Source/MultiUserClient/Private/Replication/Misc/PropertySelection/AutoPropertyOwnershipTaker.cpp b/Engine/Plugins/Developer/Concert/ConcertApp/MultiUserClient/Source/MultiUserClient/Private/Replication/Misc/PropertySelection/AutoPropertyOwnershipTaker.cpp
new file mode 100644
--- /dev/null
+++ b/Engine/Plugins/Developer/Concert/ConcertApp/MultiUserClient/Source/MultiUserClient/Private/Replication/Misc/PropertySelection/AutoPropertyOwnershipTaker.cpp
@@ -0,0 +1,51 @@
+﻿// Copyright Epic Games, Inc. All Rights Reserved.
+
+#include "AutoPropertyOwnershipTaker.h"
+
+#include "UserPropertySelector.h"
+#include "Replication/Client/Online/LocalClient.h"
+#include "Replication/Misc/GlobalAuthorityCache.h"
+
+namespace UE::MultiUserClient::Replication
+{
+	FAutoPropertyOwnershipTaker::FAutoPropertyOwnershipTaker(
+		FUserPropertySelector& InPropertySelector,
+		FLocalClient& InLocalClient,
+		FGlobalAuthorityCache& InReplicationCache
+		)
+		: PropertySelector(InPropertySelector)
+		, LocalClient(InLocalClient)
+		, ReplicationCache(InReplicationCache)
+	{
+		PropertySelector.OnPropertiesAddedByUser().AddRaw(this, &FAutoPropertyOwnershipTaker::OnPropertiesAddedByUser);
+		PropertySelector.OnPropertiesRemovedByUser().AddRaw(this, &FAutoPropertyOwnershipTaker::OnPropertiesRemovedByUser);
+	}
+
+	FAutoPropertyOwnershipTaker::~FAutoPropertyOwnershipTaker()
+	{
+		PropertySelector.OnPropertiesAddedByUser().RemoveAll(this);
+		PropertySelector.OnPropertiesRemovedByUser().RemoveAll(this);
+	}
+
+	void FAutoPropertyOwnershipTaker::OnPropertiesAddedByUser(UObject* Object, TArrayView<const FConcertPropertyChain> Properties) const
+	{
+		TArray<FConcertPropertyChain> AddedProperties;
+		Algo::TransformIf(Properties, AddedProperties,
+			[this, Object](const FConcertPropertyChain& Property){ return !ReplicationCache.IsPropertyReferencedByAnyClientStream(Object, Property); },
+			[](const FConcertPropertyChain& Property){ return Property; }
+			);
+		
+		if (!AddedProperties.IsEmpty())
+		{
+			const TSharedRef<ConcertSharedSlate::IEditableReplicationStreamModel>& EditModel = LocalClient.GetClientEditModel();
+			EditModel->AddObjects({ Object });
+			EditModel->AddProperties(Object, AddedProperties);
+		}
+	}
+
+	void FAutoPropertyOwnershipTaker::OnPropertiesRemovedByUser(UObject* Object, TArrayView<const FConcertPropertyChain> Properties) const
+	{
+		const TSharedRef<ConcertSharedSlate::IEditableReplicationStreamModel>& EditModel = LocalClient.GetClientEditModel();
+		EditModel->RemoveProperties(Object, Properties);
+	}
+}
diff --git a/Engine/Plugins/Developer/Concert/ConcertApp/MultiUserClient/Source/MultiUserClient/Private/Replication/Misc/PropertySelection/AutoPropertyOwnershipTaker.h b/Engine/Plugins/Developer/Concert/ConcertApp/MultiUserClient/Source/MultiUserClient/Private/Replication/Misc/PropertySelection/AutoPropertyOwnershipTaker.h
new file mode 100644
--- /dev/null
+++ b/Engine/Plugins/Developer/Concert/ConcertApp/MultiUserClient/Source/MultiUserClient/Private/Replication/Misc/PropertySelection/AutoPropertyOwnershipTaker.h
@@ -0,0 +1,43 @@
+﻿// Copyright Epic Games, Inc. All Rights Reserved.
+
+#pragma once
+
+#include "Containers/ContainersFwd.h"
+#include "HAL/Platform.h"
+#include "Templates/UnrealTemplate.h"
+
+class UObject;
+struct FConcertPropertyChain;
+
+namespace UE::MultiUserClient::Replication
+{
+	class FGlobalAuthorityCache;
+	class FLocalClient;
+	class FUserPropertySelector;
+
+	/** Automatically takes and releases ownership over properties the user adds to the set of selected properties. */
+	class FAutoPropertyOwnershipTaker : public FNoncopyable
+	{
+	public:
+		
+		FAutoPropertyOwnershipTaker(
+			FUserPropertySelector& InPropertySelector UE_LIFETIMEBOUND,
+			FLocalClient& InLocalClient UE_LIFETIMEBOUND,
+			FGlobalAuthorityCache& InReplicationCache UE_LIFETIMEBOUND
+			);
+		~FAutoPropertyOwnershipTaker();
+
+	private:
+
+		/** Used to detect when a property is added or removed from the selected set of properties. */
+		FUserPropertySelector& PropertySelector;
+		/** The property ownership is changed for this client. */
+		FLocalClient& LocalClient;
+		/** Used to quickly check whether user selected properties are already owned. */
+		FGlobalAuthorityCache& ReplicationCache;
+		
+		void OnPropertiesAddedByUser(UObject* Object, TArrayView<const FConcertPropertyChain> Properties) const;
+		void OnPropertiesRemovedByUser(UObject* Object, TArrayView<const FConcertPropertyChain> ConcertPropertyChains) const;
+	};
+}
+
diff --git a/Engine/Plugins/Developer/Concert/ConcertApp/MultiUserClient/Source/MultiUserClient/Private/Replication/Misc/PropertySelection/UserPropertySelector.cpp b/Engine/Plugins/Developer/Concert/ConcertApp/MultiUserClient/Source/MultiUserClient/Private/Replication/Misc/PropertySelection/UserPropertySelector.cpp
--- a/Engine/Plugins/Developer/Concert/ConcertApp/MultiUserClient/Source/MultiUserClient/Private/Replication/Misc/PropertySelection/UserPropertySelector.cpp
+++ b/Engine/Plugins/Developer/Concert/ConcertApp/MultiUserClient/Source/MultiUserClient/Private/Replication/Misc/PropertySelection/UserPropertySelector.cpp
@@ -62,20 +62,22 @@ namespace UE::MultiUserClient::Replication
 		FCoreUObjectDelegates::OnObjectTransacted.RemoveAll(this);
 	}
 
-	void FUserPropertySelector::AddSelectedProperties(UObject* Object, TConstArrayView<FConcertPropertyChain> Properties)
+	void FUserPropertySelector::AddUserSelectedProperties(UObject* Object, TConstArrayView<FConcertPropertyChain> Properties)
 	{
 		const FScopedTransaction Transaction(LOCTEXT("AddSelectedProperties", "Select replicated property"));
 		PropertySelection->Modify();
 		
 		InternalAddSelectedProperties(Object, Properties);
+		OnPropertiesAddedByUserDelegate.Broadcast(Object, Properties);
 	}
 	
-	void FUserPropertySelector::RemoveSelectedProperties(UObject* Object, TConstArrayView<FConcertPropertyChain> Properties)
+	void FUserPropertySelector::RemoveUserSelectedProperties(UObject* Object, TConstArrayView<FConcertPropertyChain> Properties)
 	{
 		const FScopedTransaction Transaction(LOCTEXT("RemoveSelectedProperties", "Deselect replicated property"));
 		PropertySelection->Modify();
 		
 		InternalRemoveSelectedProperties(Object, Properties);
+		OnPropertiesRemovedByUserDelegate.Broadcast(Object, Properties);
 	}
 
 	bool FUserPropertySelector::IsPropertySelected(const FSoftObjectPath& Object, const FConcertPropertyChain& Property) const
diff --git a/Engine/Plugins/Developer/Concert/ConcertApp/MultiUserClient/Source/MultiUserClient/Private/Replication/Misc/PropertySelection/UserPropertySelector.h b/Engine/Plugins/Developer/Concert/ConcertApp/MultiUserClient/Source/MultiUserClient/Private/Replication/Misc/PropertySelection/UserPropertySelector.h
--- a/Engine/Plugins/Developer/Concert/ConcertApp/MultiUserClient/Source/MultiUserClient/Private/Replication/Misc/PropertySelection/UserPropertySelector.h
+++ b/Engine/Plugins/Developer/Concert/ConcertApp/MultiUserClient/Source/MultiUserClient/Private/Replication/Misc/PropertySelection/UserPropertySelector.h
@@ -49,10 +49,10 @@ namespace UE::MultiUserClient::Replication
 			);
 		virtual ~FUserPropertySelector() override;
 
-		/** Add Properties from the user's selection for Object. */
-		void AddSelectedProperties(UObject* Object, TConstArrayView<FConcertPropertyChain> Properties);
-		/** Removes Properties from the user's selection for Object. */
-		void RemoveSelectedProperties(UObject* Object, TConstArrayView<FConcertPropertyChain> Properties);
+		/** Add Properties to the user's selection for Object. These properties were purposefully selected added by the user. */
+		void AddUserSelectedProperties(UObject* Object, TConstArrayView<FConcertPropertyChain> Properties);
+		/** Removes Properties from the user's selection for Object. These properties were purposefully selected added by the user. */
+		void RemoveUserSelectedProperties(UObject* Object, TConstArrayView<FConcertPropertyChain> Properties);
 		/** @return Whether Property is selected for Object. */
 		bool IsPropertySelected(const FSoftObjectPath& Object, const FConcertPropertyChain& Property) const;
 
@@ -61,6 +61,12 @@ namespace UE::MultiUserClient::Replication
 		DECLARE_MULTICAST_DELEGATE(FOnPropertySelectionChanged)
 		/** @return Event that broadcasts when the user property selection changes. */
 		FOnPropertySelectionChanged& OnPropertySelectionChanged() { return OnPropertySelectionChangedDelegate; }
+
+		DECLARE_MULTICAST_DELEGATE_TwoParams(FOnPropertiesChangedByUser, UObject* /*Object*/, TConstArrayView<FConcertPropertyChain> /*Properties*/);
+		/** @return Event that broadcasts when the user adds a property manually (through UI). */
+		FOnPropertiesChangedByUser& OnPropertiesAddedByUser() { return OnPropertiesAddedByUserDelegate; }
+		/** @return Event that broadcasts when the user removes a property manually (through UI). */
+		FOnPropertiesChangedByUser& OnPropertiesRemovedByUser() { return OnPropertiesRemovedByUserDelegate; }
 		
 		//~ Begin FGCObject Interface
 		virtual void AddReferencedObjects(FReferenceCollector& Collector) override;
@@ -87,6 +93,10 @@ namespace UE::MultiUserClient::Replication
 
 		/** Broadcasts when the user property selection changes. */
 		FOnPropertySelectionChanged OnPropertySelectionChangedDelegate;
+		/** Broadcasts when the user adds a property manually (through UI). */
+		FOnPropertiesChangedByUser OnPropertiesAddedByUserDelegate;
+		/** Broadcasts when the user removes a property manually (through UI). */
+		FOnPropertiesChangedByUser OnPropertiesRemovedByUserDelegate;
 
 		/** Called when a remote client joins. */
 		void OnClientAdded(FRemoteClient& Client);
diff --git a/Engine/Plugins/Developer/Concert/ConcertApp/MultiUserClient/Source/MultiUserClient/Private/Replication/MultiUserReplicationManager.cpp b/Engine/Plugins/Developer/Concert/ConcertApp/MultiUserClient/Source/MultiUserClient/Private/Replication/MultiUserReplicationManager.cpp
--- a/Engine/Plugins/Developer/Concert/ConcertApp/MultiUserClient/Source/MultiUserClient/Private/Replication/MultiUserReplicationManager.cpp
+++ b/Engine/Plugins/Developer/Concert/ConcertApp/MultiUserClient/Source/MultiUserClient/Private/Replication/MultiUserReplicationManager.cpp
@@ -295,6 +295,7 @@ namespace UE::MultiUserClient::Replication
 		, MuteManager(*InClient, QueryService.GetMuteStateQueryService(), OnlineClientManager.GetAuthorityCache())
 		, PresetManager(*InClient, OnlineClientManager, MuteManager.GetSynchronizer())
 		, PropertySelector(OnlineClientManager, OfflineClientManager)
+		, AutoPropertyOwnershipTaker(PropertySelector, OnlineClientManager.GetLocalClient(), OnlineClientManager.GetAuthorityCache())
 		, ChangeLevelHandler(OnlineClientManager.GetLocalClient().GetClientEditModel().Get())
 		, PreventReplicatedPropertyTransaction(*InClient, OnlineClientManager, MuteManager)
 		, UserNotifier(*InClient->GetConcertClient(), OnlineClientManager, MuteManager)
diff --git a/Engine/Plugins/Developer/Concert/ConcertApp/MultiUserClient/Source/MultiUserClient/Private/Replication/MultiUserReplicationManager.h b/Engine/Plugins/Developer/Concert/ConcertApp/MultiUserClient/Source/MultiUserClient/Private/Replication/MultiUserReplicationManager.h
--- a/Engine/Plugins/Developer/Concert/ConcertApp/MultiUserClient/Source/MultiUserClient/Private/Replication/MultiUserReplicationManager.h
+++ b/Engine/Plugins/Developer/Concert/ConcertApp/MultiUserClient/Source/MultiUserClient/Private/Replication/MultiUserReplicationManager.h
@@ -10,6 +10,7 @@
 #include "Misc/ChangeLevelHandler.h"
 #include "Misc/Notification/ReplicationUserNotifier.h"
 #include "Misc/PreventReplicatedPropertyTransaction.h"
+#include "Misc/PropertySelection/AutoPropertyOwnershipTaker.h"
 #include "Misc/PropertySelection/UserPropertySelector.h"
 #include "Misc/Query/RegularQueryService.h"
 #include "Muting/MuteStateManager.h"
@@ -153,6 +154,9 @@ namespace UE::MultiUserClient::Replication
 			 * The bottom-half property section in the replication UI uses this to keep track of which properties the user has selected for which properties.
 			 */
 			FUserPropertySelector PropertySelector;
+			/** Automatically takes and releases ownership over properties the user adds to the set of selected properties. */
+			FAutoPropertyOwnershipTaker AutoPropertyOwnershipTaker;
+			
 			/** Clears local client's registered objects when leaving map. */
 			FChangeLevelHandler ChangeLevelHandler;
 			/** Prevents recording of transactions that change properties that are being replicated by a client. */
diff --git a/Engine/Plugins/Developer/Concert/ConcertApp/MultiUserClient/Source/MultiUserClient/Private/Widgets/ActiveSession/Replication/Client/PropertySelection/SPropertySelectionComboButton.cpp b/Engine/Plugins/Developer/Concert/ConcertApp/MultiUserClient/Source/MultiUserClient/Private/Widgets/ActiveSession/Replication/Client/PropertySelection/SPropertySelectionComboButton.cpp
--- a/Engine/Plugins/Developer/Concert/ConcertApp/MultiUserClient/Source/MultiUserClient/Private/Widgets/ActiveSession/Replication/Client/PropertySelection/SPropertySelectionComboButton.cpp
+++ b/Engine/Plugins/Developer/Concert/ConcertApp/MultiUserClient/Source/MultiUserClient/Private/Widgets/ActiveSession/Replication/Client/PropertySelection/SPropertySelectionComboButton.cpp
@@ -32,11 +32,11 @@ namespace UE::MultiUserClient::Replication
 				UObject* Object = SoftObjectPtr.Get();
 				if (bIsSelected)
 				{
-					UserSelection.RemoveSelectedProperties(Object, Property.PropertiesToAdd);
+					UserSelection.RemoveUserSelectedProperties(Object, Property.PropertiesToAdd);
 				}
 				else
 				{
-					UserSelection.AddSelectedProperties(Object, Property.PropertiesToAdd);
+					UserSelection.AddUserSelectedProperties(Object, Property.PropertiesToAdd);
 				}
 			}
 			
