Adjusted MU transaction filter used when replicating properties.
Now, transactions are blocked only if all properties being changed by the transaction are replicated. If at least one property is not replicated, the entire transaction is allowed through.

diff --git a/Engine/Plugins/Developer/Concert/ConcertApp/MultiUserClient/Source/MultiUserClient/Private/Replication/Misc/PreventReplicatedPropertyTransaction.cpp b/Engine/Plugins/Developer/Concert/ConcertApp/MultiUserClient/Source/MultiUserClient/Private/Replication/Misc/PreventReplicatedPropertyTransaction.cpp
--- a/Engine/Plugins/Developer/Concert/ConcertApp/MultiUserClient/Source/MultiUserClient/Private/Replication/Misc/PreventReplicatedPropertyTransaction.cpp
+++ b/Engine/Plugins/Developer/Concert/ConcertApp/MultiUserClient/Source/MultiUserClient/Private/Replication/Misc/PreventReplicatedPropertyTransaction.cpp
@@ -45,7 +45,7 @@ namespace UE::MultiUserClient::Replication
 		
 		// FTransactionObjectEvent::GetChangedProperties() only contains root properties, like RelativeLocation; sub-properties, like RelativeLocation.X, are not listed.
 		// Luckily, replication streams list every parent property (so if RelativeLocation.X is being replicated then RelativeLocation is also in the stream).
-		const bool bIsReplicating = !bIsMuted && Algo::AnyOf(FilterArgs.TransactionEvent.GetChangedProperties(), [this, &FilterArgs](const FName& RootProperty)
+		const bool bContainsNonReplicatedProperty = !bIsMuted && Algo::AnyOf(FilterArgs.TransactionEvent.GetChangedProperties(), [this, &FilterArgs](const FName& RootProperty)
 		{
 			const FGlobalAuthorityCache& AuthorityCache = ClientManager.GetAuthorityCache();
 			const TOptional<FGuid> AuthoringClient = AuthorityCache.GetClientWithAuthorityOverProperty(FilterArgs.ObjectToFilter, { RootProperty });
@@ -61,16 +61,17 @@ namespace UE::MultiUserClient::Replication
 				// Replication is only happening when a client also has sync control.
 				// Since we only query sync control of the local client, we only know for our local client whether it has sync control...
 				const FGuid ClientStreamId = ClientManager.GetLocalClient().GetStreamSynchronizer().GetStreamId();
-				return ReplicationManager->HasSyncControl({ ClientStreamId, FilterArgs.ObjectToFilter});
+				return !ReplicationManager->HasSyncControl({ ClientStreamId, FilterArgs.ObjectToFilter});
 			}
 
 			// ... and for remote clients we'll just assume they're replicating when they have authority.
-			return bHasAuthoringClient;
+			return !bHasAuthoringClient;
 		});
 
-		// If the property is being replicated, it should not be transacted because the two systems may interfere.
-		return bIsReplicating
-			? ETransactionFilterResult::ExcludeObject
-			: ETransactionFilterResult::UseDefault;
+		// If the transaction contains only properties that are being replicated, skip the object as the two systems may interfere.
+		// If the transaction contains a property that is not replicated, the transaction should be allowed.
+		return bContainsNonReplicatedProperty
+			? ETransactionFilterResult::UseDefault
+			: ETransactionFilterResult::ExcludeObject;
 	}
 }
