Fix replication system preventing transactions when transaction contains no properties.
- Reverted back to old logic (since that worked) and made the condition for exclusion stronger (AllOf instead of AnyOf).
- Added check for empty property list (AllOf returns true if empty)
- Also added CVar MultiUser.AllowTransactionsOnReplicatedproperties to disallow this feature altogether

diff --git a/Engine/Plugins/Developer/Concert/ConcertApp/MultiUserClient/Source/MultiUserClient/Private/Replication/Misc/PreventReplicatedPropertyTransaction.cpp b/Engine/Plugins/Developer/Concert/ConcertApp/MultiUserClient/Source/MultiUserClient/Private/Replication/Misc/PreventReplicatedPropertyTransaction.cpp
--- a/Engine/Plugins/Developer/Concert/ConcertApp/MultiUserClient/Source/MultiUserClient/Private/Replication/Misc/PreventReplicatedPropertyTransaction.cpp
+++ b/Engine/Plugins/Developer/Concert/ConcertApp/MultiUserClient/Source/MultiUserClient/Private/Replication/Misc/PreventReplicatedPropertyTransaction.cpp
@@ -5,13 +5,21 @@
 #include "IConcertSyncClient.h"
 #include "Replication/Client/Online/OnlineClientManager.h"
 #include "Replication/Misc/GlobalAuthorityCache.h"
-
-#include "Algo/AnyOf.h"
 #include "Replication/Muting/MuteStateManager.h"
+
+#include "Algo/AllOf.h"
+#include "HAL/IConsoleManager.h"
 #include "UObject/SoftObjectPath.h"
 
 namespace UE::MultiUserClient::Replication
 {
+	TAutoConsoleVariable<bool> CVarAllowTransactionsOnReplicatedProperties(
+		TEXT("MultiUser.AllowTransactionsOnReplicatedproperties"),
+		false,
+		TEXT("Whether to stop disallowing transactions on replicated properties."),
+		ECVF_Default
+		);
+	
 	namespace Private
 	{
 		FORCEINLINE FName GetFilterName() { return TEXT("MultiUserReplicationFilter"); }
@@ -39,17 +47,22 @@ namespace UE::MultiUserClient::Replication
 
 	ETransactionFilterResult FPreventReplicatedPropertyTransaction::FilterTransactionAffectedByReplication(const FConcertTransactionFilterArgs& FilterArgs)
 	{
+		if (CVarAllowTransactionsOnReplicatedProperties.GetValueOnAnyThread())
+		{
+			return ETransactionFilterResult::UseDefault;
+		}
+		
 		// 1. If an object is muted, we can save ourselves the work of analysing client streams.
 		// 2. We do not know whether remote clients have sync control (see below) but we can approximate that they don't have it if the object is muted.
 		const bool bIsMuted = MuteManager.GetSynchronizer().IsMuted(FilterArgs.ObjectToFilter);
-		
+
+		const TArray<FName>& Properties = FilterArgs.TransactionEvent.GetChangedProperties();
 		// FTransactionObjectEvent::GetChangedProperties() only contains root properties, like RelativeLocation; sub-properties, like RelativeLocation.X, are not listed.
 		// Luckily, replication streams list every parent property (so if RelativeLocation.X is being replicated then RelativeLocation is also in the stream).
-		const bool bContainsNonReplicatedProperty = !bIsMuted && Algo::AnyOf(FilterArgs.TransactionEvent.GetChangedProperties(), [this, &FilterArgs](const FName& RootProperty)
+		const bool bOnlyHasReplicatedProperties = !bIsMuted && !Properties.IsEmpty() && Algo::AllOf(Properties, [this, &FilterArgs](const FName& RootProperty)
 		{
 			const FGlobalAuthorityCache& AuthorityCache = ClientManager.GetAuthorityCache();
 			const TOptional<FGuid> AuthoringClient = AuthorityCache.GetClientWithAuthorityOverProperty(FilterArgs.ObjectToFilter, { RootProperty });
-
 			const bool bHasAuthoringClient = AuthoringClient.IsSet();
 			const bool bIsLocalClient = bHasAuthoringClient && *AuthoringClient == ClientManager.GetLocalClient().GetEndpointId();
 			if (bIsLocalClient)
@@ -61,17 +74,16 @@ namespace UE::MultiUserClient::Replication
 				// Replication is only happening when a client also has sync control.
 				// Since we only query sync control of the local client, we only know for our local client whether it has sync control...
 				const FGuid ClientStreamId = ClientManager.GetLocalClient().GetStreamSynchronizer().GetStreamId();
-				return !ReplicationManager->HasSyncControl({ ClientStreamId, FilterArgs.ObjectToFilter});
+				return ReplicationManager->HasSyncControl({ ClientStreamId, FilterArgs.ObjectToFilter});
 			}
-
 			// ... and for remote clients we'll just assume they're replicating when they have authority.
-			return !bHasAuthoringClient;
+			return bHasAuthoringClient;
 		});
-
+		
 		// If the transaction contains only properties that are being replicated, skip the object as the two systems may interfere.
 		// If the transaction contains a property that is not replicated, the transaction should be allowed.
-		return bContainsNonReplicatedProperty
-			? ETransactionFilterResult::UseDefault
-			: ETransactionFilterResult::ExcludeObject;
+		return bOnlyHasReplicatedProperties
+			? ETransactionFilterResult::ExcludeObject
+			: ETransactionFilterResult::UseDefault;
 	}
 }
