Extend Details panel rows with a button to add properties for replication

diff --git a/Engine/Plugins/Developer/Concert/ConcertApp/MultiUserClient/Source/MultiUserClient/Private/MultiUserClientModule.cpp b/Engine/Plugins/Developer/Concert/ConcertApp/MultiUserClient/Source/MultiUserClient/Private/MultiUserClientModule.cpp
--- a/Engine/Plugins/Developer/Concert/ConcertApp/MultiUserClient/Source/MultiUserClient/Private/MultiUserClientModule.cpp
+++ b/Engine/Plugins/Developer/Concert/ConcertApp/MultiUserClient/Source/MultiUserClient/Private/MultiUserClientModule.cpp
@@ -53,6 +53,7 @@
 #include "DetailCategoryBuilder.h"
 #include "DetailWidgetRow.h"
 #include "Misc/GameVisibilityColumnToggler.h"
+#include "Replication/Misc/DetailRowExtensionHandler.h"
 
 
 #include "Widgets/SConcertBrowser.h"
@@ -634,6 +635,9 @@ public:
 		if (GIsEditor)
 		{
 			VisibilityColumnToggler = MakeUnique<UE::MultiUserClient::FGameVisibilityColumnToggler>(MultiUserClient.ToSharedRef());
+			DetailReplicationRowExtensionHandler = MakeUnique<UE::MultiUserClient::Replication::FDetailRowExtensionHandler>(
+				MultiUserClient.ToSharedRef(), ReplicationManager.ToSharedRef()
+				);
 		}
 #endif
 	}
@@ -642,6 +646,7 @@ public:
 	{
 #if WITH_EDITOR
 		VisibilityColumnToggler.Reset();
+		DetailReplicationRowExtensionHandler.Reset();
 #endif
 		
 		// Unhook AppPreExit and call it
@@ -1425,6 +1430,9 @@ private:
 #if WITH_EDITOR
 	/** Toggles the game visibility column in the Levels tab when joining & leaving sessions. */
 	TUniquePtr<UE::MultiUserClient::FGameVisibilityColumnToggler> VisibilityColumnToggler;
+
+	/** Adds UI to Details panel rows for quickly adding properties to replication. */
+	TUniquePtr<UE::MultiUserClient::Replication::FDetailRowExtensionHandler> DetailReplicationRowExtensionHandler;
 #endif
 
 	/** True if the tab spawners have been registered for this module */
diff --git a/Engine/Plugins/Developer/Concert/ConcertApp/MultiUserClient/Source/MultiUserClient/Private/MultiUserReplicationStyle.cpp b/Engine/Plugins/Developer/Concert/ConcertApp/MultiUserClient/Source/MultiUserClient/Private/MultiUserReplicationStyle.cpp
--- a/Engine/Plugins/Developer/Concert/ConcertApp/MultiUserClient/Source/MultiUserClient/Private/MultiUserReplicationStyle.cpp
+++ b/Engine/Plugins/Developer/Concert/ConcertApp/MultiUserClient/Source/MultiUserClient/Private/MultiUserReplicationStyle.cpp
@@ -59,6 +59,10 @@ namespace UE::MultiUserClient
 		StyleSet->SetCoreContentRoot(FPaths::EngineContentDir() / TEXT("Slate"));
 		
 		const FVector2D Icon16x16(16.0f, 16.0f); 
+		const FVector2D Icon24x24(24.0f, 24.0f); 
+		const FVector2D Icon48x48(48.0f, 48.0f);
+
+		const FLinearColor IconColorAndOpacity(FLinearColor(1.f, 1.f, 1.f, 1.f));
 
 		// Column widths - see FConcertFrontendStyle::Initialize() also ("Default Concert Replication Columns Widths")
 		StyleSet->Set("AllClients.Object.MuteToggle", 25.f);
@@ -68,6 +72,12 @@ namespace UE::MultiUserClient
 		// Timing
 		StyleSet->Set("AllClients.Reassignment.DisplayThrobberAfterSeconds", 0.2f);
 
+		// Icons
+		StyleSet->Set("MultiUser.Icons.AddProperty",				new IMAGE_PLUGIN_BRUSH("icon_AddProperty_48x", Icon48x48, IconColorAndOpacity));
+		StyleSet->Set("MultiUser.Icons.AddProperty.Small",		new IMAGE_PLUGIN_BRUSH("icon_AddProperty_48x", Icon24x24, IconColorAndOpacity));
+		StyleSet->Set("MultiUser.Icons.RemoveProperty",			new IMAGE_PLUGIN_BRUSH("icon_RemoveProperty_48x", Icon48x48, IconColorAndOpacity));
+		StyleSet->Set("MultiUser.Icons.RemoveProperty.Small",	new IMAGE_PLUGIN_BRUSH("icon_RemoveProperty_48x", Icon24x24, IconColorAndOpacity));
+
 		// Muting
 		FSlateImageBrush* PlayBrush = new IMAGE_PLUGIN_BRUSH("generic_play_16x", Icon16x16);
 		FSlateImageBrush* PauseBrush = new IMAGE_PLUGIN_BRUSH("generic_pause_16x", Icon16x16);
diff --git a/Engine/Plugins/Developer/Concert/ConcertApp/MultiUserClient/Source/MultiUserClient/Private/Replication/Misc/DetailRowExtensionHandler.cpp b/Engine/Plugins/Developer/Concert/ConcertApp/MultiUserClient/Source/MultiUserClient/Private/Replication/Misc/DetailRowExtensionHandler.cpp
new file mode 100644
--- /dev/null
+++ b/Engine/Plugins/Developer/Concert/ConcertApp/MultiUserClient/Source/MultiUserClient/Private/Replication/Misc/DetailRowExtensionHandler.cpp
@@ -0,0 +1,315 @@
+﻿// Copyright Epic Games, Inc. All Rights Reserved.
+
+#include "DetailRowExtensionHandler.h"
+
+#include "IConcertSyncClient.h"
+#include "MultiUserReplicationStyle.h"
+#include "Replication/MultiUserReplicationManager.h"
+#include "Replication/PropertyChainUtils.h"
+#include "Replication/Client/ClientUtils.h"
+#include "Widgets/ActiveSession/Replication/Client/Multi/Columns/AssignProperty/AssignPropertyModel.h"
+
+#include "PropertyEditorModule.h"
+#include "PropertyHandle.h"
+#include "PropertyPath.h"
+
+#define LOCTEXT_NAMESPACE "FDetailRowExtensionHandler"
+
+namespace UE::MultiUserClient::Replication::DetailRowExtensionHandler
+{
+	static UStruct* GetOwningStruct(const FPropertyInfo& LeafPropertyInfo)
+	{
+		const FProperty* LeafProperty = LeafPropertyInfo.Property.Get();
+		return LeafProperty ? Cast<UStruct>(LeafProperty->GetOwnerUObject()) : nullptr;
+	}
+
+	static int32 FindIndexOfRootPropertyOfLeafmostObject(const FPropertyPath& PropertyPath)
+	{
+		UObject* LeafOwningObject = GetOwningStruct(PropertyPath.GetLeafMostProperty());
+		for (int32 i = 0; i < PropertyPath.GetNumProperties(); ++i)
+		{
+			UObject* CurrentObject = GetOwningStruct(PropertyPath.GetPropertyInfo(i));
+			if (CurrentObject == LeafOwningObject)
+			{
+				return i;
+			}
+		}
+		return INDEX_NONE;
+	}
+
+	struct FPathInfo
+	{
+		TArray<FName> Path;
+		UStruct* LeafClass = nullptr;
+
+		bool IsValid() const { return LeafClass != nullptr; }
+	};
+	
+	/**
+	 * Retrieves the property path leading to the deepest (leaf-most) property, which can be used to construct an FConcertPropertyChain.
+	 * Example input: 
+	 *  [0] { "StaticMeshComponent", Owner = AStaticMeshActor }
+	 *  [1] { "RelativeLocation", Owner = UStaticMeshComponent }
+	 *  [2] { "X", Owner = UStaticMeshComponent }
+	 * 
+	 * Example output: 
+	 *  [0 ]"RelativeLocation",	
+	 *  [1] "X"
+	 * It skipped the StaticMeshComponent property.
+	 */
+	static FPathInfo GetPropertyPath(const TSharedPtr<IPropertyHandle>& PropertyHandle)
+	{
+		FPathInfo Result;
+		
+		const TSharedPtr<FPropertyPath> PropertyPath = PropertyHandle ? PropertyHandle->CreateFPropertyPath() : nullptr;
+		if (!PropertyPath || PropertyPath->GetNumProperties() == 0)
+		{
+			return Result;
+		}
+
+		/**
+		 * Depending on the selection in the details panel, the property path will vary. 
+		 * The path originates from the root object in the details view and includes all intermediate properties leading to the selected property.
+		 * 
+		 * Example 1:
+		 *  If you select a StaticMeshActor in the Outliner, and PropertyHandle refers to "RelativeLocation", the path would be:
+		 *  { 
+		 *    [0] { "StaticMeshComponent", Owner = AStaticMeshActor },
+		 *    [1] { "RelativeLocation", Owner = UStaticMeshComponent }
+		 *  }.
+		 * 
+		 * Example 2:
+		 *  If you select the StaticMeshComponent directly from the component hierarchy, the path would be:
+		 *  { 
+		 *    [0] { "RelativeLocation", Owner = UStaticMeshComponent }
+		 *  }.
+		 * 
+		 * This behavior also applies to nested subobjects.
+		 * Summary: We adjust StartIndex to point to the root property of the deepest UObject in the property path.
+		 */
+		const int32 StartIndex = FindIndexOfRootPropertyOfLeafmostObject(*PropertyPath);
+		if (const bool bIsInvalidIndex = StartIndex < 0 || StartIndex >= PropertyPath->GetNumProperties())
+		{
+			return Result;
+		}
+
+		UStruct* LeafObject = GetOwningStruct(PropertyPath->GetPropertyInfo(StartIndex));
+		checkf(LeafObject, TEXT("Valid StartIndex was returned but the corresponding UObject is invalid!"));
+		Result.LeafClass = LeafObject;
+		
+		for (int32 i = StartIndex; i < PropertyPath->GetNumProperties(); ++i)
+		{
+			if (FProperty* Property = PropertyPath->GetPropertyInfo(i).Property.Get(); ensure(Property))
+			{
+				Result.Path.Add(Property->GetFName());
+			}
+			else
+			{
+				return {};
+			}
+		}
+		
+		return Result;
+	}
+}
+
+namespace UE::MultiUserClient::Replication
+{
+	FDetailRowExtensionHandler::FDetailRowExtensionHandler(
+		const TSharedRef<IConcertSyncClient>& InClient,
+		const TSharedRef<FMultiUserReplicationManager>& InReplicationManager
+		)
+		: WeakClient(InClient)
+		, WeakReplicationManager(InReplicationManager)
+	{
+		FPropertyEditorModule& EditModule = FModuleManager::Get().GetModuleChecked<FPropertyEditorModule>("PropertyEditor");
+		EditModule.GetGlobalRowExtensionDelegate().AddRaw(
+			this, &FDetailRowExtensionHandler::RegisterExtensionHandler
+			);
+	}
+
+	FDetailRowExtensionHandler::~FDetailRowExtensionHandler()
+	{
+		if (FPropertyEditorModule* EditModule = FModuleManager::Get().GetModulePtr<FPropertyEditorModule>("PropertyEditor"))
+		{
+			EditModule->GetGlobalRowExtensionDelegate().RemoveAll(this);
+		}
+	}
+
+	void FDetailRowExtensionHandler::RegisterExtensionHandler(
+		const FOnGenerateGlobalRowExtensionArgs& Args,
+		TArray<FPropertyRowExtensionButton>& OutExtensionButtons
+		)
+	{
+		const TSharedPtr<IPropertyHandle> PropertyHandle = Args.PropertyHandle;
+		if (!PropertyHandle.IsValid())
+		{
+			return;
+		}
+		
+		FPropertyRowExtensionButton& CreateKey = OutExtensionButtons.AddDefaulted_GetRef();
+		CreateKey.Icon = FSlateIcon(FMultiUserReplicationStyle::GetStyleSetName(), TEXT("MultiUser.Icons.AddProperty.Small"));
+		CreateKey.Label = LOCTEXT("AddProperty.Label", "Replicate property");
+		CreateKey.ToolTip = TAttribute<FText>::CreateRaw(this, &FDetailRowExtensionHandler::GetToolTipText, PropertyHandle);
+		CreateKey.UIAction = FUIAction(
+			FExecuteAction::CreateRaw(this, &FDetailRowExtensionHandler::OnAssignPropertyClicked, PropertyHandle),
+			FCanExecuteAction::CreateRaw(this, &FDetailRowExtensionHandler::CanExecuteAssignPropertyClick, PropertyHandle),
+			FGetActionCheckState(),
+			FIsActionButtonVisible::CreateRaw(this, &FDetailRowExtensionHandler::IsAssignPropertyButtonVisible, PropertyHandle)
+		);
+	}
+	
+	FText FDetailRowExtensionHandler::GetToolTipText(TSharedPtr<IPropertyHandle> PropertyHandle) const
+	{
+		FText Reason;
+		return CanAssignPropertyWithReason(PropertyHandle, &Reason)
+			? LOCTEXT("AddProperty.Label.Tooltip", "Assign the property to yourself for replication in Multi-User.")
+			: Reason;
+	}
+
+	void FDetailRowExtensionHandler::OnAssignPropertyClicked(TSharedPtr<IPropertyHandle> PropertyHandle) const
+	{
+		using namespace DetailRowExtensionHandler;
+		const TSharedPtr<FMultiUserReplicationManager> ReplicationManager = WeakReplicationManager.Pin();
+		
+		if (!ensure(ReplicationManager) || !CanExecuteAssignPropertyClick(PropertyHandle))
+		{
+			return;
+		}
+		
+		const FPathInfo PathInfo = GetPropertyPath(PropertyHandle);
+		const TOptional<FConcertPropertyChain> PropertyChain = PathInfo.IsValid()
+			? FConcertPropertyChain::CreateFromPath(*PathInfo.LeafClass, PathInfo.Path)
+			: TOptional<FConcertPropertyChain>{};
+
+		TArray<UObject*> OuterObjects;
+		PropertyHandle->GetOuterObjects(OuterObjects);
+		
+		if (PropertyChain
+			&& ensure(!OuterObjects.IsEmpty())
+			&& ensure(ReplicationManager->GetConnectionState() == EMultiUserReplicationConnectionState::Connected))
+		{
+			// Handle multi-edit: only assign to those objects that do not have the property owned by any client, yet.
+			TArray<TSoftObjectPtr<>> ObjectsToAssignTo;
+			Algo::TransformIf(OuterObjects, ObjectsToAssignTo,
+				[this, &PathInfo](const UObject* Object) { return IsPropertyNotYetAssigned(*Object, PathInfo.Path); },
+				[](UObject* Object){ return Object; }
+				);
+			
+			MultiStreamColumns::FAssignPropertyModel::AssignPropertyTo(
+				*ReplicationManager->GetUnifiedClientView(),
+				ReplicationManager->GetOnlineClientManager()->GetLocalClient().GetEndpointId(),
+				ObjectsToAssignTo,
+				*PropertyChain
+				);
+		}
+	}
+
+	bool FDetailRowExtensionHandler::IsAssignPropertyButtonVisible(TSharedPtr<IPropertyHandle> PropertyHandle) const
+	{
+		const FProperty* Property = PropertyHandle->GetProperty();
+		const bool bIsReplicatableProperty = Property && ConcertSyncCore::PropertyChain::IsReplicatableProperty(*Property);
+		
+		const TSharedPtr<FMultiUserReplicationManager> ReplicationManager = WeakReplicationManager.Pin();
+		const bool bIsInReplicationSession = ensure(ReplicationManager)
+			&& ReplicationManager->GetConnectionState() == EMultiUserReplicationConnectionState::Connected;
+		return bIsReplicatableProperty && bIsInReplicationSession;
+	}
+
+#define SET_REASON(X) if (OutReason){ *OutReason = X; }
+	bool FDetailRowExtensionHandler::CanAssignPropertyWithReason(
+		TSharedPtr<IPropertyHandle> PropertyHandle,
+		FText* OutReason
+		) const
+	{
+		using namespace DetailRowExtensionHandler;
+		
+		const TSharedPtr<FMultiUserReplicationManager> ReplicationManager = WeakReplicationManager.Pin();
+		if (ensure(ReplicationManager) && ReplicationManager->GetConnectionState() != EMultiUserReplicationConnectionState::Connected)
+		{
+			SET_REASON(LOCTEXT("Reason.NotInSession", "You're not in any Multi-User session."));
+			return false;
+		}
+
+		// Do not construct FConcertPropertyChain: it's expensive since it iterates the property hierarchy.
+		// This function is called every tick for every shown property. 
+		const FPathInfo PathInfo = GetPropertyPath(PropertyHandle);
+		if (!PathInfo.IsValid())
+		{
+			SET_REASON(LOCTEXT("Reason.InvalidProperty", "This property cannot be replicated"));
+			return false;
+		}
+
+		TArray<UObject*> OuterObjects;
+		PropertyHandle->GetOuterObjects(OuterObjects);
+		if (OuterObjects.IsEmpty())
+		{
+			// Not sure what would cause this to happen - never hit it while testing.
+			SET_REASON(LOCTEXT("Reason.NoObjects", "That object in the details panel is outdated."));
+			return false;
+		}
+
+		// Handle multi-edit: if at least one property can be assigned, return true...
+		const bool bIsAnyUnowned = Algo::AnyOf(OuterObjects, [this, &PathInfo](const UObject* Object)
+		{
+			return IsPropertyNotYetAssigned(*Object, PathInfo.Path);
+		});
+		if (bIsAnyUnowned)
+		{
+			return true;
+		}
+
+		// ... at this point we know none of the properties can be assigned. We'll try to give the user a descriptive reason.
+		if (OuterObjects.Num() > 1)
+		{
+			SET_REASON(LOCTEXT("OwnedByMultiple", "You're editing multiple objects and the property is being replicated on each object already."));
+			return false;
+		}
+
+		// We can give a better reason text if there's only a single object to assign to.
+		return IsPropertyNotYetAssignedWithReason(*OuterObjects[0], PathInfo.Path, OutReason);
+	}
+	
+	bool FDetailRowExtensionHandler::IsPropertyNotYetAssignedWithReason(
+		const UObject& AssignedToObject,
+		const TArray<FName>& PropertyPath,
+		FText* OutReason
+		) const
+	{
+		const TSharedPtr<FMultiUserReplicationManager> ReplicationManager = WeakReplicationManager.Pin();
+		
+		const TOptional<FGuid> OwningClient = ensure(ReplicationManager)
+			? ReplicationManager->GetOnlineClientManager()->GetAuthorityCache().GetClientWithAuthorityOverProperty(&AssignedToObject, PropertyPath)
+			: TOptional<FGuid>{};
+		if (OwningClient && *OwningClient == ReplicationManager->GetOnlineClientManager()->GetLocalClient().GetEndpointId())
+		{
+			SET_REASON(LOCTEXT("AlreadyOwnedLocally", "You are already replicating this property."));
+			return false;
+		}
+		
+		if (OwningClient)
+		{
+			// Use lambda to defer evaluation below.
+			const auto GetClientName = [this, &OwningClient]
+			{
+				const TSharedPtr<IConcertSyncClient> Client = WeakClient.Pin();
+				const TSharedPtr<IConcertClientSession> Session = Client ? Client->GetConcertClient()->GetCurrentSession() : nullptr;
+				checkf(Session, TEXT("We checked for EMultiUserReplicationConnectionState::Connected above but have no valid session"));
+				return ClientUtils::GetClientDisplayName(*Session, *OwningClient);
+			};
+			
+			SET_REASON(
+				FText::Format(
+					LOCTEXT("Reason.OwnedByOtherClient", "This property is already owned by client {0}. Re-assign manually in replication view."),
+					FText::FromString(GetClientName())
+				)
+			);
+			return false;
+		}
+		
+		return true;
+	}
+#undef SET_REASON
+}
+
+#undef LOCTEXT_NAMESPACE
\ No newline at end of file
diff --git a/Engine/Plugins/Developer/Concert/ConcertApp/MultiUserClient/Source/MultiUserClient/Private/Replication/Misc/DetailRowExtensionHandler.h b/Engine/Plugins/Developer/Concert/ConcertApp/MultiUserClient/Source/MultiUserClient/Private/Replication/Misc/DetailRowExtensionHandler.h
new file mode 100644
--- /dev/null
+++ b/Engine/Plugins/Developer/Concert/ConcertApp/MultiUserClient/Source/MultiUserClient/Private/Replication/Misc/DetailRowExtensionHandler.h
@@ -0,0 +1,63 @@
+﻿// Copyright Epic Games, Inc. All Rights Reserved.
+
+#pragma once
+
+#include "Containers/ContainersFwd.h"
+#include "Templates/SharedPointer.h"
+#include "Templates/UnrealTemplate.h"
+
+class IConcertSyncClient;
+class IPropertyHandle;
+class UObject;
+
+struct FOnGenerateGlobalRowExtensionArgs;
+struct FPropertyRowExtensionButton;
+
+namespace UE::MultiUserClient::Replication
+{
+	class FMultiUserReplicationManager;
+
+	/** Adds a button to Details panel rows for quickly assigning the property to the local client for replication. */
+	class FDetailRowExtensionHandler : public FNoncopyable
+	{
+	public:
+		
+		FDetailRowExtensionHandler(const TSharedRef<IConcertSyncClient>& InClient, const TSharedRef<FMultiUserReplicationManager>& InReplicationManager);
+		~FDetailRowExtensionHandler();
+
+	private:
+
+		/** Used to look up client names. */
+		TWeakPtr<IConcertSyncClient> WeakClient;
+		/** Used to get the local client, to which properties will be assigned by default. */
+		TWeakPtr<FMultiUserReplicationManager> WeakReplicationManager;
+
+		/** Adds the assign button to valid property rows.  */
+		void RegisterExtensionHandler(
+			const FOnGenerateGlobalRowExtensionArgs& Args,
+			TArray<FPropertyRowExtensionButton>& OutExtensionButtons
+			);
+
+		/** @return Displays the reason why the property cannot be assigned. */
+		FText GetToolTipText(TSharedPtr<IPropertyHandle> PropertyHandle) const;
+
+		/** Assigns the property to the local client. */
+		void OnAssignPropertyClicked(TSharedPtr<IPropertyHandle> PropertyHandle) const;
+		/** @return Whether the property can be assigned. */
+		bool CanExecuteAssignPropertyClick(TSharedPtr<IPropertyHandle> PropertyHandle) const { return CanAssignProperty(PropertyHandle); }
+		/** @return Whether the client is in a MU session and the property can ever be replicated. */
+		bool IsAssignPropertyButtonVisible(TSharedPtr<IPropertyHandle> PropertyHandle) const;
+
+		/** @return Whether PropertyHandle can be assigned to its outer objects. */
+		bool CanAssignPropertyWithReason(TSharedPtr<IPropertyHandle> PropertyHandle, FText* OutReason = nullptr) const;
+		bool CanAssignProperty(TSharedPtr<IPropertyHandle> PropertyHandle) const { return CanAssignPropertyWithReason(PropertyHandle); }
+		
+		/** @return Whether the property identified by PropertyPath is not yet assigned to AssigedToObject by any client. */
+		bool IsPropertyNotYetAssignedWithReason(const UObject& AssignedToObject, const TArray<FName>& PropertyPath, FText* OutReason = nullptr) const;
+		bool IsPropertyNotYetAssigned(const UObject& Object, const TArray<FName>& PropertyChainPath) const
+		{
+			return IsPropertyNotYetAssignedWithReason(Object, PropertyChainPath);
+		}
+	};
+}
+
diff --git a/Engine/Plugins/Developer/Concert/ConcertApp/MultiUserClient/Source/MultiUserClient/Private/Widgets/ActiveSession/Replication/Client/Multi/Columns/AssignProperty/AssignPropertyModel.cpp b/Engine/Plugins/Developer/Concert/ConcertApp/MultiUserClient/Source/MultiUserClient/Private/Widgets/ActiveSession/Replication/Client/Multi/Columns/AssignProperty/AssignPropertyModel.cpp
--- a/Engine/Plugins/Developer/Concert/ConcertApp/MultiUserClient/Source/MultiUserClient/Private/Widgets/ActiveSession/Replication/Client/Multi/Columns/AssignProperty/AssignPropertyModel.cpp
+++ b/Engine/Plugins/Developer/Concert/ConcertApp/MultiUserClient/Source/MultiUserClient/Private/Widgets/ActiveSession/Replication/Client/Multi/Columns/AssignProperty/AssignPropertyModel.cpp
@@ -125,6 +125,25 @@ namespace UE::MultiUserClient::Replication::MultiStreamColumns::AssignPropertyMo
 
 namespace UE::MultiUserClient::Replication::MultiStreamColumns
 {
+	void FAssignPropertyModel::AssignPropertyTo(
+		const FUnifiedClientView& ClientView,
+		const FGuid& ClientId,
+		TConstArrayView<TSoftObjectPtr<>> Objects,
+		const FConcertPropertyChain& Property
+		)
+	{
+		if (const TSharedPtr<ConcertSharedSlate::IEditableReplicationStreamModel> StreamModel = ClientView.GetEditableClientStreamById(ClientId))
+		{
+			// Remove the property from all clients but the one we'll assign to ...
+			AssignPropertyModel::UnassignPropertyFromClients(ClientView, Objects, Property,
+				[ClientId](const FGuid& ClientToRemoveFrom){ return ClientId != ClientToRemoveFrom; }
+				);
+
+			// ... and then assign the property
+			AssignPropertyModel::AssignPropertyTo(*StreamModel, Objects, Property);
+		}
+	}
+
 	FAssignPropertyModel::FAssignPropertyModel(FUnifiedClientView& InClientView, FMultiViewOptions& InViewOptions)
 		: ClientView(InClientView)
 		, ViewOptions(InViewOptions)
@@ -218,8 +237,7 @@ namespace UE::MultiUserClient::Replication::MultiStreamColumns
 		) const
 	{
 		// Remote clients can disconnect after the combo-box is opened.
-		const TSharedPtr<ConcertSharedSlate::IEditableReplicationStreamModel> StreamModel = ClientView.GetEditableClientStreamById(ClientId);
-		if (!StreamModel)
+		if (!ClientView.GetEditableClientStreamById(ClientId))
 		{
 			return;
 		}
@@ -241,13 +259,7 @@ namespace UE::MultiUserClient::Replication::MultiStreamColumns
 		}
 		else
 		{
-			// ... remove the property from all clients but the one we'll assign to ...
-			AssignPropertyModel::UnassignPropertyFromClients(ClientView, Objects, Property,
-				[ClientId](const FGuid& ClientToRemoveFrom){ return ClientId != ClientToRemoveFrom; }
-				);
-
-			// ... and then assign the property
-			AssignPropertyModel::AssignPropertyTo(*StreamModel, Objects, Property);
+			AssignPropertyTo(ClientView, ClientId, Objects, Property);
 		}
 	}
 
diff --git a/Engine/Plugins/Developer/Concert/ConcertApp/MultiUserClient/Source/MultiUserClient/Private/Widgets/ActiveSession/Replication/Client/Multi/Columns/AssignProperty/AssignPropertyModel.h b/Engine/Plugins/Developer/Concert/ConcertApp/MultiUserClient/Source/MultiUserClient/Private/Widgets/ActiveSession/Replication/Client/Multi/Columns/AssignProperty/AssignPropertyModel.h
--- a/Engine/Plugins/Developer/Concert/ConcertApp/MultiUserClient/Source/MultiUserClient/Private/Widgets/ActiveSession/Replication/Client/Multi/Columns/AssignProperty/AssignPropertyModel.h
+++ b/Engine/Plugins/Developer/Concert/ConcertApp/MultiUserClient/Source/MultiUserClient/Private/Widgets/ActiveSession/Replication/Client/Multi/Columns/AssignProperty/AssignPropertyModel.h
@@ -10,6 +10,7 @@
 enum class EBreakBehavior : uint8;
 struct FConcertPropertyChain;
 
+namespace UE::ConcertSharedSlate { class IEditableReplicationStreamModel; }
 namespace UE::MultiUserClient::Replication
 {
 	class FMultiViewOptions;
@@ -34,6 +35,14 @@ namespace UE::MultiUserClient::Replication::MultiStreamColumns
 	class FAssignPropertyModel : public FNoncopyable
 	{
 	public:
+
+		/** Assigns the property to the given client and unassigns it from all others. */
+		static void AssignPropertyTo(
+			const FUnifiedClientView& ClientView,
+			const FGuid& ClientId,
+			TConstArrayView<TSoftObjectPtr<>> Objects,
+			const FConcertPropertyChain& Property
+		);
 		
 		FAssignPropertyModel(FUnifiedClientView& InClientView UE_LIFETIMEBOUND, FMultiViewOptions& InViewOptions UE_LIFETIMEBOUND);
 		~FAssignPropertyModel();
