Add getters to browser widget classes (in preparation for future changes).

diff --git a/Engine/Plugins/Developer/Concert/ConcertApp/MultiUserClient/Source/MultiUserClient/Private/MultiUserClientModule.cpp b/Engine/Plugins/Developer/Concert/ConcertApp/MultiUserClient/Source/MultiUserClient/Private/MultiUserClientModule.cpp
--- a/Engine/Plugins/Developer/Concert/ConcertApp/MultiUserClient/Source/MultiUserClient/Private/MultiUserClientModule.cpp
+++ b/Engine/Plugins/Developer/Concert/ConcertApp/MultiUserClient/Source/MultiUserClient/Private/MultiUserClientModule.cpp
@@ -1073,7 +1073,7 @@ private:
 		return SNew(SDockTab)
 			.TabRole(NomadTab)
 			[
-				SNew(SConcertBrowser, MultiUserClient.ToSharedRef(), ReplicationManager.ToSharedRef())
+				SAssignNew(BrowserTab, SConcertBrowser, MultiUserClient.ToSharedRef(), ReplicationManager.ToSharedRef())
 			];
 	}
 
@@ -1427,6 +1427,9 @@ private:
 	/** Interacts with the replication system on behalf of Multi-User. */
 	TSharedPtr<UE::MultiUserClient::Replication::FMultiUserReplicationManager> ReplicationManager;
 
+	/** The content of the MU tab, if open. Invalid or stale if closed. */
+	TWeakPtr<SConcertBrowser> BrowserTab;
+
 #if WITH_EDITOR
 	/** Toggles the game visibility column in the Levels tab when joining & leaving sessions. */
 	TUniquePtr<UE::MultiUserClient::FGameVisibilityColumnToggler> VisibilityColumnToggler;
diff --git a/Engine/Plugins/Developer/Concert/ConcertApp/MultiUserClient/Source/MultiUserClient/Private/Widgets/ActiveSession/Replication/Client/Multi/SAllClientsView.cpp b/Engine/Plugins/Developer/Concert/ConcertApp/MultiUserClient/Source/MultiUserClient/Private/Widgets/ActiveSession/Replication/Client/Multi/SAllClientsView.cpp
--- a/Engine/Plugins/Developer/Concert/ConcertApp/MultiUserClient/Source/MultiUserClient/Private/Widgets/ActiveSession/Replication/Client/Multi/SAllClientsView.cpp
+++ b/Engine/Plugins/Developer/Concert/ConcertApp/MultiUserClient/Source/MultiUserClient/Private/Widgets/ActiveSession/Replication/Client/Multi/SAllClientsView.cpp
@@ -15,7 +15,7 @@ namespace UE::MultiUserClient::Replication
 		
 		ChildSlot
 		[
-			SNew(SMultiClientView, InConcertClient, InMultiUserReplicationManager, *AllOnlineClientsModel, *AllOfflineClientsModel)
+			SAssignNew(ClientView, SMultiClientView, InConcertClient, InMultiUserReplicationManager, *AllOnlineClientsModel, *AllOfflineClientsModel)
 		];
 	}
 }
diff --git a/Engine/Plugins/Developer/Concert/ConcertApp/MultiUserClient/Source/MultiUserClient/Private/Widgets/ActiveSession/Replication/Client/Multi/SAllClientsView.h b/Engine/Plugins/Developer/Concert/ConcertApp/MultiUserClient/Source/MultiUserClient/Private/Widgets/ActiveSession/Replication/Client/Multi/SAllClientsView.h
--- a/Engine/Plugins/Developer/Concert/ConcertApp/MultiUserClient/Source/MultiUserClient/Private/Widgets/ActiveSession/Replication/Client/Multi/SAllClientsView.h
+++ b/Engine/Plugins/Developer/Concert/ConcertApp/MultiUserClient/Source/MultiUserClient/Private/Widgets/ActiveSession/Replication/Client/Multi/SAllClientsView.h
@@ -16,6 +16,7 @@ namespace UE::MultiUserClient::Replication
 	class FMultiUserReplicationManager;
 	class FOnlineClient;
 	class FOnlineClientManager;
+	class SMultiClientView;
 
 	/** Leverages SMultiClientView to display all online and offline clients. */
 	class SAllClientsView : public SCompoundWidget
@@ -25,11 +26,16 @@ namespace UE::MultiUserClient::Replication
 
 		void Construct(const FArguments&, TSharedRef<IConcertClient> InConcertClient, FMultiUserReplicationManager& InMultiUserReplicationManager UE_LIFETIMEBOUND);
 
+		const TSharedPtr<SMultiClientView>& GetClientView() const { return ClientView; }
+		
 	private:
 
 		/** Keeps track of all online clients. */
 		TUniquePtr<FAllOnlineClientsSelectionModel> AllOnlineClientsModel;
 		/** Keeps track of all offline clients. */
 		TUniquePtr<FAllOfflineClientsSelectionModel> AllOfflineClientsModel;
+
+		/** Shows the content of all replication clients. */
+		TSharedPtr<SMultiClientView> ClientView;
 	};
 }
diff --git a/Engine/Plugins/Developer/Concert/ConcertApp/MultiUserClient/Source/MultiUserClient/Private/Widgets/ActiveSession/Replication/Client/Multi/SMultiClientView.h b/Engine/Plugins/Developer/Concert/ConcertApp/MultiUserClient/Source/MultiUserClient/Private/Widgets/ActiveSession/Replication/Client/Multi/SMultiClientView.h
--- a/Engine/Plugins/Developer/Concert/ConcertApp/MultiUserClient/Source/MultiUserClient/Private/Widgets/ActiveSession/Replication/Client/Multi/SMultiClientView.h
+++ b/Engine/Plugins/Developer/Concert/ConcertApp/MultiUserClient/Source/MultiUserClient/Private/Widgets/ActiveSession/Replication/Client/Multi/SMultiClientView.h
@@ -48,6 +48,8 @@ namespace UE::MultiUserClient::Replication
 			);
 		virtual ~SMultiClientView() override;
 
+		const TSharedPtr<ConcertSharedSlate::IMultiReplicationStreamEditor>& GetStreamEditor() const { return StreamEditor; }
+
 	private:
 
 		TSharedPtr<IConcertClient> ConcertClient;
diff --git a/Engine/Plugins/Developer/Concert/ConcertApp/MultiUserClient/Source/MultiUserClient/Private/Widgets/ActiveSession/Replication/Joined/SReplicationJoinedView.cpp b/Engine/Plugins/Developer/Concert/ConcertApp/MultiUserClient/Source/MultiUserClient/Private/Widgets/ActiveSession/Replication/Joined/SReplicationJoinedView.cpp
--- a/Engine/Plugins/Developer/Concert/ConcertApp/MultiUserClient/Source/MultiUserClient/Private/Widgets/ActiveSession/Replication/Joined/SReplicationJoinedView.cpp
+++ b/Engine/Plugins/Developer/Concert/ConcertApp/MultiUserClient/Source/MultiUserClient/Private/Widgets/ActiveSession/Replication/Joined/SReplicationJoinedView.cpp
@@ -16,7 +16,7 @@ namespace UE::MultiUserClient::Replication
 	{
 		ChildSlot
 		[
-			SNew(SAllClientsView, InClient->GetConcertClient(), *InReplicationManager)
+			SAssignNew(AllClientsView, SAllClientsView, InClient->GetConcertClient(), *InReplicationManager)
 		];
 	}
 }
\ No newline at end of file
diff --git a/Engine/Plugins/Developer/Concert/ConcertApp/MultiUserClient/Source/MultiUserClient/Private/Widgets/ActiveSession/Replication/Joined/SReplicationJoinedView.h b/Engine/Plugins/Developer/Concert/ConcertApp/MultiUserClient/Source/MultiUserClient/Private/Widgets/ActiveSession/Replication/Joined/SReplicationJoinedView.h
--- a/Engine/Plugins/Developer/Concert/ConcertApp/MultiUserClient/Source/MultiUserClient/Private/Widgets/ActiveSession/Replication/Joined/SReplicationJoinedView.h
+++ b/Engine/Plugins/Developer/Concert/ConcertApp/MultiUserClient/Source/MultiUserClient/Private/Widgets/ActiveSession/Replication/Joined/SReplicationJoinedView.h
@@ -14,7 +14,7 @@ namespace UE::ConcertSharedSlate { class IReplicationStreamEditor; }
 namespace UE::MultiUserClient::Replication
 {
 	class FMultiUserReplicationManager;
-	class SReplicationClientView;
+	class SAllClientsView;
 
 	/** This widget is displayed by SReplicationRootWidget when the client has joined replication. */
 	class SReplicationJoinedView : public SCompoundWidget
@@ -31,11 +31,16 @@ namespace UE::MultiUserClient::Replication
 			const TSharedRef<IConcertSyncClient>& InClient
 			);
 
+		const TSharedPtr<SAllClientsView>& GetAllClientsView() const { return AllClientsView; }
+
 	private:
 
 		/** The local client this widget is created for. */
 		TSharedPtr<IConcertSyncClient> Client;
 		/** Acts as the model of this view */
 		TSharedPtr<FMultiUserReplicationManager> ReplicationManager;
+
+		/** Shows the content of all replication clients. */
+		TSharedPtr<SAllClientsView> AllClientsView;
 	};
 }
\ No newline at end of file
diff --git a/Engine/Plugins/Developer/Concert/ConcertApp/MultiUserClient/Source/MultiUserClient/Private/Widgets/ActiveSession/Replication/SReplicationRootWidget.cpp b/Engine/Plugins/Developer/Concert/ConcertApp/MultiUserClient/Source/MultiUserClient/Private/Widgets/ActiveSession/Replication/SReplicationRootWidget.cpp
--- a/Engine/Plugins/Developer/Concert/ConcertApp/MultiUserClient/Source/MultiUserClient/Private/Widgets/ActiveSession/Replication/SReplicationRootWidget.cpp
+++ b/Engine/Plugins/Developer/Concert/ConcertApp/MultiUserClient/Source/MultiUserClient/Private/Widgets/ActiveSession/Replication/SReplicationRootWidget.cpp
@@ -54,6 +54,8 @@ namespace UE::MultiUserClient::Replication
 
 	void SReplicationRootWidget::ShowWidget_Connecting()
 	{
+		ConnectedView.Reset();
+		
 		ChildSlot
 		[
 			SNew(SBox)
@@ -85,12 +87,14 @@ namespace UE::MultiUserClient::Replication
 	{
 		ChildSlot
 		[
-			SNew(SReplicationJoinedView, ReplicationManager.ToSharedRef(), Client.ToSharedRef())
+			SAssignNew(ConnectedView, SReplicationJoinedView, ReplicationManager.ToSharedRef(), Client.ToSharedRef())
 		];
 	}
 
 	void SReplicationRootWidget::ShowWidget_Disconnected()
 	{
+		ConnectedView.Reset();
+		
 		ChildSlot
 		[
 			SNew(SBox)
diff --git a/Engine/Plugins/Developer/Concert/ConcertApp/MultiUserClient/Source/MultiUserClient/Private/Widgets/ActiveSession/Replication/SReplicationRootWidget.h b/Engine/Plugins/Developer/Concert/ConcertApp/MultiUserClient/Source/MultiUserClient/Private/Widgets/ActiveSession/Replication/SReplicationRootWidget.h
--- a/Engine/Plugins/Developer/Concert/ConcertApp/MultiUserClient/Source/MultiUserClient/Private/Widgets/ActiveSession/Replication/SReplicationRootWidget.h
+++ b/Engine/Plugins/Developer/Concert/ConcertApp/MultiUserClient/Source/MultiUserClient/Private/Widgets/ActiveSession/Replication/SReplicationRootWidget.h
@@ -11,6 +11,7 @@ class IConcertSyncClient;
 namespace UE::MultiUserClient::Replication
 {
 	class FMultiUserReplicationManager;
+	class SReplicationJoinedView;
 	
 	enum class EMultiUserReplicationConnectionState : uint8;
 
@@ -25,12 +26,17 @@ namespace UE::MultiUserClient::Replication
 
 		void Construct(const FArguments& InArgs, TSharedRef<FMultiUserReplicationManager> InReplicationManager, TSharedRef<IConcertSyncClient> InClient);
 
+		const TWeakPtr<SReplicationJoinedView>& GetConnectedView() const { return ConnectedView; }
+
 	private:
 
 		/** The client this widget was created for. */
 		TSharedPtr<IConcertSyncClient> Client;
 		/** Manages the business logic which we represent. */
 		TSharedPtr<FMultiUserReplicationManager> ReplicationManager;
+
+		/** Shown when the replication handshake has been completed. Invalid otherwise. */
+		TWeakPtr<SReplicationJoinedView> ConnectedView;
 		
 		/** Called when joining replication session state changes. */
 		void OnReplicationConnectionStateChanged(EMultiUserReplicationConnectionState);
diff --git a/Engine/Plugins/Developer/Concert/ConcertApp/MultiUserClient/Source/MultiUserClient/Private/Widgets/ActiveSession/SActiveSessionRoot.cpp b/Engine/Plugins/Developer/Concert/ConcertApp/MultiUserClient/Source/MultiUserClient/Private/Widgets/ActiveSession/SActiveSessionRoot.cpp
--- a/Engine/Plugins/Developer/Concert/ConcertApp/MultiUserClient/Source/MultiUserClient/Private/Widgets/ActiveSession/SActiveSessionRoot.cpp
+++ b/Engine/Plugins/Developer/Concert/ConcertApp/MultiUserClient/Source/MultiUserClient/Private/Widgets/ActiveSession/SActiveSessionRoot.cpp
@@ -57,16 +57,21 @@ namespace UE::MultiUserClient
 
 				+SWidgetSwitcher::Slot()
 				[
-					SNew(SActiveSessionOverviewTab, InConcertSyncClient)
+					SAssignNew(OverviewContent, SActiveSessionOverviewTab, InConcertSyncClient)
 				]
 				+SWidgetSwitcher::Slot()
 				[
-					SNew(Replication::SReplicationRootWidget, InReplicationManager, InConcertSyncClient.ToSharedRef())
+					SAssignNew(ReplicationContent, Replication::SReplicationRootWidget, InReplicationManager, InConcertSyncClient.ToSharedRef())
 				]
 			]
 		];
 	}
 
+	void SActiveSessionRoot::OpenTab(EMultiUserTab Tab)
+	{
+		TabSwitcher->SetActiveWidgetIndex(static_cast<int32>(Tab));
+	}
+
 	TSharedRef<SWidget> SActiveSessionRoot::CreateTabArea()
 	{
 		const auto CreateTabEntry = [this](int32 Index, const ANSICHAR* ImageBrush, FText Label, FText ToolTipText)
diff --git a/Engine/Plugins/Developer/Concert/ConcertApp/MultiUserClient/Source/MultiUserClient/Private/Widgets/ActiveSession/SActiveSessionRoot.h b/Engine/Plugins/Developer/Concert/ConcertApp/MultiUserClient/Source/MultiUserClient/Private/Widgets/ActiveSession/SActiveSessionRoot.h
--- a/Engine/Plugins/Developer/Concert/ConcertApp/MultiUserClient/Source/MultiUserClient/Private/Widgets/ActiveSession/SActiveSessionRoot.h
+++ b/Engine/Plugins/Developer/Concert/ConcertApp/MultiUserClient/Source/MultiUserClient/Private/Widgets/ActiveSession/SActiveSessionRoot.h
@@ -9,17 +9,27 @@ class FSpawnTabArgs;
 class FTabManager;
 class FWorkspaceItem;
 class IConcertSyncClient;
+class SActiveSessionOverviewTab;
 class SDockTab;
 class SWidgetSwitcher;
 class SWindow;
-
-namespace UE::MultiUserClient::Replication { class FMultiUserReplicationManager; }
+namespace UE::MultiUserClient::Replication
+{
+	class FMultiUserReplicationManager;
+	class SReplicationRootWidget;
+}
 
 namespace UE::MultiUserClient
 {
+	enum class EMultiUserTab : uint8
+	{
+		Overview,
+		Replication
+	};
+	
 	/**
 	 * Displayed when the client is connected to an active session.
-	 * Manages the child content in tabs.
+	 * Manages the child content in inline tabs.
 	 */
 	class SActiveSessionRoot : public SCompoundWidget
 	{
@@ -34,11 +44,22 @@ namespace UE::MultiUserClient
 
 		void Construct(const FArguments& InArgs, TSharedPtr<IConcertSyncClient> InConcertSyncClient, TSharedRef<Replication::FMultiUserReplicationManager> InReplicationManager);
 
+		const TSharedPtr<SActiveSessionOverviewTab>& GetOverviewContent() const { return OverviewContent; }
+		const TSharedPtr<Replication::SReplicationRootWidget>& GetReplicationContent() const { return ReplicationContent; }
+
+		/** Opens the specified inline tab. */
+		void OpenTab(EMultiUserTab Tab);
+		
 	private:
 
 		/** This switches "tabs" when a button in the "tab" area is changed. */
 		TSharedPtr<SWidgetSwitcher> TabSwitcher;
 
+		/** Shows general stats, like transaction history and other clients. */
+		TSharedPtr<SActiveSessionOverviewTab> OverviewContent;
+		/** Shows content specific to replication. */
+		TSharedPtr<Replication::SReplicationRootWidget> ReplicationContent;
+
 		TSharedRef<SWidget> CreateTabArea();
 	};
 }
diff --git a/Engine/Plugins/Developer/Concert/ConcertApp/MultiUserClient/Source/MultiUserClient/Private/Widgets/SConcertBrowser.cpp b/Engine/Plugins/Developer/Concert/ConcertApp/MultiUserClient/Source/MultiUserClient/Private/Widgets/SConcertBrowser.cpp
--- a/Engine/Plugins/Developer/Concert/ConcertApp/MultiUserClient/Source/MultiUserClient/Private/Widgets/SConcertBrowser.cpp
+++ b/Engine/Plugins/Developer/Concert/ConcertApp/MultiUserClient/Source/MultiUserClient/Private/Widgets/SConcertBrowser.cpp
@@ -60,13 +60,16 @@ void SConcertBrowser::AttachChildWidget(EConcertConnectionStatus ConnectionStatu
 	if (const TSharedPtr<IConcertSyncClient> ConcertSyncClient = WeakConcertSyncClient.Pin()
 		; ensure(ConcertSyncClient))
 	{
+		ActiveSessionWidget.Reset();
+		SessionBrowser.Reset();
+		
 		if (ConnectionStatus == EConcertConnectionStatus::Connected)
 		{
 			if (const TSharedPtr<UE::MultiUserClient::Replication::FMultiUserReplicationManager> ReplicationManager = WeakReplicationManager.Pin()
 				; ensure(ReplicationManager))
 			{
 				ChildSlot.AttachWidget(
-					SNew(UE::MultiUserClient::SActiveSessionRoot,
+					SAssignNew(ActiveSessionWidget, UE::MultiUserClient::SActiveSessionRoot,
 						ConcertSyncClient,
 						ReplicationManager.ToSharedRef()
 						)
@@ -76,7 +79,7 @@ void SConcertBrowser::AttachChildWidget(EConcertConnectionStatus ConnectionStatu
 		else if (ConnectionStatus == EConcertConnectionStatus::Disconnected)
 		{
 			ChildSlot.AttachWidget(
-				SNew(SConcertClientSessionBrowser,
+				SAssignNew(SessionBrowser, SConcertClientSessionBrowser,
 					ConcertSyncClient->GetConcertClient(),
 					SearchedText
 					)
diff --git a/Engine/Plugins/Developer/Concert/ConcertApp/MultiUserClient/Source/MultiUserClient/Private/Widgets/SConcertBrowser.h b/Engine/Plugins/Developer/Concert/ConcertApp/MultiUserClient/Source/MultiUserClient/Private/Widgets/SConcertBrowser.h
--- a/Engine/Plugins/Developer/Concert/ConcertApp/MultiUserClient/Source/MultiUserClient/Private/Widgets/SConcertBrowser.h
+++ b/Engine/Plugins/Developer/Concert/ConcertApp/MultiUserClient/Source/MultiUserClient/Private/Widgets/SConcertBrowser.h
@@ -11,6 +11,8 @@
 
 class IConcertClientSession;
 class IConcertSyncClient;
+class SConcertClientSessionBrowser;
+namespace UE::MultiUserClient { class SActiveSessionRoot; }
 namespace UE::MultiUserClient::Replication { class FMultiUserReplicationManager; }
 
 /**
@@ -32,7 +34,13 @@ public:
 	* @param InSyncClient The sync client.
 	* @param InReplicationManager Used to create replication UI when in an active session
 	*/
-	void Construct(const FArguments& InArgs, TSharedRef<IConcertSyncClient> InSyncClient, TSharedRef<UE::MultiUserClient::Replication::FMultiUserReplicationManager> InReplicationManager);
+	void Construct(const FArguments& InArgs,
+		TSharedRef<IConcertSyncClient> InSyncClient,
+		TSharedRef<UE::MultiUserClient::Replication::FMultiUserReplicationManager> InReplicationManager
+		);
+
+	const TWeakPtr<UE::MultiUserClient::SActiveSessionRoot>& GetActiveSessionWidget() const { return ActiveSessionWidget; }
+	const TWeakPtr<SConcertClientSessionBrowser>& GetSessionBrowser() const { return SessionBrowser; }
 
 private:
 
@@ -41,6 +49,11 @@ private:
 	/** Interacts with the replication system on behalf of Multi-User. */
 	TWeakPtr<UE::MultiUserClient::Replication::FMultiUserReplicationManager> WeakReplicationManager;
 
+	/** Only valid when in session. */
+	TWeakPtr<UE::MultiUserClient::SActiveSessionRoot> ActiveSessionWidget;
+	/** Only valid when not in session. */
+	TWeakPtr<SConcertClientSessionBrowser> SessionBrowser;
+
 	/** Keeps the session browser searched text in memory to reapply it when a user leaves a session and goes back to the session browser. */
 	TSharedPtr<FText> SearchedText;
 	
