Update the toolbar button colour when the tab is opened via API call.

diff --git a/Engine/Plugins/Developer/Concert/ConcertApp/MultiUserClient/Source/MultiUserClient/Private/Widgets/ActiveSession/SActiveSessionRoot.cpp b/Engine/Plugins/Developer/Concert/ConcertApp/MultiUserClient/Source/MultiUserClient/Private/Widgets/ActiveSession/SActiveSessionRoot.cpp
--- a/Engine/Plugins/Developer/Concert/ConcertApp/MultiUserClient/Source/MultiUserClient/Private/Widgets/ActiveSession/SActiveSessionRoot.cpp
+++ b/Engine/Plugins/Developer/Concert/ConcertApp/MultiUserClient/Source/MultiUserClient/Private/Widgets/ActiveSession/SActiveSessionRoot.cpp
@@ -69,7 +69,9 @@ namespace UE::MultiUserClient
 
 	void SActiveSessionRoot::OpenTab(EMultiUserTab Tab)
 	{
-		TabSwitcher->SetActiveWidgetIndex(static_cast<int32>(Tab));
+		const int32 Index = static_cast<int32>(Tab);
+		TabArea->SetButtonActivated(Index);
+		TabSwitcher->SetActiveWidgetIndex(Index);
 	}
 
 	TSharedRef<SWidget> SActiveSessionRoot::CreateTabArea()
@@ -108,7 +110,7 @@ namespace UE::MultiUserClient
 			CreateTabEntry(1, "Concert.MultiUser", LOCTEXT("ReplicationTab.Label", "Replication"), LOCTEXT("ReplicationTab.Tooltip", "Manage real-time object replication"))
 		};
 		
-		return SNew(STabArea)
+		return SAssignNew(TabArea, STabArea)
 			.Tabs(Tabs)
 			.ActiveTabIndex(0);
 	}
diff --git a/Engine/Plugins/Developer/Concert/ConcertApp/MultiUserClient/Source/MultiUserClient/Private/Widgets/ActiveSession/SActiveSessionRoot.h b/Engine/Plugins/Developer/Concert/ConcertApp/MultiUserClient/Source/MultiUserClient/Private/Widgets/ActiveSession/SActiveSessionRoot.h
--- a/Engine/Plugins/Developer/Concert/ConcertApp/MultiUserClient/Source/MultiUserClient/Private/Widgets/ActiveSession/SActiveSessionRoot.h
+++ b/Engine/Plugins/Developer/Concert/ConcertApp/MultiUserClient/Source/MultiUserClient/Private/Widgets/ActiveSession/SActiveSessionRoot.h
@@ -21,6 +21,8 @@ namespace UE::MultiUserClient::Replication
 
 namespace UE::MultiUserClient
 {
+	class STabArea;
+
 	enum class EMultiUserTab : uint8
 	{
 		Overview,
@@ -54,6 +56,8 @@ namespace UE::MultiUserClient
 
 		/** This switches "tabs" when a button in the "tab" area is changed. */
 		TSharedPtr<SWidgetSwitcher> TabSwitcher;
+		/** Holds the buttons for switching inline tabs. */
+		TSharedPtr<STabArea> TabArea;
 
 		/** Shows general stats, like transaction history and other clients. */
 		TSharedPtr<SActiveSessionOverviewTab> OverviewContent;
diff --git a/Engine/Plugins/Developer/Concert/ConcertApp/MultiUserClient/Source/MultiUserClient/Private/Widgets/ActiveSession/STabArea.cpp b/Engine/Plugins/Developer/Concert/ConcertApp/MultiUserClient/Source/MultiUserClient/Private/Widgets/ActiveSession/STabArea.cpp
--- a/Engine/Plugins/Developer/Concert/ConcertApp/MultiUserClient/Source/MultiUserClient/Private/Widgets/ActiveSession/STabArea.cpp
+++ b/Engine/Plugins/Developer/Concert/ConcertApp/MultiUserClient/Source/MultiUserClient/Private/Widgets/ActiveSession/STabArea.cpp
@@ -43,7 +43,16 @@ namespace UE::MultiUserClient
 			Content
 		];
 	}
-	
+
+	void STabArea::SetButtonActivated(int32 ButtonIndex)
+	{
+		if (ensure(TabButtons.IsValidIndex(ButtonIndex)))
+		{
+			TabButtons[ButtonIndex]->Activate();
+			OnButtonActivated(ButtonIndex);
+		}
+	}
+
 	void STabArea::OnButtonActivated(int32 ButtonIndex)
 	{
 		for (int32 i = 0; i < TabButtons.Num(); ++i)
diff --git a/Engine/Plugins/Developer/Concert/ConcertApp/MultiUserClient/Source/MultiUserClient/Private/Widgets/ActiveSession/STabArea.h b/Engine/Plugins/Developer/Concert/ConcertApp/MultiUserClient/Source/MultiUserClient/Private/Widgets/ActiveSession/STabArea.h
--- a/Engine/Plugins/Developer/Concert/ConcertApp/MultiUserClient/Source/MultiUserClient/Private/Widgets/ActiveSession/STabArea.h
+++ b/Engine/Plugins/Developer/Concert/ConcertApp/MultiUserClient/Source/MultiUserClient/Private/Widgets/ActiveSession/STabArea.h
@@ -37,6 +37,9 @@ namespace UE::MultiUserClient
 
 		void Construct(const FArguments& InArgs);
 
+		/** Informs the owning view has manually changed to the content displayed by ButtonIndex. Make the button appear selected. */
+		void SetButtonActivated(int32 ButtonIndex);
+
 	private:
 
 		/** The buttons in this area. */
