After adding an object for replication using the Details panel, show that property in the replication UI.

diff --git a/Engine/Plugins/Developer/Concert/ConcertApp/MultiUserClient/Source/MultiUserClient/Private/MultiUserClientModule.cpp b/Engine/Plugins/Developer/Concert/ConcertApp/MultiUserClient/Source/MultiUserClient/Private/MultiUserClientModule.cpp
--- a/Engine/Plugins/Developer/Concert/ConcertApp/MultiUserClient/Source/MultiUserClient/Private/MultiUserClientModule.cpp
+++ b/Engine/Plugins/Developer/Concert/ConcertApp/MultiUserClient/Source/MultiUserClient/Private/MultiUserClientModule.cpp
@@ -636,7 +636,8 @@ public:
 		{
 			VisibilityColumnToggler = MakeUnique<UE::MultiUserClient::FGameVisibilityColumnToggler>(MultiUserClient.ToSharedRef());
 			DetailReplicationRowExtensionHandler = MakeUnique<UE::MultiUserClient::Replication::FDetailRowExtensionHandler>(
-				MultiUserClient.ToSharedRef(), ReplicationManager.ToSharedRef()
+				MultiUserClient.ToSharedRef(), ReplicationManager.ToSharedRef(),
+				UE::MultiUserClient::FGetConcertBrowserWidget::CreateRaw(this, &FMultiUserClientModule::GetOrOpenBrowserTabContent)
 				);
 		}
 #endif
@@ -1062,6 +1063,13 @@ private:
 		FGlobalTabmanager::Get()->TryInvokeTab(FTabId(ConcertBrowserTabName));
 	}
 
+	TSharedPtr<SConcertBrowser> GetOrOpenBrowserTabContent() const
+	{
+		// If the tab is already open, it's ok that its flashed. Some UI code is getting the content of the tab to show something in it.
+		FGlobalTabmanager::Get()->TryInvokeTab(FTabId(ConcertBrowserTabName));
+		return BrowserTab.Pin();
+	}
+	
 	/**
 	 * Creates a new Concert Browser front-end tab.
 	 *
diff --git a/Engine/Plugins/Developer/Concert/ConcertApp/MultiUserClient/Source/MultiUserClient/Private/Replication/Misc/DetailRowExtensionHandler.cpp b/Engine/Plugins/Developer/Concert/ConcertApp/MultiUserClient/Source/MultiUserClient/Private/Replication/Misc/DetailRowExtensionHandler.cpp
--- a/Engine/Plugins/Developer/Concert/ConcertApp/MultiUserClient/Source/MultiUserClient/Private/Replication/Misc/DetailRowExtensionHandler.cpp
+++ b/Engine/Plugins/Developer/Concert/ConcertApp/MultiUserClient/Source/MultiUserClient/Private/Replication/Misc/DetailRowExtensionHandler.cpp
@@ -4,11 +4,16 @@
 
 #include "IConcertSyncClient.h"
 #include "MultiUserReplicationStyle.h"
+#include "Replication/Client/ClientUtils.h"
+#include "Replication/Editor/View/IMultiReplicationStreamEditor.h"
+#include "Replication/Editor/View/IReplicationStreamEditor.h"
 #include "Replication/MultiUserReplicationManager.h"
 #include "Replication/PropertyChainUtils.h"
-#include "Replication/Client/ClientUtils.h"
 #include "Widgets/ActiveSession/Replication/Client/Multi/Columns/AssignProperty/AssignPropertyModel.h"
+#include "Widgets/ActiveSession/SActiveSessionRoot.h"
 
+#include "Algo/RemoveIf.h"
+#include "GameFramework/Actor.h"
 #include "PropertyEditorModule.h"
 #include "PropertyHandle.h"
 #include "PropertyPath.h"
@@ -117,11 +122,15 @@ namespace UE::MultiUserClient::Replication
 {
 	FDetailRowExtensionHandler::FDetailRowExtensionHandler(
 		const TSharedRef<IConcertSyncClient>& InClient,
-		const TSharedRef<FMultiUserReplicationManager>& InReplicationManager
+		const TSharedRef<FMultiUserReplicationManager>& InReplicationManager,
+		FGetConcertBrowserWidget InGetOrInvokeBrowserTabDelegate
 		)
 		: WeakClient(InClient)
 		, WeakReplicationManager(InReplicationManager)
+		, GetOrInvokeBrowserTabDelegate(MoveTemp(InGetOrInvokeBrowserTabDelegate))
 	{
+		check(GetOrInvokeBrowserTabDelegate.IsBound());
+		
 		FPropertyEditorModule& EditModule = FModuleManager::Get().GetModuleChecked<FPropertyEditorModule>("PropertyEditor");
 		EditModule.GetGlobalRowExtensionDelegate().AddRaw(
 			this, &FDetailRowExtensionHandler::RegisterExtensionHandler
@@ -184,17 +193,17 @@ namespace UE::MultiUserClient::Replication
 
 		TArray<UObject*> OuterObjects;
 		PropertyHandle->GetOuterObjects(OuterObjects);
+		// Handle multi-edit: only assign to those objects that do not have the property owned by any client, yet.
+		OuterObjects.SetNum(
+			Algo::RemoveIf(OuterObjects, [this, &PathInfo](const UObject* Object){ return !IsPropertyNotYetAssigned(*Object, PathInfo.Path);})
+		);
 		
 		if (PropertyChain
-			&& ensure(!OuterObjects.IsEmpty())
+			&& !OuterObjects.IsEmpty()
 			&& ensure(ReplicationManager->GetConnectionState() == EMultiUserReplicationConnectionState::Connected))
 		{
-			// Handle multi-edit: only assign to those objects that do not have the property owned by any client, yet.
 			TArray<TSoftObjectPtr<>> ObjectsToAssignTo;
-			Algo::TransformIf(OuterObjects, ObjectsToAssignTo,
-				[this, &PathInfo](const UObject* Object) { return IsPropertyNotYetAssigned(*Object, PathInfo.Path); },
-				[](UObject* Object){ return Object; }
-				);
+			Algo::Transform(OuterObjects, ObjectsToAssignTo, [](UObject* Object){ return Object; });
 			
 			MultiStreamColumns::FAssignPropertyModel::AssignPropertyTo(
 				*ReplicationManager->GetUnifiedClientView(),
@@ -202,6 +211,9 @@ namespace UE::MultiUserClient::Replication
 				ObjectsToAssignTo,
 				*PropertyChain
 				);
+
+			// To make the user aware of what they've just done, show the UI if it's closed.
+			SelectObjectsInReplicationUI(OuterObjects);
 		}
 	}
 
@@ -216,6 +228,31 @@ namespace UE::MultiUserClient::Replication
 		return bIsReplicatableProperty && bIsInReplicationSession;
 	}
 
+	void FDetailRowExtensionHandler::SelectObjectsInReplicationUI(TConstArrayView<UObject*> ObjectsToSelect) const
+	{
+		TArray<TSoftObjectPtr<>> ActorsToSelect;
+		Algo::Transform(ObjectsToSelect, ActorsToSelect, [](UObject* Object)
+		{
+			// IReplicationStreamViewer::SetSelectedObjects only allows selecting top-level objects, i.e. actors.
+			// Subobjects will be shown implicitly by the bottom panel,.
+			return Object->IsA<AActor>() ? Object : Object->GetTypedOuter<AActor>();
+		});
+		
+		const TSharedPtr<SConcertBrowser> Browser = GetOrInvokeBrowserTabDelegate.Execute();
+		if (!ensure(Browser))
+		{
+			return;
+		}
+		
+		const TSharedPtr<SActiveSessionRoot> SessionRoot = GetActiveSessionWidgetFromBrowser(*Browser);
+		const TSharedPtr<ConcertSharedSlate::IMultiReplicationStreamEditor> ReplicationEditor = GetReplicationStreamEditorWidgetFromBrowser(*Browser);
+		if (SessionRoot && ReplicationEditor)
+		{
+			SessionRoot->OpenTab(EMultiUserTab::Replication);
+			ReplicationEditor->GetEditorBase().SetSelectedObjects(ActorsToSelect);
+		}
+	}
+
 #define SET_REASON(X) if (OutReason){ *OutReason = X; }
 	bool FDetailRowExtensionHandler::CanAssignPropertyWithReason(
 		TSharedPtr<IPropertyHandle> PropertyHandle,
diff --git a/Engine/Plugins/Developer/Concert/ConcertApp/MultiUserClient/Source/MultiUserClient/Private/Replication/Misc/DetailRowExtensionHandler.h b/Engine/Plugins/Developer/Concert/ConcertApp/MultiUserClient/Source/MultiUserClient/Private/Replication/Misc/DetailRowExtensionHandler.h
--- a/Engine/Plugins/Developer/Concert/ConcertApp/MultiUserClient/Source/MultiUserClient/Private/Replication/Misc/DetailRowExtensionHandler.h
+++ b/Engine/Plugins/Developer/Concert/ConcertApp/MultiUserClient/Source/MultiUserClient/Private/Replication/Misc/DetailRowExtensionHandler.h
@@ -2,6 +2,8 @@
 
 #pragma once
 
+#include "Widgets/WidgetGetters.h"
+
 #include "Containers/ContainersFwd.h"
 #include "Templates/SharedPointer.h"
 #include "Templates/UnrealTemplate.h"
@@ -22,15 +24,21 @@ namespace UE::MultiUserClient::Replication
 	{
 	public:
 		
-		FDetailRowExtensionHandler(const TSharedRef<IConcertSyncClient>& InClient, const TSharedRef<FMultiUserReplicationManager>& InReplicationManager);
+		FDetailRowExtensionHandler(
+			const TSharedRef<IConcertSyncClient>& InClient,
+			const TSharedRef<FMultiUserReplicationManager>& InReplicationManager,
+			FGetConcertBrowserWidget InGetOrInvokeBrowserTabDelegate
+			);
 		~FDetailRowExtensionHandler();
 
 	private:
 
 		/** Used to look up client names. */
-		TWeakPtr<IConcertSyncClient> WeakClient;
+		const TWeakPtr<IConcertSyncClient> WeakClient;
 		/** Used to get the local client, to which properties will be assigned by default. */
-		TWeakPtr<FMultiUserReplicationManager> WeakReplicationManager;
+		const TWeakPtr<FMultiUserReplicationManager> WeakReplicationManager;
+		/** Gets the content of the Multi User browser tab invoking it if it is not open. */
+		const FGetConcertBrowserWidget GetOrInvokeBrowserTabDelegate;
 
 		/** Adds the assign button to valid property rows.  */
 		void RegisterExtensionHandler(
@@ -48,6 +56,9 @@ namespace UE::MultiUserClient::Replication
 		/** @return Whether the client is in a MU session and the property can ever be replicated. */
 		bool IsAssignPropertyButtonVisible(TSharedPtr<IPropertyHandle> PropertyHandle) const;
 
+		/** Shows the replication UI if hidden and selects Objects. */
+		void SelectObjectsInReplicationUI(TConstArrayView<UObject*> ObjectsToSelect) const;
+
 		/** @return Whether PropertyHandle can be assigned to its outer objects. */
 		bool CanAssignPropertyWithReason(TSharedPtr<IPropertyHandle> PropertyHandle, FText* OutReason = nullptr) const;
 		bool CanAssignProperty(TSharedPtr<IPropertyHandle> PropertyHandle) const { return CanAssignPropertyWithReason(PropertyHandle); }
diff --git a/Engine/Plugins/Developer/Concert/ConcertApp/MultiUserClient/Source/MultiUserClient/Private/Widgets/WidgetGetters.cpp b/Engine/Plugins/Developer/Concert/ConcertApp/MultiUserClient/Source/MultiUserClient/Private/Widgets/WidgetGetters.cpp
new file mode 100644
--- /dev/null
+++ b/Engine/Plugins/Developer/Concert/ConcertApp/MultiUserClient/Source/MultiUserClient/Private/Widgets/WidgetGetters.cpp
@@ -0,0 +1,32 @@
+﻿// Copyright Epic Games, Inc. All Rights Reserved.
+
+#include "WidgetGetters.h"
+
+#include "ActiveSession/Replication/Client/Multi/SAllClientsView.h"
+#include "ActiveSession/Replication/Client/Multi/SMultiClientView.h"
+#include "ActiveSession/Replication/Joined/SReplicationJoinedView.h"
+#include "ActiveSession/Replication/SReplicationRootWidget.h"
+#include "ActiveSession/SActiveSessionRoot.h"
+#include "SConcertBrowser.h"
+
+namespace UE::MultiUserClient
+{
+	TSharedPtr<SActiveSessionRoot> GetActiveSessionWidgetFromBrowser(const SConcertBrowser& Browser)
+	{
+		return Browser.GetActiveSessionWidget().Pin();
+	}
+	
+	TSharedPtr<ConcertSharedSlate::IMultiReplicationStreamEditor> GetReplicationStreamEditorWidgetFromBrowser(const SConcertBrowser& Browser)
+	{
+		const TSharedPtr<SActiveSessionRoot> Root = GetActiveSessionWidgetFromBrowser(Browser);
+		if (!Root)
+		{
+			return nullptr;
+		}
+
+		const TSharedPtr<Replication::SReplicationJoinedView> ConnectedView = Root->GetReplicationContent()->GetConnectedView().Pin();
+		return ConnectedView
+			? ConnectedView->GetAllClientsView()->GetClientView()->GetStreamEditor()
+			: nullptr;
+	}
+}
diff --git a/Engine/Plugins/Developer/Concert/ConcertApp/MultiUserClient/Source/MultiUserClient/Private/Widgets/WidgetGetters.h b/Engine/Plugins/Developer/Concert/ConcertApp/MultiUserClient/Source/MultiUserClient/Private/Widgets/WidgetGetters.h
new file mode 100644
--- /dev/null
+++ b/Engine/Plugins/Developer/Concert/ConcertApp/MultiUserClient/Source/MultiUserClient/Private/Widgets/WidgetGetters.h
@@ -0,0 +1,24 @@
+﻿// Copyright Epic Games, Inc. All Rights Reserved.
+
+#pragma once
+
+#include "Delegates/Delegate.h"
+
+class SConcertBrowser;
+
+namespace UE::ConcertSharedSlate { class IMultiReplicationStreamEditor; }
+
+namespace UE::MultiUserClient
+{
+	class SActiveSessionRoot;
+
+	DECLARE_DELEGATE_RetVal(TSharedPtr<SConcertBrowser>, FGetConcertBrowserWidget);
+	DECLARE_DELEGATE_RetVal(TSharedPtr<SActiveSessionRoot>, FGetActiveSessionWidget);
+	DECLARE_DELEGATE_RetVal(TSharedPtr<ConcertSharedSlate::IMultiReplicationStreamEditor>, FGetReplicationStreamEditorWidget);
+
+	/** Traverses the UI tree to get the SActiveSessionRoot. */
+	TSharedPtr<SActiveSessionRoot> GetActiveSessionWidgetFromBrowser(const SConcertBrowser& Browser);
+	
+	/** Traverses the UI tree to get the SActiveSessionRoot. */
+	TSharedPtr<ConcertSharedSlate::IMultiReplicationStreamEditor> GetReplicationStreamEditorWidgetFromBrowser(const SConcertBrowser& Browser);
+}
