Fix non-unity build issue.

diff --git a/Engine/Plugins/Developer/Concert/ConcertApp/MultiUserClient/Source/MultiUserClient/Private/Replication/Misc/PropertySelection/UserPropertySelector.cpp b/Engine/Plugins/Developer/Concert/ConcertApp/MultiUserClient/Source/MultiUserClient/Private/Replication/Misc/PropertySelection/UserPropertySelector.cpp
--- a/Engine/Plugins/Developer/Concert/ConcertApp/MultiUserClient/Source/MultiUserClient/Private/Replication/Misc/PropertySelection/UserPropertySelector.cpp
+++ b/Engine/Plugins/Developer/Concert/ConcertApp/MultiUserClient/Source/MultiUserClient/Private/Replication/Misc/PropertySelection/UserPropertySelector.cpp
@@ -3,12 +3,12 @@
 #include "UserPropertySelector.h"
 
 #include "UserPropertySelectionSource.h"
-#include "Replication/ClientReplicationWidgetFactories.h"
+#include "Replication/Client/Offline/OfflineClientManager.h"
 #include "Replication/Client/Online/OnlineClientManager.h"
 #include "Replication/Editor/Model/IEditableReplicationStreamModel.h"
+#include "Replication/ReplicationWidgetFactories.h"
 
 #include "ScopedTransaction.h"
-#include "Replication/Client/Offline/OfflineClientManager.h"
 #include "UObject/Package.h"
 #include "UObject/UObjectGlobals.h"
 
