Increase width of Details panel right column when joining MU so it can fit the button for adding a property to replication.

diff --git a/Engine/Plugins/Developer/Concert/ConcertApp/MultiUserClient/Source/MultiUserClient/Private/Replication/Misc/DetailRowExtensionHandler.cpp b/Engine/Plugins/Developer/Concert/ConcertApp/MultiUserClient/Source/MultiUserClient/Private/Replication/Misc/DetailRowExtensionHandler.cpp
--- a/Engine/Plugins/Developer/Concert/ConcertApp/MultiUserClient/Source/MultiUserClient/Private/Replication/Misc/DetailRowExtensionHandler.cpp
+++ b/Engine/Plugins/Developer/Concert/ConcertApp/MultiUserClient/Source/MultiUserClient/Private/Replication/Misc/DetailRowExtensionHandler.cpp
@@ -3,6 +3,7 @@
 #include "DetailRowExtensionHandler.h"
 
 #include "IConcertSyncClient.h"
+#include "IDetailTreeNode.h"
 #include "MultiUserReplicationStyle.h"
 #include "Replication/Client/ClientUtils.h"
 #include "Replication/Editor/View/IMultiReplicationStreamEditor.h"
@@ -116,6 +117,29 @@ namespace UE::MultiUserClient::Replication::DetailRowExtensionHandler
 		
 		return Result;
 	}
+	
+	static void SetMinimumRightColumnWidth(bool bIsConnectedToReplication)
+	{
+		const FName DetailsTabIdentifiers[] =
+			{ "LevelEditorSelectionDetails", "LevelEditorSelectionDetails2", "LevelEditorSelectionDetails3", "LevelEditorSelectionDetails4" };
+		
+		FPropertyEditorModule& PropertyEditor = FModuleManager::LoadModuleChecked<FPropertyEditorModule>("PropertyEditor");
+		for (const FName& DetailsTabIdentifier : DetailsTabIdentifiers)
+		{
+			if (const TSharedPtr<IDetailsView> DetailsView = PropertyEditor.FindDetailView(DetailsTabIdentifier))
+			{
+				const float DeltaChang = bIsConnectedToReplication ? 22 : -22;
+				DetailsView->SetRightColumnMinWidth(DetailsView->GetRightColumnMinWidth() + DeltaChang);
+				DetailsView->ForceRefresh();
+			}
+		}
+	}
+	
+	static void OnConnectionStateChanged(EMultiUserReplicationConnectionState NewState)
+	{
+		const bool bIsConnected = NewState == EMultiUserReplicationConnectionState::Connected;
+		SetMinimumRightColumnWidth(bIsConnected);
+	}
 }
 
 namespace UE::MultiUserClient::Replication
@@ -128,6 +152,10 @@ namespace UE::MultiUserClient::Replication
 		: WeakClient(InClient)
 		, WeakReplicationManager(InReplicationManager)
 		, GetOrInvokeBrowserTabDelegate(MoveTemp(InGetOrInvokeBrowserTabDelegate))
+		, OnReplicationConnectionChangedHandle(
+			// Without this, the right column will be too small to host all icons (Reset to Default, MU, Sequencer, etc.) - so make it wider if needed.
+			InReplicationManager->OnReplicationConnectionStateChanged().AddStatic(&DetailRowExtensionHandler::OnConnectionStateChanged)
+			)
 	{
 		check(GetOrInvokeBrowserTabDelegate.IsBound());
 		
@@ -143,6 +171,11 @@ namespace UE::MultiUserClient::Replication
 		{
 			EditModule->GetGlobalRowExtensionDelegate().RemoveAll(this);
 		}
+
+		if (const TSharedPtr<FMultiUserReplicationManager> ReplicationManager = WeakReplicationManager.Pin())
+		{
+			ReplicationManager->OnReplicationConnectionStateChanged().Remove(OnReplicationConnectionChangedHandle);
+		}
 	}
 
 	void FDetailRowExtensionHandler::RegisterExtensionHandler(
@@ -221,11 +254,13 @@ namespace UE::MultiUserClient::Replication
 	{
 		const FProperty* Property = PropertyHandle->GetProperty();
 		const bool bIsReplicatableProperty = Property && ConcertSyncCore::PropertyChain::IsReplicatableProperty(*Property);
-		
+		return bIsReplicatableProperty && IsConnectedToReplication();
+	}
+
+	bool FDetailRowExtensionHandler::IsConnectedToReplication() const
+	{
 		const TSharedPtr<FMultiUserReplicationManager> ReplicationManager = WeakReplicationManager.Pin();
-		const bool bIsInReplicationSession = ensure(ReplicationManager)
-			&& ReplicationManager->GetConnectionState() == EMultiUserReplicationConnectionState::Connected;
-		return bIsReplicatableProperty && bIsInReplicationSession;
+		return ensure(ReplicationManager) && ReplicationManager->GetConnectionState() == EMultiUserReplicationConnectionState::Connected;
 	}
 
 	void FDetailRowExtensionHandler::SelectObjectsInReplicationUI(TConstArrayView<UObject*> ObjectsToSelect) const
diff --git a/Engine/Plugins/Developer/Concert/ConcertApp/MultiUserClient/Source/MultiUserClient/Private/Replication/Misc/DetailRowExtensionHandler.h b/Engine/Plugins/Developer/Concert/ConcertApp/MultiUserClient/Source/MultiUserClient/Private/Replication/Misc/DetailRowExtensionHandler.h
--- a/Engine/Plugins/Developer/Concert/ConcertApp/MultiUserClient/Source/MultiUserClient/Private/Replication/Misc/DetailRowExtensionHandler.h
+++ b/Engine/Plugins/Developer/Concert/ConcertApp/MultiUserClient/Source/MultiUserClient/Private/Replication/Misc/DetailRowExtensionHandler.h
@@ -39,6 +39,8 @@ namespace UE::MultiUserClient::Replication
 		const TWeakPtr<FMultiUserReplicationManager> WeakReplicationManager;
 		/** Gets the content of the Multi User browser tab invoking it if it is not open. */
 		const FGetConcertBrowserWidget GetOrInvokeBrowserTabDelegate;
+		/** Handle to FMultiUserReplicationManager::OnReplicationConnectionStateChanged. */
+		const FDelegateHandle OnReplicationConnectionChangedHandle;
 
 		/** Adds the assign button to valid property rows.  */
 		void RegisterExtensionHandler(
@@ -56,6 +58,9 @@ namespace UE::MultiUserClient::Replication
 		/** @return Whether the client is in a MU session and the property can ever be replicated. */
 		bool IsAssignPropertyButtonVisible(TSharedPtr<IPropertyHandle> PropertyHandle) const;
 
+		/** @return Whether in a MU session and completed the replication handshake. */
+		bool IsConnectedToReplication() const;
+
 		/** Shows the replication UI if hidden and selects Objects. */
 		void SelectObjectsInReplicationUI(TConstArrayView<UObject*> ObjectsToSelect) const;
 
