Fix: Saving map resets assigned replicated properties on different clients.
- When a remote user saves a map, Concert will hot reload the map on all other clients
- Hot reload creates a temporary map called "Untitled"
- During hot reload, ignore the delegates for changing maps

diff --git a/Engine/Plugins/Developer/Concert/ConcertApp/MultiUserClient/Source/MultiUserClient/Private/Replication/Misc/ChangeLevelHandler.cpp b/Engine/Plugins/Developer/Concert/ConcertApp/MultiUserClient/Source/MultiUserClient/Private/Replication/Misc/ChangeLevelHandler.cpp
--- a/Engine/Plugins/Developer/Concert/ConcertApp/MultiUserClient/Source/MultiUserClient/Private/Replication/Misc/ChangeLevelHandler.cpp
+++ b/Engine/Plugins/Developer/Concert/ConcertApp/MultiUserClient/Source/MultiUserClient/Private/Replication/Misc/ChangeLevelHandler.cpp
@@ -2,6 +2,8 @@
 
 #include "ChangeLevelHandler.h"
 
+#include "IConcertClientWorkspace.h"
+#include "IConcertSyncClient.h"
 #include "Replication/Editor/Model/IEditableReplicationStreamModel.h"
 
 #include "Engine/Engine.h"
@@ -9,8 +11,12 @@
 
 namespace UE::MultiUserClient::Replication
 {
-	FChangeLevelHandler::FChangeLevelHandler(ConcertSharedSlate::IEditableReplicationStreamModel& UpdatedModel)
-		: UpdatedModel(UpdatedModel)
+	FChangeLevelHandler::FChangeLevelHandler(
+		IConcertSyncClient& Client,
+		ConcertSharedSlate::IEditableReplicationStreamModel& UpdatedModel
+		)
+		: Client(Client)
+		, UpdatedModel(UpdatedModel)
 	{
 		if (ensure(GEngine))
 		{
@@ -30,19 +36,33 @@ namespace UE::MultiUserClient::Replication
 
 	void FChangeLevelHandler::OnWorldDestroyed(UWorld* World)
 	{
-		// Remember the current map so when the next world
-		if (World->WorldType == EWorldType::Editor)
+		if (IsValidWorldType(World) && !IsConcertHotReloadingWorld())
 		{
+			// Remember the current map so when the next world
 			PreviousWorldPath = FSoftObjectPath(World);
 		}
 	}
 
 	void FChangeLevelHandler::OnWorldAdded(UWorld* World) const
 	{
-		// User reloaded the same map - keep around all settings
-		if (PreviousWorldPath != World)
+		if (IsValidWorldType(World) && !IsConcertHotReloadingWorld()
+			// If user reloaded the same map - keep around all settings
+			&& PreviousWorldPath != World)
 		{
 			UpdatedModel.Clear();
 		}
 	}
-}
+
+	bool FChangeLevelHandler::IsConcertHotReloadingWorld() const
+	{
+		const TSharedPtr<IConcertClientWorkspace>& Workspace = Client.GetWorkspace();
+		const FName PackagePath = PreviousWorldPath.GetLongPackageFName();
+		const bool bIsReloadingWorld = Workspace && Workspace->IsReloadingPackage(PackagePath);
+		return bIsReloadingWorld;
+	}
+
+	bool FChangeLevelHandler::IsValidWorldType(UWorld* World) const
+	{
+		return World->WorldType == EWorldType::Editor;
+	}
+}
\ No newline at end of file
diff --git a/Engine/Plugins/Developer/Concert/ConcertApp/MultiUserClient/Source/MultiUserClient/Private/Replication/Misc/ChangeLevelHandler.h b/Engine/Plugins/Developer/Concert/ConcertApp/MultiUserClient/Source/MultiUserClient/Private/Replication/Misc/ChangeLevelHandler.h
--- a/Engine/Plugins/Developer/Concert/ConcertApp/MultiUserClient/Source/MultiUserClient/Private/Replication/Misc/ChangeLevelHandler.h
+++ b/Engine/Plugins/Developer/Concert/ConcertApp/MultiUserClient/Source/MultiUserClient/Private/Replication/Misc/ChangeLevelHandler.h
@@ -6,6 +6,7 @@
 #include "Templates/UnrealTemplate.h"
 #include "UObject/SoftObjectPath.h"
 
+class IConcertSyncClient;
 class UWorld;
 
 namespace UE::ConcertSharedSlate { class IEditableReplicationStreamModel; }
@@ -21,13 +22,16 @@ namespace UE::MultiUserClient::Replication
 	public:
 		
 		/**
+		 * @param Client The Concert client. Used to get the workspace for checking against hot reloading. The caller ensures that it outlives the constructed object.
 		 * @param UpdatedModel The client model to update when the local editor changes maps. The caller ensures that it outlives the constructed object.
 		 */
-		FChangeLevelHandler(ConcertSharedSlate::IEditableReplicationStreamModel& UpdatedModel UE_LIFETIMEBOUND);
+		FChangeLevelHandler(IConcertSyncClient& Client UE_LIFETIMEBOUND, ConcertSharedSlate::IEditableReplicationStreamModel& UpdatedModel UE_LIFETIMEBOUND);
 		~FChangeLevelHandler();
 
 	private:
-
+		
+		/** The Concert client. Used to get the workspace for checking against hot reloading. */
+		IConcertSyncClient& Client;
 		/** The client model to update when the local editor changes maps. */
 		ConcertSharedSlate::IEditableReplicationStreamModel& UpdatedModel;
 
@@ -36,5 +40,14 @@ namespace UE::MultiUserClient::Replication
 		
 		void OnWorldDestroyed(UWorld* World);
 		void OnWorldAdded(UWorld* World) const;
+
+		/**
+		 * Hot reload: When a remote user saves the world, Concert will reload the world's package on the other clients.
+		 * In that case, a temporary world called "Untitled" is created. During hot reload, OnWorldDestroyed and OnWorldAdded calls should be ignored.
+		 * 
+		 * @return Whether concert is hot reloading the currently open world.
+		 */
+		bool IsConcertHotReloadingWorld() const;
+		bool IsValidWorldType(UWorld* World) const;
 	};
 }
diff --git a/Engine/Plugins/Developer/Concert/ConcertApp/MultiUserClient/Source/MultiUserClient/Private/Replication/MultiUserReplicationManager.cpp b/Engine/Plugins/Developer/Concert/ConcertApp/MultiUserClient/Source/MultiUserClient/Private/Replication/MultiUserReplicationManager.cpp
--- a/Engine/Plugins/Developer/Concert/ConcertApp/MultiUserClient/Source/MultiUserClient/Private/Replication/MultiUserReplicationManager.cpp
+++ b/Engine/Plugins/Developer/Concert/ConcertApp/MultiUserClient/Source/MultiUserClient/Private/Replication/MultiUserReplicationManager.cpp
@@ -296,7 +296,7 @@ namespace UE::MultiUserClient::Replication
 		, PresetManager(*InClient, OnlineClientManager, MuteManager.GetSynchronizer())
 		, PropertySelector(OnlineClientManager, OfflineClientManager)
 		, AutoPropertyOwnershipTaker(PropertySelector, OnlineClientManager.GetLocalClient(), OnlineClientManager.GetAuthorityCache())
-		, ChangeLevelHandler(OnlineClientManager.GetLocalClient().GetClientEditModel().Get())
+		, ChangeLevelHandler(*InClient, OnlineClientManager.GetLocalClient().GetClientEditModel().Get())
 		, PreventReplicatedPropertyTransaction(*InClient, OnlineClientManager, MuteManager)
 		, UserNotifier(*InClient->GetConcertClient(), OnlineClientManager, MuteManager)
 		, AnalyticsHandler(*InClient->GetConcertClient(), OnlineClientManager)
