Undo / redo now correctly gives back authority to the client.
- The fix is to fully rebuild the array of added objects.

diff --git a/Engine/Plugins/Developer/Concert/ConcertApp/MultiUserClient/Source/MultiUserClient/Private/Replication/Client/Online/OnlineClient.cpp b/Engine/Plugins/Developer/Concert/ConcertApp/MultiUserClient/Source/MultiUserClient/Private/Replication/Client/Online/OnlineClient.cpp
--- a/Engine/Plugins/Developer/Concert/ConcertApp/MultiUserClient/Source/MultiUserClient/Private/Replication/Client/Online/OnlineClient.cpp
+++ b/Engine/Plugins/Developer/Concert/ConcertApp/MultiUserClient/Source/MultiUserClient/Private/Replication/Client/Online/OnlineClient.cpp
@@ -84,7 +84,7 @@ namespace UE::MultiUserClient::Replication
 		ConcertSharedSlate::EReplicatedObjectChangeReason ReplicatedObjectChangeReason
 		)
 	{
-		DeferOnModelChanged(AddedObjects);
+		DeferOnModelChanged();
 	}
 
 	void FOnlineClient::OnPropertiesChanged()
@@ -121,73 +121,71 @@ namespace UE::MultiUserClient::Replication
 		DeferOnModelChanged();
 	}
 
-	void FOnlineClient::DeferOnModelChanged(TConstArrayView<UObject*> AddedObjects)
+	void FOnlineClient::DeferOnModelChanged()
 	{
-		if (!DeferredOnModelChangedData)
+		if (!bHasDeferredDataChanges)
 		{
-			DeferredOnModelChangedData.Emplace();
+			bHasDeferredDataChanges = true;
 			FCoreDelegates::OnEndFrame.AddRaw(this, &FOnlineClient::ProcessOnModelChanged);
 		}
-
-		Algo::Transform(AddedObjects, DeferredOnModelChangedData->AccumulatedAddedObjects, [](UObject* Object)
-		{
-			return Object;
-		});
 	}
 
 	void FOnlineClient::ProcessOnModelChanged()
 	{
-		check(DeferredOnModelChangedData);
-		const FDeferredOnModelChangedData ChangeData = MoveTemp(*DeferredOnModelChangedData);
-		DeferredOnModelChangedData.Reset();
+		check(bHasDeferredDataChanges);
+		bHasDeferredDataChanges = false;
 		FCoreDelegates::OnEndFrame.RemoveAll(this);
 		
 		// Could improve performance by just considering what actually changed instead of doing a full rebuild
 		// This must be done before SetAuthorityIfAllowed because it uses the cache for checking whether the object has properties assigned
 		LocalClientStreamDiffer.RefreshChangesCache();
+		const TArray<UObject*> AddedObjects = DetermineAddedObjects();
 		
 		// Better UX for user: automatically take authority for newly added objects (but only if it is allowed and causes no conflicts)
-		TakeAuthorityOverNewlyAddedObjects(ChangeData);
+		TArray<FSoftObjectPath> ObjectPaths;
+		Algo::Transform(AddedObjects, ObjectPaths, [](UObject* Object){ return FSoftObjectPath(Object); });
+		LocalAuthorityDiffer.SetAuthorityIfAllowed(ObjectPaths, true);
+		
 		// Refresh because local authority changes may no longer be valid after modifying the stream
 		LocalAuthorityDiffer.RefreshChanges();
+		
 		// Check whether added objects are supposed to default to some special frequency settings
-		ApplyDefaultFrequencySettings(ChangeData);
+		ApplyDefaultFrequencySettings(AddedObjects);
 
 		// Finally, let everybody else know.
 		OnModelChangedDelegate.Broadcast();
 		AutoSubmissionPolicy.ProcessAccumulatedChangesAndSubmit();
 	}
 	
-	void FOnlineClient::TakeAuthorityOverNewlyAddedObjects(const FDeferredOnModelChangedData& ChangeData)
+	TArray<UObject*> FOnlineClient::DetermineAddedObjects() const
 	{
-		TArray<FSoftObjectPath> ObjectPaths;
-		Algo::TransformIf(ChangeData.AccumulatedAddedObjects, ObjectPaths,
-			[](const TWeakObjectPtr<UObject>& Object)
-			{
-				// The object might have been made invalid last frame.
-				return Object.IsValid();
-			},
-			[](const TWeakObjectPtr<UObject>& Object)
+		TArray<UObject*> AddedObjects;
+
+		const FGuid MultiUserStreamId = StreamSynchronizer->GetStreamId();
+		const FConcertObjectReplicationMap& CurrentState = StreamSynchronizer->GetServerState();
+		for (const TPair<FConcertObjectInStreamID, FConcertReplication_ChangeStream_PutObject>& Pair : LocalClientStreamDiffer.GetCachedDeltaChange().ObjectsToPut)
+		{
+			const FSoftObjectPath& ObjectPath = Pair.Key.Object;
+			const bool bIsNew = Pair.Key.StreamId == MultiUserStreamId && !CurrentState.ReplicatedObjects.Contains(ObjectPath);
+			UObject* Object = bIsNew ? ObjectPath.ResolveObject() : nullptr;
+			if (Object)
 			{
-				return FSoftObjectPath(Object.Get()) ;
-			});
-		LocalAuthorityDiffer.SetAuthorityIfAllowed(ObjectPaths, true);
+				AddedObjects.Add(Object);
+			}
+		}
+
+		return AddedObjects;
 	}
 
-	void FOnlineClient::ApplyDefaultFrequencySettings(const FDeferredOnModelChangedData& ChangeData)
+	void FOnlineClient::ApplyDefaultFrequencySettings(const TArray<UObject*>& AddedObjects)
 	{
 		UMultiUserReplicationSettings* Settings = UMultiUserReplicationSettings::Get();
-		for (const TWeakObjectPtr<UObject>& AddedObject : ChangeData.AccumulatedAddedObjects)
+		for (UObject* AddedObject : AddedObjects)
 		{
-			if (!AddedObject.IsValid())
-			{
-				continue;
-			}
-			
 			const TOptional<FConcertObjectReplicationSettings> FrequencyOverride = Settings->DetermineObjectFrequencySettings(*AddedObject);
 			if (FrequencyOverride)
 			{
-				LocalFrequencyChangeTracker.AddOverride(AddedObject.Get(), *FrequencyOverride);
+				LocalFrequencyChangeTracker.AddOverride(AddedObject, *FrequencyOverride);
 			}
 		}
 	}
diff --git a/Engine/Plugins/Developer/Concert/ConcertApp/MultiUserClient/Source/MultiUserClient/Private/Replication/Client/Online/OnlineClient.h b/Engine/Plugins/Developer/Concert/ConcertApp/MultiUserClient/Source/MultiUserClient/Private/Replication/Client/Online/OnlineClient.h
--- a/Engine/Plugins/Developer/Concert/ConcertApp/MultiUserClient/Source/MultiUserClient/Private/Replication/Client/Online/OnlineClient.h
+++ b/Engine/Plugins/Developer/Concert/ConcertApp/MultiUserClient/Source/MultiUserClient/Private/Replication/Client/Online/OnlineClient.h
@@ -168,11 +168,8 @@ namespace UE::MultiUserClient::Replication
 		 */
 		FOnModelExternallyChanged OnModelChangedDelegate;
 
-		struct FDeferredOnModelChangedData
-		{
-			TSet<TWeakObjectPtr<UObject>> AccumulatedAddedObjects;
-		};
-		TOptional<FDeferredOnModelChangedData> DeferredOnModelChangedData;
+		/** Whether there are any deferred data changes to process. */
+		bool bHasDeferredDataChanges = false;
 
 		// Respond to model changing
 		void OnObjectsChanged(TConstArrayView<UObject*> AddedObjects, TConstArrayView<FSoftObjectPath> RemovedObjects, ConcertSharedSlate::EReplicatedObjectChangeReason ReplicatedObjectChangeReason);
@@ -181,16 +178,15 @@ namespace UE::MultiUserClient::Replication
 		void OnServerStateChanged();
 
 		/** Defers rebuilding operations, such as refreshing state and calling OnModelChanged, in case there are multiple changes in the same frame. */
-		void DeferOnModelChanged() { DeferOnModelChanged({}); }
-		void DeferOnModelChanged(TConstArrayView<UObject*> AddedObjects);
+		void DeferOnModelChanged();
 		/** Processes all changes that have happened to the stream this frame. */
 		void ProcessOnModelChanged();
 		
-		/** Takes authority over newly added objects for better UX */
-		void TakeAuthorityOverNewlyAddedObjects(const FDeferredOnModelChangedData& ChangeData);
+		/** Determines the objects that were newly added. */
+		TArray<UObject*> DetermineAddedObjects() const;
 
 		/** Applies default replication frequency settings specified in the Multi User settings */
-		void ApplyDefaultFrequencySettings(const FDeferredOnModelChangedData& ChangeData);
+		void ApplyDefaultFrequencySettings(const TArray<UObject*>& ChangeData);
 		
 		/** Removes authority if request fails */
 		void OnAuthoritySubmissionCompleted(const FSubmitAuthorityChangesRequest& Request, const FSubmitAuthorityChangesResponse& Response);
