Warn the user if they try to persist the sandbox in a state other than the default send/receive mode.

diff --git a/Engine/Plugins/Developer/Concert/ConcertApp/MultiUserClient/Source/MultiUserClient/Private/ConcertWorkspaceUI.cpp b/Engine/Plugins/Developer/Concert/ConcertApp/MultiUserClient/Source/MultiUserClient/Private/ConcertWorkspaceUI.cpp
--- a/Engine/Plugins/Developer/Concert/ConcertApp/MultiUserClient/Source/MultiUserClient/Private/ConcertWorkspaceUI.cpp
+++ b/Engine/Plugins/Developer/Concert/ConcertApp/MultiUserClient/Source/MultiUserClient/Private/ConcertWorkspaceUI.cpp
@@ -2,6 +2,7 @@
 
 #include "ConcertWorkspaceUI.h"
 
+#include "ConcertMessages.h"
 #include "IConcertSyncClientModule.h"
 #include "IConcertClientWorkspace.h"
 #include "IConcertSyncClient.h"
@@ -11,6 +12,8 @@
 #include "ConcertLogGlobal.h"
 #include "ContentBrowserMenuContexts.h"
 
+#include "Misc/MessageDialog.h"
+
 #include "Algo/Transform.h"
 #include "Modules/ModuleManager.h"
 #include "ContentBrowserModule.h"
@@ -612,12 +615,46 @@ bool FConcertWorkspaceUI::HasSessionChanges() const
 	return ClientWorkspacePin.IsValid() && ClientWorkspacePin->HasSessionChanges();
 }
 
+namespace UE::ConcertWorkspaceUI::Private
+{
+bool ShouldContinueWithSendReceiveState(EConcertSendReceiveState SendReceiveState)
+{
+	if (SendReceiveState == EConcertSendReceiveState::Default)
+	{
+		return true;
+	}
+
+	FText WarningMsg;
+	if (SendReceiveState == EConcertSendReceiveState::SendOnly)
+	{
+		WarningMsg = LOCTEXT("PersistInNonDefaultState", "You are about to persist a session in send only mode.  Changes from other clients may not be included.  Continue?");
+	}
+	else if (SendReceiveState == EConcertSendReceiveState::ReceiveOnly)
+	{
+		WarningMsg = LOCTEXT("PersistInNonDefaultState", "You are about to persist a session in receive only mode.  Your changes have not been sent to other clients.  Continue?");
+	}
+	else
+	{
+		checkNoEntry();
+	}
+	EAppReturnType::Type RetType = FMessageDialog::Open(EAppMsgType::YesNo, EAppReturnType::No, WarningMsg);
+	
+	return RetType == EAppReturnType::Yes;
+}
+}
+
 bool FConcertWorkspaceUI::PromptPersistSessionChanges()
 {
 	TSharedPtr<IConcertClientWorkspace> ClientWorkspacePin = ClientWorkspace.Pin();
 	TArray<TSharedPtr<FConcertPersistItem>> PersistItems;
 	if (ClientWorkspacePin.IsValid())
 	{
+		IConcertClientSession& Session = ClientWorkspacePin->GetSession();
+		if (!UE::ConcertWorkspaceUI::Private::ShouldContinueWithSendReceiveState(Session.GetSendReceiveState()))
+		{
+			return false;
+		}
+
 		// Get source control status of packages 
 		TArray<FString> PackageFilenames;
 		TArray<FName> PackageNames;
