Correct localization tags.

diff --git a/Engine/Plugins/Developer/Concert/ConcertApp/MultiUserClient/Source/MultiUserClient/Private/ConcertWorkspaceUI.cpp b/Engine/Plugins/Developer/Concert/ConcertApp/MultiUserClient/Source/MultiUserClient/Private/ConcertWorkspaceUI.cpp
--- a/Engine/Plugins/Developer/Concert/ConcertApp/MultiUserClient/Source/MultiUserClient/Private/ConcertWorkspaceUI.cpp
+++ b/Engine/Plugins/Developer/Concert/ConcertApp/MultiUserClient/Source/MultiUserClient/Private/ConcertWorkspaceUI.cpp
@@ -627,11 +627,11 @@ bool ShouldContinueWithSendReceiveState(EConcertSendReceiveState SendReceiveStat
 	FText WarningMsg;
 	if (SendReceiveState == EConcertSendReceiveState::SendOnly)
 	{
-		WarningMsg = LOCTEXT("PersistInNonDefaultState", "You are about to persist a session in send only mode.  Changes from other clients may not be included.  Continue?");
+		WarningMsg = LOCTEXT("PersistInNonDefaultStateSendOnly", "You are about to persist a session in send only mode.  Changes from other clients may not be included.  Continue?");
 	}
 	else if (SendReceiveState == EConcertSendReceiveState::ReceiveOnly)
 	{
-		WarningMsg = LOCTEXT("PersistInNonDefaultState", "You are about to persist a session in receive only mode.  Your changes have not been sent to other clients.  Continue?");
+		WarningMsg = LOCTEXT("PersistInNonDefaultStateReceiveOnly", "You are about to persist a session in receive only mode.  Your changes have not been sent to other clients.  Continue?");
 	}
 	else
 	{
