Enable MultiUser.AllowTransactionsOnReplicatedProperties by default. This will prevent transactions on replicated objects.

diff --git a/Engine/Plugins/Developer/Concert/ConcertApp/MultiUserClient/Source/MultiUserClient/Private/Replication/Misc/PreventReplicatedPropertyTransaction.cpp b/Engine/Plugins/Developer/Concert/ConcertApp/MultiUserClient/Source/MultiUserClient/Private/Replication/Misc/PreventReplicatedPropertyTransaction.cpp
--- a/Engine/Plugins/Developer/Concert/ConcertApp/MultiUserClient/Source/MultiUserClient/Private/Replication/Misc/PreventReplicatedPropertyTransaction.cpp
+++ b/Engine/Plugins/Developer/Concert/ConcertApp/MultiUserClient/Source/MultiUserClient/Private/Replication/Misc/PreventReplicatedPropertyTransaction.cpp
@@ -14,8 +14,8 @@
 namespace UE::MultiUserClient::Replication
 {
 	TAutoConsoleVariable<bool> CVarAllowTransactionsOnReplicatedProperties(
-		TEXT("MultiUser.AllowTransactionsOnReplicatedproperties"),
-		false,
+		TEXT("MultiUser.AllowTransactionsOnReplicatedProperties"),
+		true,
 		TEXT("Whether to stop disallowing transactions on replicated properties."),
 		ECVF_Default
 		);
