Disable replicated properties transacting

diff --git a/Engine/Plugins/Developer/Concert/ConcertApp/MultiUserClient/Source/MultiUserClient/Private/Replication/Misc/PreventReplicatedPropertyTransaction.cpp b/Engine/Plugins/Developer/Concert/ConcertApp/MultiUserClient/Source/MultiUserClient/Private/Replication/Misc/PreventReplicatedPropertyTransaction.cpp
--- a/Engine/Plugins/Developer/Concert/ConcertApp/MultiUserClient/Source/MultiUserClient/Private/Replication/Misc/PreventReplicatedPropertyTransaction.cpp
+++ b/Engine/Plugins/Developer/Concert/ConcertApp/MultiUserClient/Source/MultiUserClient/Private/Replication/Misc/PreventReplicatedPropertyTransaction.cpp
@@ -15,7 +15,7 @@ namespace UE::MultiUserClient::Replication
 {
 	TAutoConsoleVariable<bool> CVarAllowTransactionsOnReplicatedProperties(
 		TEXT("MultiUser.AllowTransactionsOnReplicatedProperties"),
-		true,
+		false,
 		TEXT("Whether to stop disallowing transactions on replicated properties."),
 		ECVF_Default
 		);
