Check all exe pathnames when shutting down MU server.

#rnx

diff --git a/Engine/Plugins/Developer/Concert/ConcertApp/MultiUserClient/Source/MultiUserClient/Private/MultiUserClientModule.cpp b/Engine/Plugins/Developer/Concert/ConcertApp/MultiUserClient/Source/MultiUserClient/Private/MultiUserClientModule.cpp
--- a/Engine/Plugins/Developer/Concert/ConcertApp/MultiUserClient/Source/MultiUserClient/Private/MultiUserClientModule.cpp
+++ b/Engine/Plugins/Developer/Concert/ConcertApp/MultiUserClient/Source/MultiUserClient/Private/MultiUserClientModule.cpp
@@ -1181,9 +1181,15 @@ private:
 
 	void ShutdownConcertServer() override
 	{
-		FString ServerPath = GetConfiguredMultiUserServerExePathname();
-		ServerPath = FPaths::ConvertRelativePathToFull(ServerPath);
-		FPaths::NormalizeDirectoryName(ServerPath); // Need to have all slashes the same side.
+		auto GetCleanPath = [](EConcertServerType Type) -> FString
+		{
+			FString ServerPath = GetMultiUserServerExePathname(Type);
+			ServerPath = FPaths::ConvertRelativePathToFull(ServerPath);
+			FPaths::NormalizeDirectoryName(ServerPath);
+			return ServerPath;
+		};
+		const FString ConsoleServerPath = GetCleanPath(EConcertServerType::Console);
+		const FString SlateServerPath = GetCleanPath(EConcertServerType::Slate);
 
 		FPlatformProcess::FProcEnumerator ProcIter;
 		while (ProcIter.MoveNext())
@@ -1191,7 +1197,7 @@ private:
 			FPlatformProcess::FProcEnumInfo ProcInfo = ProcIter.GetCurrent();
 			FString Candidate = ProcInfo.GetFullPath();
 			FPaths::NormalizeDirectoryName(Candidate);
-			if (Candidate == ServerPath)
+			if (Candidate == ConsoleServerPath || Candidate == SlateServerPath)
 			{
 				MultiUserClientUtil::KillProcess(ProcInfo.GetPID());
 				// Continue -> Will close all servers instances.
