Add a command line flag to bypass source control checks when joining a Multi-user session.  The new command line arg is  -CONCERTSKIPSOURCECONTROLCHECKS

diff --git a/Engine/Plugins/Developer/Concert/ConcertApp/MultiUserClient/Source/MultiUserClient/Private/MultiUserClientModule.cpp b/Engine/Plugins/Developer/Concert/ConcertApp/MultiUserClient/Source/MultiUserClient/Private/MultiUserClientModule.cpp
--- a/Engine/Plugins/Developer/Concert/ConcertApp/MultiUserClient/Source/MultiUserClient/Private/MultiUserClientModule.cpp
+++ b/Engine/Plugins/Developer/Concert/ConcertApp/MultiUserClient/Source/MultiUserClient/Private/MultiUserClientModule.cpp
@@ -932,7 +932,18 @@ private:
 
 	void GetPreConnectionTasks(const IConcertClient& InClient, TArray<TUniquePtr<IConcertClientConnectionTask>>& OutTasks)
 	{
-		OutTasks.Emplace(MakeUnique<FConcertClientConnectionSourceControlValidationTask>(InClient.GetConfiguration()));
+		const TCHAR* CommandLine = FCommandLine::Get();
+		bool bSkipSCChecks = FParse::Param(CommandLine, TEXT("-CONCERTSKIPSOURCECONTROLCHECKS"));
+		FParse::Bool(CommandLine, TEXT("-CONCERTSKIPSOURCECONTROLCHECKS="), bSkipSCChecks);
+
+		if (!bSkipSCChecks)
+		{
+			OutTasks.Emplace(MakeUnique<FConcertClientConnectionSourceControlValidationTask>(InClient.GetConfiguration()));
+		}
+		else
+		{
+			UE_LOG(LogConcert, Display, TEXT("Source control validation task was bypassed."));
+		}
 		OutTasks.Emplace(MakeUnique<FConcertClientConnectionDirtyPackagesValidationTask>(InClient.GetConfiguration()));
 	}
 
