Multi User: Introduce MultiUser.Replication.LogRemoteChangeRequestsAndResponses to help debugging Concert Replication.

diff --git a/Engine/Plugins/Developer/Concert/ConcertApp/MultiUserClient/Source/MultiUserClient/Private/Replication/MultiUserReplicationCVars.cpp b/Engine/Plugins/Developer/Concert/ConcertApp/MultiUserClient/Source/MultiUserClient/Private/Replication/MultiUserReplicationCVars.cpp
new file mode 100644
--- /dev/null
+++ b/Engine/Plugins/Developer/Concert/ConcertApp/MultiUserClient/Source/MultiUserClient/Private/Replication/MultiUserReplicationCVars.cpp
@@ -0,0 +1,12 @@
+﻿// Copyright Epic Games, Inc. All Rights Reserved.
+
+#include "Replication/MultiUserReplicationCVars.h"
+
+namespace UE::MultiUserClient::Replication
+{
+TAutoConsoleVariable<bool> CVarLogRemoteChangeRequestsAndResponses(
+	TEXT("MultiUser.Replication.LogRemoteChangeRequestsAndResponses"),
+	false,
+	TEXT("Whether to log FMultiUser_ChangeRemote_Request & FMultiUser_ChangeRemote_Response received from other MU clients.")
+	);
+}
diff --git a/Engine/Plugins/Developer/Concert/ConcertApp/MultiUserClient/Source/MultiUserClient/Private/Replication/MultiUserReplicationCVars.h b/Engine/Plugins/Developer/Concert/ConcertApp/MultiUserClient/Source/MultiUserClient/Private/Replication/MultiUserReplicationCVars.h
new file mode 100644
--- /dev/null
+++ b/Engine/Plugins/Developer/Concert/ConcertApp/MultiUserClient/Source/MultiUserClient/Private/Replication/MultiUserReplicationCVars.h
@@ -0,0 +1,25 @@
+﻿// Copyright Epic Games, Inc. All Rights Reserved.
+
+#pragma once
+
+#include "ConcertLogGlobal.h"
+#include "HAL/IConsoleManager.h"
+#include "JsonObjectConverter.h"
+
+namespace UE::MultiUserClient::Replication
+{
+/** Whether to log FMultiUser_ChangeRemote_Request & FMultiUser_ChangeRemote_Response received from other MU clients. */
+extern TAutoConsoleVariable<bool> CVarLogRemoteChangeRequestsAndResponses;
+
+/** Logs a Concert network message, which is USTRUCT(). */
+template<typename TMessage>
+static void LogNetworkMessage(const TAutoConsoleVariable<bool>& ShouldLog, const TMessage& Message)
+{
+	if (ShouldLog.GetValueOnAnyThread())
+	{
+		FString JsonString;
+		FJsonObjectConverter::UStructToJsonObjectString(TMessage::StaticStruct(), &Message, JsonString, 0, 0);
+		UE_LOG(LogConcert, Log, TEXT("%s\n%s"), *TMessage::StaticStruct()->GetName(), *JsonString);
+	}
+}
+}
diff --git a/Engine/Plugins/Developer/Concert/ConcertApp/MultiUserClient/Source/MultiUserClient/Private/Replication/Submission/Remote/RemoteSubmissionListener.cpp b/Engine/Plugins/Developer/Concert/ConcertApp/MultiUserClient/Source/MultiUserClient/Private/Replication/Submission/Remote/RemoteSubmissionListener.cpp
--- a/Engine/Plugins/Developer/Concert/ConcertApp/MultiUserClient/Source/MultiUserClient/Private/Replication/Submission/Remote/RemoteSubmissionListener.cpp
+++ b/Engine/Plugins/Developer/Concert/ConcertApp/MultiUserClient/Source/MultiUserClient/Private/Replication/Submission/Remote/RemoteSubmissionListener.cpp
@@ -5,6 +5,7 @@
 #include "ConcertLogGlobal.h"
 #include "IConcertSession.h"
 #include "RemoteSubmissionMessages.h"
+#include "Replication/MultiUserReplicationCVars.h"
 #include "Replication/Client/ClientUtils.h"
 #include "Replication/Stream/IClientStreamSynchronizer.h"
 #include "Replication/Submission/ISubmissionOperation.h"
@@ -38,6 +39,8 @@ namespace UE::MultiUserClient::Replication
 		)
 	{
 		UE_LOG(LogConcert, Log, TEXT("Received remote change request from %s"), *ClientUtils::GetClientDisplayName(*ConcertSession, Context.SourceEndpointId));
+		LogNetworkMessage(CVarLogRemoteChangeRequestsAndResponses, Request);
+		ON_SCOPE_EXIT{ LogNetworkMessage(CVarLogRemoteChangeRequestsAndResponses, Response); };
 		Response.bTimeout = false;
 
 		if (InProgressOperation.IsSet())
