Multi User: Fix typo in tooltip

diff --git a/Engine/Plugins/Developer/Concert/ConcertApp/MultiUserClient/Source/MultiUserClient/Private/Widgets/ActiveSession/Replication/Client/Multi/Columns/AssignProperty/SAssignPropertyComboBox.cpp b/Engine/Plugins/Developer/Concert/ConcertApp/MultiUserClient/Source/MultiUserClient/Private/Widgets/ActiveSession/Replication/Client/Multi/Columns/AssignProperty/SAssignPropertyComboBox.cpp
--- a/Engine/Plugins/Developer/Concert/ConcertApp/MultiUserClient/Source/MultiUserClient/Private/Widgets/ActiveSession/Replication/Client/Multi/Columns/AssignProperty/SAssignPropertyComboBox.cpp
+++ b/Engine/Plugins/Developer/Concert/ConcertApp/MultiUserClient/Source/MultiUserClient/Private/Widgets/ActiveSession/Replication/Client/Multi/Columns/AssignProperty/SAssignPropertyComboBox.cpp
@@ -226,7 +226,7 @@ namespace UE::MultiUserClient::Replication::MultiStreamColumns
 		}
 
 		return FText::Format(
-			LOCTEXT("AssignProperty.ToolTip.Normal", "The {0}|plural(one=client,other=clients) that {0}|plural(one=is,other=are) have registered to replicate this property."),
+			LOCTEXT("AssignProperty.ToolTip.Normal", "The {0}|plural(one=client,other=clients) that {0}|plural(one=is,other=are) registered to replicate this property."),
 			NumClients
 			);
 	}
