If we are auto connecting and the user asked to disconnect allow the disconnect to happen by stopping the auto connect task.

#rnx

diff --git a/Engine/Plugins/Developer/Concert/ConcertApp/MultiUserClient/Source/MultiUserClient/Private/MultiUserClientModule.cpp b/Engine/Plugins/Developer/Concert/ConcertApp/MultiUserClient/Source/MultiUserClient/Private/MultiUserClientModule.cpp
--- a/Engine/Plugins/Developer/Concert/ConcertApp/MultiUserClient/Source/MultiUserClient/Private/MultiUserClientModule.cpp
+++ b/Engine/Plugins/Developer/Concert/ConcertApp/MultiUserClient/Source/MultiUserClient/Private/MultiUserClientModule.cpp
@@ -752,6 +752,10 @@ public:
 				ConcertClient->DisconnectSession();
 			}
 		}
+		else if (ConcertClient->IsAutoConnecting())
+		{
+			ConcertClient->StopAutoConnect();
+		}
 
 		return true; // The client is disconnected.
 	}
