Multi User: Fix missing icons in persist menu.

This must bug have been a typo, and it must have been an issue since 2018.

diff --git a/Engine/Plugins/Developer/Concert/ConcertApp/MultiUserClient/Source/MultiUserClient/Private/Widgets/SConcertSandboxPersistWidget.cpp b/Engine/Plugins/Developer/Concert/ConcertApp/MultiUserClient/Source/MultiUserClient/Private/Widgets/SConcertSandboxPersistWidget.cpp
--- a/Engine/Plugins/Developer/Concert/ConcertApp/MultiUserClient/Source/MultiUserClient/Private/Widgets/SConcertSandboxPersistWidget.cpp
+++ b/Engine/Plugins/Developer/Concert/ConcertApp/MultiUserClient/Source/MultiUserClient/Private/Widgets/SConcertSandboxPersistWidget.cpp
@@ -70,8 +70,8 @@ public:
 				.VAlign(VAlign_Center)
 				[
 					SNew(SImage)
-					.Image(FAppStyle::GetBrush(Item->GetIconName()))
-				.ToolTipText(Item->GetIconTooltip())
+					.Image(Item->GetIcon().GetIcon())
+					.ToolTipText(Item->GetIconTooltip())
 				];
 		}
 		else if (ColumnID == SConcertSandboxPersistWidgetDefs::ColumnID_FileLabel)
diff --git a/Engine/Plugins/Developer/Concert/ConcertApp/MultiUserClient/Source/MultiUserClient/Private/Widgets/SConcertSandboxPersistWidget.h b/Engine/Plugins/Developer/Concert/ConcertApp/MultiUserClient/Source/MultiUserClient/Private/Widgets/SConcertSandboxPersistWidget.h
--- a/Engine/Plugins/Developer/Concert/ConcertApp/MultiUserClient/Source/MultiUserClient/Private/Widgets/SConcertSandboxPersistWidget.h
+++ b/Engine/Plugins/Developer/Concert/ConcertApp/MultiUserClient/Source/MultiUserClient/Private/Widgets/SConcertSandboxPersistWidget.h
@@ -49,7 +49,10 @@ public:
 	FText GetDisplayName() const { return DisplayName; }
 
 	/** Returns the name of the icon to be used in the list item widget */
-	FName GetIconName() const { return Item->GetIcon().GetStyleSetName(); }
+	FName GetIconName() const { return Item->GetIcon().GetStyleName(); }
+	
+	/** Returns the icon to be used in the list item widget */
+	FSlateIcon GetIcon() const { return Item->GetIcon(); }
 
 	/** Returns the tooltip text for the icon */
 	FText GetIconTooltip() const { return Item->GetDisplayTooltip(); }
