Added support for invalidation on content browser thumbnail extensions

diff --git a/Engine/Plugins/Developer/Concert/ConcertApp/MultiUserClient/Source/MultiUserClient/Private/ConcertWorkspaceUI.cpp b/Engine/Plugins/Developer/Concert/ConcertApp/MultiUserClient/Source/MultiUserClient/Private/ConcertWorkspaceUI.cpp
--- a/Engine/Plugins/Developer/Concert/ConcertApp/MultiUserClient/Source/MultiUserClient/Private/ConcertWorkspaceUI.cpp
+++ b/Engine/Plugins/Developer/Concert/ConcertApp/MultiUserClient/Source/MultiUserClient/Private/ConcertWorkspaceUI.cpp
@@ -59,6 +59,12 @@ static const FName ConcertHistoryTabName(TEXT("ConcertHistory"));
 // an asset is locked or modified by somebody else.
 //-----------------------------------------------------------------------------
 
+template <>
+struct TWidgetTypeTraits<class SConcertWorkspaceLockStateIndicator>
+{
+	static constexpr bool SupportsInvalidation() { return true; }
+};
+
 /**
  * Controls the appearance of the Concert workspace lock state icon. The lock state icon is displayed
  * on an asset in the editor content browser seen when a Concert user saves an assets or explicitly locks
diff --git a/Engine/Plugins/Enterprise/DatasmithImporter/Source/DirectLinkExtensionEditor/Private/SDirectLinkAssetThumbnailIndicator.cpp b/Engine/Plugins/Enterprise/DatasmithImporter/Source/DirectLinkExtensionEditor/Private/SDirectLinkAssetThumbnailIndicator.cpp
--- a/Engine/Plugins/Enterprise/DatasmithImporter/Source/DirectLinkExtensionEditor/Private/SDirectLinkAssetThumbnailIndicator.cpp
+++ b/Engine/Plugins/Enterprise/DatasmithImporter/Source/DirectLinkExtensionEditor/Private/SDirectLinkAssetThumbnailIndicator.cpp
@@ -15,6 +15,12 @@
 
 #define LOCTEXT_NAMESPACE "DirectLinkAssetThumbnailIndicator"
 
+template <>
+struct TWidgetTypeTraits<class UE::DatasmithImporter::SDirectLinkAssetThumbnailIndicator>
+{
+	static constexpr bool SupportsInvalidation() { return true; }
+};
+
 namespace UE::DatasmithImporter
 {
 	namespace ThumbnailIndicatorImpl
