Slate Dynamic Invalidation - Content Browser

diff --git a/Engine/Plugins/Developer/Concert/ConcertApp/MultiUserClient/Source/MultiUserClient/Private/ConcertWorkspaceUI.cpp b/Engine/Plugins/Developer/Concert/ConcertApp/MultiUserClient/Source/MultiUserClient/Private/ConcertWorkspaceUI.cpp
--- a/Engine/Plugins/Developer/Concert/ConcertApp/MultiUserClient/Source/MultiUserClient/Private/ConcertWorkspaceUI.cpp
+++ b/Engine/Plugins/Developer/Concert/ConcertApp/MultiUserClient/Source/MultiUserClient/Private/ConcertWorkspaceUI.cpp
@@ -238,6 +238,13 @@ private:
 	TSharedPtr<FConcertWorkspaceUI> WorkspaceFrontend;
 };
 
+
+template<>
+struct TWidgetTypeTraits<class SConcertWorkspaceModifiedByOtherIndicator>
+{
+	static constexpr bool SupportsInvalidation() { return true; }
+};
+
 /**
  * Controls the appearance of the workspace 'modified by other' icon. The icon is displayed on an
  * asset in the editor content browser when the a client different from this workspace client has
