Multi User: Register Multi User with the sandbox system.
- Entry point: Other editor locations, for now just Sandboxed Editing, will now show that Multi User owns the current sandbox and display a button to summon the UI.
- Overlay: Display an overlay over the browser when the engine is already sandboxed

#Synced-CL 47637010

diff --git a/Engine/Plugins/Developer/Concert/ConcertApp/MultiUserClient/Source/MultiUserClient/MultiUserClient.Build.cs b/Engine/Plugins/Developer/Concert/ConcertApp/MultiUserClient/Source/MultiUserClient/MultiUserClient.Build.cs
--- a/Engine/Plugins/Developer/Concert/ConcertApp/MultiUserClient/Source/MultiUserClient/MultiUserClient.Build.cs
+++ b/Engine/Plugins/Developer/Concert/ConcertApp/MultiUserClient/Source/MultiUserClient/MultiUserClient.Build.cs
@@ -24,6 +24,7 @@ namespace UnrealBuildTool.Rules
 					"DesktopPlatform",
 					"EditorStyle",
 					"Engine",
+					"FileSandboxUI",
 					"InputCore",
 					"JsonUtilities",
 					"MessageLog",
@@ -42,7 +43,7 @@ namespace UnrealBuildTool.Rules
 					"ConcertSharedSlate",
 					"ConcertSyncClient",
 					"ConcertSyncCore",
-					"ConcertTransport",
+					"ConcertTransport", 
 				}
 			);
 
diff --git a/Engine/Plugins/Developer/Concert/ConcertApp/MultiUserClient/Source/MultiUserClient/Private/MultiUserClientModule.cpp b/Engine/Plugins/Developer/Concert/ConcertApp/MultiUserClient/Source/MultiUserClient/Private/MultiUserClientModule.cpp
--- a/Engine/Plugins/Developer/Concert/ConcertApp/MultiUserClient/Source/MultiUserClient/Private/MultiUserClientModule.cpp
+++ b/Engine/Plugins/Developer/Concert/ConcertApp/MultiUserClient/Source/MultiUserClient/Private/MultiUserClientModule.cpp
@@ -54,6 +54,7 @@
 #include "DetailWidgetRow.h"
 #include "Misc/GameVisibilityColumnToggler.h"
 #include "Replication/Misc/DetailRowExtensionHandler.h"
+#include "Sandbox/MultiUserSandboxIntegration.h"
 
 
 #include "Widgets/SConcertBrowser.h"
@@ -628,9 +629,6 @@ public:
 		MultiUserClient->Startup(ClientConfig, EConcertSyncSessionFlags::Default_MultiUserSession);
 		ReplicationManager = MakeShared<UE::MultiUserClient::Replication::FMultiUserReplicationManager>(MultiUserClient.ToSharedRef());
 
-		// Hook UI elements in the tool bar (and setup commands).
-		RegisterUI();
-		RegisterLogging();
 #if WITH_EDITOR
 		if (GIsEditor)
 		{
@@ -639,8 +637,13 @@ public:
 				MultiUserClient.ToSharedRef(), ReplicationManager.ToSharedRef(),
 				UE::MultiUserClient::FGetConcertBrowserWidget::CreateRaw(this, &FMultiUserClientModule::GetOrOpenBrowserTabContent)
 				);
+			SandboxIntegration = MakeUnique<UE::MultiUserClient::FMultiUserSandboxIntegration>(MultiUserClient.ToSharedRef());
 		}
 #endif
+		
+		// Hook UI elements in the tool bar (and setup commands).
+		RegisterUI();
+		RegisterLogging();
 	}
 
 	virtual void ShutdownModule() override
@@ -779,6 +782,11 @@ public:
 		{
 			BrowserSpawnerEntry.SetGroup(WorkspaceGroup.ToSharedRef());
 		}
+		
+		SandboxIntegration->GetEntryPoint()->OnSummonProviderUI().AddLambda([]
+		{
+			FGlobalTabmanager::Get()->TryInvokeTab(ConcertBrowserTabName);
+		});
 	}
 
 	void UnregisterTabSpawner()
@@ -1093,10 +1101,17 @@ private:
 	 */
 	TSharedRef<SDockTab> SpawnConcertBrowserTab(const FSpawnTabArgs& SpawnTabArgs)
 	{
+		TSharedPtr<UE::FileSandboxUI::IExternalSandboxActiveViewModel> ExternalSandboxViewModel;
+#if WITH_EDITOR
+		ExternalSandboxViewModel = SandboxIntegration->GetExternalSandboxViewModel();
+#endif
+		
 		return SNew(SDockTab)
 			.TabRole(NomadTab)
 			[
-				SAssignNew(BrowserTab, SConcertBrowser, MultiUserClient.ToSharedRef(), ReplicationManager.ToSharedRef())
+				SAssignNew(BrowserTab, SConcertBrowser, 
+					MultiUserClient.ToSharedRef(), ReplicationManager.ToSharedRef(), ExternalSandboxViewModel
+					)
 			];
 	}
 
@@ -1465,6 +1480,9 @@ private:
 
 	/** Adds UI to Details panel rows for quickly adding properties to replication. */
 	TUniquePtr<UE::MultiUserClient::Replication::FDetailRowExtensionHandler> DetailReplicationRowExtensionHandler;
+	
+	/** Responsible for registering Multi User with the sandbox system. */
+	TUniquePtr<UE::MultiUserClient::FMultiUserSandboxIntegration> SandboxIntegration;
 #endif
 
 	/** True if the tab spawners have been registered for this module */
diff --git a/Engine/Plugins/Developer/Concert/ConcertApp/MultiUserClient/Source/MultiUserClient/Private/Sandbox/MultiUserSandboxEntryPoint.cpp b/Engine/Plugins/Developer/Concert/ConcertApp/MultiUserClient/Source/MultiUserClient/Private/Sandbox/MultiUserSandboxEntryPoint.cpp
new file mode 100644
--- /dev/null
+++ b/Engine/Plugins/Developer/Concert/ConcertApp/MultiUserClient/Source/MultiUserClient/Private/Sandbox/MultiUserSandboxEntryPoint.cpp
@@ -0,0 +1,33 @@
+﻿// Copyright Epic Games, Inc. All Rights Reserved.
+
+#include "MultiUserSandboxEntryPoint.h"
+
+#include "IConcertSyncClient.h"
+#include "Internationalization/Internationalization.h"
+#include "Internationalization/Text.h"
+
+#define LOCTEXT_NAMESPACE "FMultiUserSandboxEntryPoint"
+
+namespace UE::MultiUserClient
+{
+FMultiUserSandboxEntryPoint::FMultiUserSandboxEntryPoint(const TSharedRef<IConcertSyncClient>& InMultiUserClient)
+	: MultiUserClient(InMultiUserClient)
+{}
+
+void FMultiUserSandboxEntryPoint::SummonProviderUI()
+{
+	OnSummonProviderUIDelegate.Broadcast(); 
+}
+
+FText FMultiUserSandboxEntryPoint::GetEntryPointLabel() const
+{
+	return LOCTEXT("EntryPoint.Label", "Multi User");
+}
+
+bool FMultiUserSandboxEntryPoint::OwnsSandbox(const FileSandboxCore::ISandboxInstance& InSandbox) const
+{
+	return MultiUserClient->GetConcertClient()->GetCurrentSession().IsValid();
+}
+}
+
+#undef LOCTEXT_NAMESPACE
\ No newline at end of file
diff --git a/Engine/Plugins/Developer/Concert/ConcertApp/MultiUserClient/Source/MultiUserClient/Private/Sandbox/MultiUserSandboxEntryPoint.h b/Engine/Plugins/Developer/Concert/ConcertApp/MultiUserClient/Source/MultiUserClient/Private/Sandbox/MultiUserSandboxEntryPoint.h
new file mode 100644
--- /dev/null
+++ b/Engine/Plugins/Developer/Concert/ConcertApp/MultiUserClient/Source/MultiUserClient/Private/Sandbox/MultiUserSandboxEntryPoint.h
@@ -0,0 +1,37 @@
+﻿// Copyright Epic Games, Inc. All Rights Reserved.
+
+#pragma once
+
+#include "Delegates/Delegate.h"
+#include "EntryPoint/ISandboxEntryPoint.h"
+
+class IConcertSyncClient;
+
+namespace UE::MultiUserClient
+{
+/** Represents Multi User as entry point so it becomes discoverable by other editor systems. */
+class FMultiUserSandboxEntryPoint : public FileSandboxUI::ISandboxEntryPoint
+{
+public:
+	
+	explicit FMultiUserSandboxEntryPoint(const TSharedRef<IConcertSyncClient>& InMultiUserClient);
+	
+	//~ Begin ISandboxEntryPoint Interface
+	virtual void SummonProviderUI() override;
+	virtual FText GetEntryPointLabel() const override;
+	virtual bool OwnsSandbox(const FileSandboxCore::ISandboxInstance& InSandbox) const override;
+	//~ End ISandboxEntryPoint Interface
+	
+	/** Event invoked when the UI is supposed to be summoned. */
+	FSimpleMulticastDelegate& OnSummonProviderUI() { return OnSummonProviderUIDelegate; }
+	
+private:
+	
+	/** Multi User's client instance. Used to get the current session: if client is in a session, we'll report that the sandbox is owned by Multi User. */
+	const TSharedRef<IConcertSyncClient> MultiUserClient;
+	
+	/** Event invoked when the UI is supposed to be summoned. */
+	FSimpleMulticastDelegate OnSummonProviderUIDelegate;
+};
+}
+
diff --git a/Engine/Plugins/Developer/Concert/ConcertApp/MultiUserClient/Source/MultiUserClient/Private/Sandbox/MultiUserSandboxIntegration.cpp b/Engine/Plugins/Developer/Concert/ConcertApp/MultiUserClient/Source/MultiUserClient/Private/Sandbox/MultiUserSandboxIntegration.cpp
new file mode 100644
--- /dev/null
+++ b/Engine/Plugins/Developer/Concert/ConcertApp/MultiUserClient/Source/MultiUserClient/Private/Sandbox/MultiUserSandboxIntegration.cpp
@@ -0,0 +1,29 @@
+﻿// Copyright Epic Games, Inc. All Rights Reserved.
+
+#include "MultiUserSandboxIntegration.h"
+
+#include "EntryPoint/EntryPointWidgetFactory.h"
+#include "EntryPoint/ISandboxEntryPointRegistry.h"
+#include "IFileSandboxUIModule.h"
+
+namespace UE::MultiUserClient
+{
+FMultiUserSandboxIntegration::FMultiUserSandboxIntegration(const TSharedRef<IConcertSyncClient>& InMultiUserClient)
+	: EntryPoint(MakeShared<FMultiUserSandboxEntryPoint>(InMultiUserClient))
+	, ExternalSandboxViewModel(FileSandboxUI::MakeExternalSandboxActiveViewModel(EntryPoint))
+{
+	using namespace UE::FileSandboxUI;
+	ISandboxEntryPointRegistry& EntryPointsRegistry = IFileSandboxUIModule::Get().GetEntryPointRegistry();
+	EntryPointsRegistry.RegisterEntryPoint(EntryPoint);
+}
+
+FMultiUserSandboxIntegration::~FMultiUserSandboxIntegration()
+{
+	using namespace UE::FileSandboxUI;
+	if (IFileSandboxUIModule::IsAvailable())
+	{
+		ISandboxEntryPointRegistry& EntryPointsRegistry = IFileSandboxUIModule::Get().GetEntryPointRegistry();
+		EntryPointsRegistry.UnregisterEntryPoint(EntryPoint);
+	}
+}
+}
diff --git a/Engine/Plugins/Developer/Concert/ConcertApp/MultiUserClient/Source/MultiUserClient/Private/Sandbox/MultiUserSandboxIntegration.h b/Engine/Plugins/Developer/Concert/ConcertApp/MultiUserClient/Source/MultiUserClient/Private/Sandbox/MultiUserSandboxIntegration.h
new file mode 100644
--- /dev/null
+++ b/Engine/Plugins/Developer/Concert/ConcertApp/MultiUserClient/Source/MultiUserClient/Private/Sandbox/MultiUserSandboxIntegration.h
@@ -0,0 +1,33 @@
+﻿// Copyright Epic Games, Inc. All Rights Reserved.
+
+#pragma once
+
+#include "MultiUserSandboxEntryPoint.h"
+#include "Templates/SharedPointer.h"
+#include "Templates/UnrealTemplate.h"
+
+namespace UE::FileSandboxUI { class IExternalSandboxActiveViewModel; }
+
+namespace UE::MultiUserClient
+{
+/** Responsible for registering Multi User with the sandbox system. */
+class FMultiUserSandboxIntegration : public FNoncopyable
+{
+public:
+	
+	FMultiUserSandboxIntegration(const TSharedRef<IConcertSyncClient>& InMultiUserClient);
+	~FMultiUserSandboxIntegration();
+	
+	const TSharedRef<FMultiUserSandboxEntryPoint>& GetEntryPoint() const { return EntryPoint; }
+	const TSharedRef<FileSandboxUI::IExternalSandboxActiveViewModel>& GetExternalSandboxViewModel() const { return ExternalSandboxViewModel; }
+	
+private:
+	
+	/** Represents Multi User as entry point so it becomes discoverable by other editor systems. */
+	const TSharedRef<FMultiUserSandboxEntryPoint> EntryPoint;
+	
+	/** This view model is used to overlay a widget onto the browser when the engine is already sandboxed. */
+	const TSharedRef<FileSandboxUI::IExternalSandboxActiveViewModel> ExternalSandboxViewModel;
+};
+}
+
diff --git a/Engine/Plugins/Developer/Concert/ConcertApp/MultiUserClient/Source/MultiUserClient/Private/Widgets/SConcertBrowser.cpp b/Engine/Plugins/Developer/Concert/ConcertApp/MultiUserClient/Source/MultiUserClient/Private/Widgets/SConcertBrowser.cpp
--- a/Engine/Plugins/Developer/Concert/ConcertApp/MultiUserClient/Source/MultiUserClient/Private/Widgets/SConcertBrowser.cpp
+++ b/Engine/Plugins/Developer/Concert/ConcertApp/MultiUserClient/Source/MultiUserClient/Private/Widgets/SConcertBrowser.cpp
@@ -9,6 +9,7 @@
 #include "Widgets/Disconnected/SConcertClientSessionBrowser.h"
 #include "Widgets/Disconnected/SConcertNoAvailability.h"
 
+#include "EntryPoint/EntryPointWidgetFactory.h"
 #include "Internationalization/Regex.h"
 #include "Styling/AppStyle.h"
 #include "Textures/SlateIcon.h"
@@ -20,7 +21,9 @@
 void SConcertBrowser::Construct(
 	const FArguments& InArgs,
 	TSharedRef<IConcertSyncClient> InSyncClient,
-	TSharedRef<UE::MultiUserClient::Replication::FMultiUserReplicationManager> InReplicationManager)
+	TSharedRef<UE::MultiUserClient::Replication::FMultiUserReplicationManager> InReplicationManager,
+	TSharedPtr<UE::FileSandboxUI::IExternalSandboxActiveViewModel> InExternalSandboxViewModel
+	)
 {
 	if (!MultiUserClientUtils::HasServerCompatibleCommunicationPluginEnabled())
 	{
@@ -36,6 +39,8 @@ void SConcertBrowser::Construct(
 
 	WeakConcertSyncClient = InSyncClient;
 	WeakReplicationManager = InReplicationManager;
+	ExternalSandboxViewModel = InExternalSandboxViewModel;
+	
 	if (TSharedPtr<IConcertSyncClient> ConcertSyncClient = WeakConcertSyncClient.Pin())
 	{
 		SearchedText = MakeShared<FText>(); // Will keep in memory the session browser search text between join/leave UI transitions.
@@ -65,27 +70,56 @@ void SConcertBrowser::AttachChildWidget(EConcertConnectionStatus ConnectionStatu
 		
 		if (ConnectionStatus == EConcertConnectionStatus::Connected)
 		{
-			if (const TSharedPtr<UE::MultiUserClient::Replication::FMultiUserReplicationManager> ReplicationManager = WeakReplicationManager.Pin()
-				; ensure(ReplicationManager))
-			{
-				ChildSlot.AttachWidget(
-					SAssignNew(ActiveSessionWidget, UE::MultiUserClient::SActiveSessionRoot,
-						ConcertSyncClient,
-						ReplicationManager.ToSharedRef()
-						)
-					);
-			}
+			AttachActiveWidget(ConcertSyncClient.ToSharedRef());
 		}
 		else if (ConnectionStatus == EConcertConnectionStatus::Disconnected)
 		{
-			ChildSlot.AttachWidget(
-				SAssignNew(SessionBrowser, SConcertClientSessionBrowser,
-					ConcertSyncClient->GetConcertClient(),
-					SearchedText
-					)
-				);
+			AttachBrowserWidget(ConcertSyncClient.ToSharedRef());
 		}
 	}
 }
 
+void SConcertBrowser::AttachActiveWidget(const TSharedRef<IConcertSyncClient>& ConcertSyncClient)
+{
+	if (const TSharedPtr<UE::MultiUserClient::Replication::FMultiUserReplicationManager> ReplicationManager = WeakReplicationManager.Pin()
+		; ensure(ReplicationManager))
+	{
+		ChildSlot.AttachWidget(
+			SAssignNew(ActiveSessionWidget, UE::MultiUserClient::SActiveSessionRoot,
+				ConcertSyncClient,
+				ReplicationManager.ToSharedRef()
+			)
+		);
+	}
+}
+
+void SConcertBrowser::AttachBrowserWidget(const TSharedRef<IConcertSyncClient>& InConcertSyncClient)
+{
+	const TSharedRef<SConcertClientSessionBrowser> Browser = SAssignNew(SessionBrowser, SConcertClientSessionBrowser,
+		InConcertSyncClient->GetConcertClient(),
+		SearchedText
+		);
+	
+	const TSharedPtr<UE::FileSandboxUI::IExternalSandboxActiveViewModel> ViewModelPin = ExternalSandboxViewModel.Pin();
+	if (!ViewModelPin)
+	{
+		ChildSlot.AttachWidget(Browser);
+		return;
+	}
+	
+	const TSharedRef<SWidget> Overlay = SNew(SOverlay)
+		+SOverlay::Slot()
+		.ZOrder(0)
+		[
+			Browser
+		]
+		+SOverlay::Slot()
+		.ZOrder(0)
+		[
+			UE::FileSandboxUI::MakeExternalSandboxActiveOverlay(ViewModelPin.ToSharedRef())
+		];
+	
+	ChildSlot.AttachWidget(Overlay);
+}
+
 #undef LOCTEXT_NAMESPACE
diff --git a/Engine/Plugins/Developer/Concert/ConcertApp/MultiUserClient/Source/MultiUserClient/Private/Widgets/SConcertBrowser.h b/Engine/Plugins/Developer/Concert/ConcertApp/MultiUserClient/Source/MultiUserClient/Private/Widgets/SConcertBrowser.h
--- a/Engine/Plugins/Developer/Concert/ConcertApp/MultiUserClient/Source/MultiUserClient/Private/Widgets/SConcertBrowser.h
+++ b/Engine/Plugins/Developer/Concert/ConcertApp/MultiUserClient/Source/MultiUserClient/Private/Widgets/SConcertBrowser.h
@@ -12,6 +12,7 @@
 class IConcertClientSession;
 class IConcertSyncClient;
 class SConcertClientSessionBrowser;
+namespace UE::FileSandboxUI { class IExternalSandboxActiveViewModel; }
 namespace UE::MultiUserClient { class SActiveSessionRoot; }
 namespace UE::MultiUserClient::Replication { class FMultiUserReplicationManager; }
 
@@ -36,7 +37,8 @@ public:
 	*/
 	void Construct(const FArguments& InArgs,
 		TSharedRef<IConcertSyncClient> InSyncClient,
-		TSharedRef<UE::MultiUserClient::Replication::FMultiUserReplicationManager> InReplicationManager
+		TSharedRef<UE::MultiUserClient::Replication::FMultiUserReplicationManager> InReplicationManager,
+		TSharedPtr<UE::FileSandboxUI::IExternalSandboxActiveViewModel> InExternalSandboxViewModel
 		);
 
 	const TWeakPtr<UE::MultiUserClient::SActiveSessionRoot>& GetActiveSessionWidget() const { return ActiveSessionWidget; }
@@ -48,6 +50,8 @@ private:
 	TWeakPtr<IConcertSyncClient> WeakConcertSyncClient;
 	/** Interacts with the replication system on behalf of Multi-User. */
 	TWeakPtr<UE::MultiUserClient::Replication::FMultiUserReplicationManager> WeakReplicationManager;
+	/** This view model is used to overlay a widget onto the browser when the engine is already sandboxed. */
+	TWeakPtr<UE::FileSandboxUI::IExternalSandboxActiveViewModel> ExternalSandboxViewModel;
 
 	/** Only valid when in session. */
 	TWeakPtr<UE::MultiUserClient::SActiveSessionRoot> ActiveSessionWidget;
@@ -62,4 +66,6 @@ private:
 
 	/** Attaches the child widgets according to the connection status. */
 	void AttachChildWidget(EConcertConnectionStatus ConnectionStatus);
+	void AttachActiveWidget(const TSharedRef<IConcertSyncClient>& InConcertSyncClient);
+	void AttachBrowserWidget(const TSharedRef<IConcertSyncClient>& InConcertSyncClient);
 };
