Multi User: Fix crash in -game

#Synced-CL 47637010

diff --git a/Engine/Plugins/Developer/Concert/ConcertApp/MultiUserClient/Source/MultiUserClient/Private/MultiUserClientModule.cpp b/Engine/Plugins/Developer/Concert/ConcertApp/MultiUserClient/Source/MultiUserClient/Private/MultiUserClientModule.cpp
--- a/Engine/Plugins/Developer/Concert/ConcertApp/MultiUserClient/Source/MultiUserClient/Private/MultiUserClientModule.cpp
+++ b/Engine/Plugins/Developer/Concert/ConcertApp/MultiUserClient/Source/MultiUserClient/Private/MultiUserClientModule.cpp
@@ -783,10 +783,15 @@ public:
 			BrowserSpawnerEntry.SetGroup(WorkspaceGroup.ToSharedRef());
 		}
 		
-		SandboxIntegration->GetEntryPoint()->OnSummonProviderUI().AddLambda([]
+#if WITH_EDITOR
+		if (SandboxIntegration)
 		{
-			FGlobalTabmanager::Get()->TryInvokeTab(ConcertBrowserTabName);
-		});
+			SandboxIntegration->GetEntryPoint()->OnSummonProviderUI().AddLambda([]
+			{
+				FGlobalTabmanager::Get()->TryInvokeTab(ConcertBrowserTabName);
+			});
+		}
+#endif
 	}
 
 	void UnregisterTabSpawner()
@@ -1103,7 +1108,7 @@ private:
 	{
 		TSharedPtr<UE::FileSandboxUI::IExternalSandboxActiveViewModel> ExternalSandboxViewModel;
 #if WITH_EDITOR
-		ExternalSandboxViewModel = SandboxIntegration->GetExternalSandboxViewModel();
+		ExternalSandboxViewModel = SandboxIntegration ? SandboxIntegration->GetExternalSandboxViewModel().ToSharedPtr() : nullptr;
 #endif
 		
 		return SNew(SDockTab)
